Copying //UE4/Dev-Framework to //UE4/Dev-Main (Source: //UE4/Dev-Framework @ 3967517)

diff --git a/Engine/Source/Runtime/NavigationSystem/Private/NavMesh/RecastNavMeshGenerator.cpp b/Engine/Source/Runtime/NavigationSystem/Private/NavMesh/RecastNavMeshGenerator.cpp
new file mode 100644
--- /dev/null
+++ b/Engine/Source/Runtime/NavigationSystem/Private/NavMesh/RecastNavMeshGenerator.cpp
@@ -0,0 +1,4756 @@
+// Copyright 1998-2018 Epic Games, Inc. All Rights Reserved.
+
+#include "NavMesh/RecastNavMeshGenerator.h"
+#include "AI/Navigation/NavRelevantInterface.h"
+#include "Components/PrimitiveComponent.h"
+#include "HAL/FileManager.h"
+#include "Misc/Paths.h"
+#include "Serialization/MemoryWriter.h"
+#include "EngineGlobals.h"
+#include "GameFramework/PlayerController.h"
+#include "Engine/Engine.h"
+#include "NavigationSystem.h"
+
+#if WITH_RECAST
+#if WITH_PHYSX
+	#include "PhysXPublic.h"
+#endif
+#include "NavMesh/PImplRecastNavMesh.h"
+
+// recast includes
+#include "Detour/DetourNavMeshBuilder.h"
+#include "DetourTileCache/DetourTileCacheBuilder.h"
+#include "NavMesh/RecastHelpers.h"
+#include "NavAreas/NavArea_LowHeight.h"
+#include "AI/NavigationSystemHelpers.h"
+#include "VisualLogger/VisualLoggerTypes.h"
+#include "PhysicsEngine/ConvexElem.h"
+#include "PhysicsEngine/BodySetup.h"
+
+#define SEAMLESS_REBUILDING_ENABLED 1
+
+#define GENERATE_SEGMENT_LINKS 1
+#define GENERATE_CLUSTER_LINKS 1
+
+#define SHOW_NAV_EXPORT_PREVIEW 0
+
+#define TEXT_WEAKOBJ_NAME(obj) (obj.IsValid(false) ? *obj->GetName() : (obj.IsValid(false, true)) ? TEXT("MT-Unreachable") : TEXT("INVALID"))
+
+#define DO_RECAST_STATS 0
+#if DO_RECAST_STATS
+#define RECAST_STAT SCOPE_CYCLE_COUNTER
+#else
+#define RECAST_STAT(...) 
+#endif
+
+struct dtTileCacheAlloc;
+
+FORCEINLINE bool DoesBoxContainOrOverlapVector(const FBox& BigBox, const FVector& In)
+{
+	return (In.X >= BigBox.Min.X) && (In.X <= BigBox.Max.X) 
+		&& (In.Y >= BigBox.Min.Y) && (In.Y <= BigBox.Max.Y) 
+		&& (In.Z >= BigBox.Min.Z) && (In.Z <= BigBox.Max.Z);
+}
+/** main difference between this and FBox::ContainsBox is that this returns true also when edges overlap */
+FORCEINLINE bool DoesBoxContainBox(const FBox& BigBox, const FBox& SmallBox)
+{
+	return DoesBoxContainOrOverlapVector(BigBox, SmallBox.Min) && DoesBoxContainOrOverlapVector(BigBox, SmallBox.Max);
+}
+
+int32 GetTilesCountHelper(const dtNavMesh* DetourMesh)
+{
+	int32 NumTiles = 0;
+	if (DetourMesh)
+	{
+		for (int32 i = 0; i < DetourMesh->getMaxTiles(); i++)
+		{
+			const dtMeshTile* TileData = DetourMesh->getTile(i);
+			if (TileData && TileData->header && TileData->dataSize > 0)
+			{
+				NumTiles++;
+			}
+		}
+	}
+
+	return NumTiles;
+}
+
+/**
+ * Exports geometry to OBJ file. Can be used to verify NavMesh generation in RecastDemo app
+ * @param FileName - full name of OBJ file with extension
+ * @param GeomVerts - list of vertices
+ * @param GeomFaces - list of triangles (3 vert indices for each)
+ */
+static void ExportGeomToOBJFile(const FString& InFileName, const TNavStatArray<float>& GeomCoords, const TNavStatArray<int32>& GeomFaces, const FString& AdditionalData)
+{
+#define USE_COMPRESSION 0
+
+#if ALLOW_DEBUG_FILES
+	SCOPE_CYCLE_COUNTER(STAT_Navigation_TileGeometryExportToObjAsync);
+
+	FString FileName = InFileName;
+
+#if USE_COMPRESSION
+	FileName += TEXT("z");
+	struct FDataChunk
+	{
+		TArray<uint8> UncompressedBuffer;
+		TArray<uint8> CompressedBuffer;
+		void CompressBuffer()
+		{
+			const int32 HeaderSize = sizeof(int32);
+			const int32 UncompressedSize = UncompressedBuffer.Num();
+			CompressedBuffer.Init(0, HeaderSize + FMath::Trunc(1.1f * UncompressedSize));
+
+			int32 CompressedSize = CompressedBuffer.Num() - HeaderSize;
+			uint8* DestBuffer = CompressedBuffer.GetData();
+			FMemory::Memcpy(DestBuffer, &UncompressedSize, HeaderSize);
+			DestBuffer += HeaderSize;
+
+			FCompression::CompressMemory((ECompressionFlags)(COMPRESS_ZLIB | COMPRESS_BiasMemory), (void*)DestBuffer, CompressedSize, (void*)UncompressedBuffer.GetData(), UncompressedSize);
+			CompressedBuffer.SetNum(CompressedSize + HeaderSize, false);
+		}
+	};
+	FDataChunk AllDataChunks[3];
+	const int32 NumberOfChunks = sizeof(AllDataChunks) / sizeof(FDataChunk);
+	{
+		FMemoryWriter ArWriter(AllDataChunks[0].UncompressedBuffer);
+		for (int32 i = 0; i < GeomCoords.Num(); i += 3)
+		{
+			FVector Vertex(GeomCoords[i + 0], GeomCoords[i + 1], GeomCoords[i + 2]);
+			ArWriter << Vertex;
+		}
+	}
+
+	{
+		FMemoryWriter ArWriter(AllDataChunks[1].UncompressedBuffer);
+		for (int32 i = 0; i < GeomFaces.Num(); i += 3)
+		{
+			FVector Face(GeomFaces[i + 0] + 1, GeomFaces[i + 1] + 1, GeomFaces[i + 2] + 1);
+			ArWriter << Face;
+		}
+	}
+
+	{
+		auto AnsiAdditionalData = StringCast<ANSICHAR>(*AdditionalData);
+		FMemoryWriter ArWriter(AllDataChunks[2].UncompressedBuffer);
+		ArWriter.Serialize((ANSICHAR*)AnsiAdditionalData.Get(), AnsiAdditionalData.Length());
+	}
+
+	FArchive* FileAr = IFileManager::Get().CreateDebugFileWriter(*FileName);
+	if (FileAr != NULL)
+	{
+		for (int32 Index = 0; Index < NumberOfChunks; ++Index)
+		{
+			AllDataChunks[Index].CompressBuffer();
+			int32 BufferSize = AllDataChunks[Index].CompressedBuffer.Num();
+			FileAr->Serialize(&BufferSize, sizeof(int32));
+			FileAr->Serialize((void*)AllDataChunks[Index].CompressedBuffer.GetData(), AllDataChunks[Index].CompressedBuffer.Num());
+		}
+		UE_LOG(LogNavigation, Error, TEXT("UncompressedBuffer size:: %d "), AllDataChunks[0].UncompressedBuffer.Num() + AllDataChunks[1].UncompressedBuffer.Num() + AllDataChunks[2].UncompressedBuffer.Num());
+		FileAr->Close();
+	}
+
+#else
+	FArchive* FileAr = IFileManager::Get().CreateDebugFileWriter(*FileName);
+	if (FileAr != NULL)
+	{
+		for (int32 Index = 0; Index < GeomCoords.Num(); Index += 3)
+		{
+			FString LineToSave = FString::Printf(TEXT("v %f %f %f\n"), GeomCoords[Index + 0], GeomCoords[Index + 1], GeomCoords[Index + 2]);
+			auto AnsiLineToSave = StringCast<ANSICHAR>(*LineToSave);
+			FileAr->Serialize((ANSICHAR*)AnsiLineToSave.Get(), AnsiLineToSave.Length());
+		}
+
+		for (int32 Index = 0; Index < GeomFaces.Num(); Index += 3)
+		{
+			FString LineToSave = FString::Printf(TEXT("f %d %d %d\n"), GeomFaces[Index + 0] + 1, GeomFaces[Index + 1] + 1, GeomFaces[Index + 2] + 1);
+			auto AnsiLineToSave = StringCast<ANSICHAR>(*LineToSave);
+			FileAr->Serialize((ANSICHAR*)AnsiLineToSave.Get(), AnsiLineToSave.Length());
+		}
+
+		auto AnsiAdditionalData = StringCast<ANSICHAR>(*AdditionalData);
+		FileAr->Serialize((ANSICHAR*)AnsiAdditionalData.Get(), AnsiAdditionalData.Length());
+		FileAr->Close();
+	}
+#endif
+
+#undef USE_COMPRESSION
+#endif
+}
+
+//----------------------------------------------------------------------//
+// 
+// 
+
+struct FRecastGeometryExport : public FNavigableGeometryExport
+{
+	FRecastGeometryExport(FNavigationRelevantData& InData) : Data(&InData) 
+	{
+		Data->Bounds = FBox(ForceInit);
+	}
+
+	FNavigationRelevantData* Data;
+	TNavStatArray<float> VertexBuffer;
+	TNavStatArray<int32> IndexBuffer;
+	FWalkableSlopeOverride SlopeOverride;
+
+#if WITH_PHYSX
+	virtual void ExportPxTriMesh16Bit(physx::PxTriangleMesh const * const TriMesh, const FTransform& LocalToWorld) override;
+	virtual void ExportPxTriMesh32Bit(physx::PxTriangleMesh const * const TriMesh, const FTransform& LocalToWorld) override;
+	virtual void ExportPxConvexMesh(physx::PxConvexMesh const * const ConvexMesh, const FTransform& LocalToWorld) override;
+	virtual void ExportPxHeightField(physx::PxHeightField const * const HeightField, const FTransform& LocalToWorld) override;
+#endif // WITH_PHYSX
+	virtual void ExportHeightFieldSlice(const FNavHeightfieldSamples& PrefetchedHeightfieldSamples, const int32 NumRows, const int32 NumCols, const FTransform& LocalToWorld, const FBox& SliceBox) override;
+	virtual void ExportCustomMesh(const FVector* InVertices, int32 NumVerts, const int32* InIndices, int32 NumIndices, const FTransform& LocalToWorld) override;
+	virtual void ExportRigidBodySetup(UBodySetup& BodySetup, const FTransform& LocalToWorld) override;
+	virtual void AddNavModifiers(const FCompositeNavModifier& Modifiers) override;
+	virtual void SetNavDataPerInstanceTransformDelegate(const FNavDataPerInstanceTransformDelegate& InDelegate) override;
+};
+
+FRecastVoxelCache::FRecastVoxelCache(const uint8* Memory)
+{
+	uint8* BytesArr = (uint8*)Memory;
+	if (Memory)
+	{
+		NumTiles = *((int32*)BytesArr);	BytesArr += sizeof(int32);
+		Tiles = (FTileInfo*)BytesArr;
+	}
+	else
+	{
+		NumTiles = 0;
+	}
+
+	FTileInfo* iTile = Tiles;	
+	for (int i = 0; i < NumTiles; i++)
+	{
+		iTile = (FTileInfo*)BytesArr; BytesArr += sizeof(FTileInfo);
+		if (iTile->NumSpans)
+		{
+			iTile->SpanData = (rcSpanCache*)BytesArr; BytesArr += sizeof(rcSpanCache) * iTile->NumSpans;
+		}
+		else
+		{
+			iTile->SpanData = 0;
+		}
+
+		iTile->NextTile = (FTileInfo*)BytesArr;
+	}
+
+	if (NumTiles > 0)
+	{
+		iTile->NextTile = 0;
+	}
+	else
+	{
+		Tiles = 0;
+	}
+}
+
+FRecastGeometryCache::FRecastGeometryCache(const uint8* Memory)
+{
+	Header = *((FHeader*)Memory);
+	Verts = (float*)(Memory + sizeof(FRecastGeometryCache));
+	Indices = (int32*)(Memory + sizeof(FRecastGeometryCache) + (sizeof(float) * Header.NumVerts * 3));
+}
+
+namespace RecastGeometryExport {
+
+static UWorld* FindEditorWorld()
+{
+	if (GEngine)
+	{
+		for (const FWorldContext& Context : GEngine->GetWorldContexts())
+		{
+			if (Context.WorldType == EWorldType::Editor)
+			{
+				return Context.World();
+			}
+		}
+	}
+
+	return NULL;
+}
+
+static void StoreCollisionCache(FRecastGeometryExport& GeomExport)
+{
+	const int32 NumFaces = GeomExport.IndexBuffer.Num() / 3;
+	const int32 NumVerts = GeomExport.VertexBuffer.Num() / 3;
+
+	if (NumFaces == 0 || NumVerts == 0)
+	{
+		GeomExport.Data->CollisionData.Empty();
+		return;
+	}
+
+	FRecastGeometryCache::FHeader HeaderInfo;
+	HeaderInfo.NumFaces = NumFaces;
+	HeaderInfo.NumVerts = NumVerts;
+	HeaderInfo.SlopeOverride = GeomExport.SlopeOverride;
+
+	// allocate memory
+	const int32 HeaderSize = sizeof(FRecastGeometryCache);
+	const int32 CoordsSize = sizeof(float) * 3 * NumVerts;
+	const int32 IndicesSize = sizeof(int32) * 3 * NumFaces;
+	const int32 CacheSize = HeaderSize + CoordsSize + IndicesSize;
+
+	HeaderInfo.Validation.DataSize = CacheSize;
+
+	// reserve + add combo to allocate exact amount (without any overhead/slack)
+	GeomExport.Data->CollisionData.Reserve(CacheSize);
+	GeomExport.Data->CollisionData.AddUninitialized(CacheSize);
+
+	// store collisions
+	uint8* RawMemory = GeomExport.Data->CollisionData.GetData();
+	FRecastGeometryCache* CacheMemory = (FRecastGeometryCache*)RawMemory;
+	CacheMemory->Header = HeaderInfo;
+	CacheMemory->Verts = 0;
+	CacheMemory->Indices = 0;
+
+	FMemory::Memcpy(RawMemory + HeaderSize, GeomExport.VertexBuffer.GetData(), CoordsSize);
+	FMemory::Memcpy(RawMemory + HeaderSize + CoordsSize, GeomExport.IndexBuffer.GetData(), IndicesSize);
+}
+
+#if WITH_PHYSX
+/** exports PxConvexMesh as trimesh */
+void ExportPxConvexMesh(PxConvexMesh const * const ConvexMesh, const FTransform& LocalToWorld,
+						TNavStatArray<float>& VertexBuffer, TNavStatArray<int32>& IndexBuffer,
+						FBox& UnrealBounds)
+{
+	// after FKConvexElem::AddCachedSolidConvexGeom
+	if(ConvexMesh == NULL)
+	{
+		return;
+	}
+
+	int32 StartVertOffset = VertexBuffer.Num() / 3;
+	const bool bNegX = LocalToWorld.GetDeterminant() < 0;
+
+	// get PhysX data
+	const PxVec3* PVertices = ConvexMesh->getVertices();
+	const PxU8* PIndexBuffer = ConvexMesh->getIndexBuffer();
+	const PxU32 NbPolygons = ConvexMesh->getNbPolygons();
+
+#if SHOW_NAV_EXPORT_PREVIEW
+	UWorld* DebugWorld = FindEditorWorld();
+#endif // SHOW_NAV_EXPORT_PREVIEW
+
+	for(PxU32 i = 0; i < NbPolygons; ++i)
+	{
+		PxHullPolygon Data;
+		bool bStatus = ConvexMesh->getPolygonData(i, Data);
+		check(bStatus);
+
+		const PxU8* indices = PIndexBuffer + Data.mIndexBase;
+		
+		// add vertices 
+		for(PxU32 j = 0; j < Data.mNbVerts; ++j)
+		{
+			const int32 VertIndex = indices[j];
+			const FVector UnrealCoords = LocalToWorld.TransformPosition( P2UVector(PVertices[VertIndex]) );
+			UnrealBounds += UnrealCoords;
+
+			VertexBuffer.Add(UnrealCoords.X);
+			VertexBuffer.Add(UnrealCoords.Y);
+			VertexBuffer.Add(UnrealCoords.Z);
+		}
+
+		// Add indices
+		const PxU32 nbTris = Data.mNbVerts - 2;
+		for(PxU32 j = 0; j < nbTris; ++j)
+		{
+			IndexBuffer.Add(StartVertOffset + 0 );
+			IndexBuffer.Add(StartVertOffset + j + 2);
+			IndexBuffer.Add(StartVertOffset + j + 1);
+
+#if SHOW_NAV_EXPORT_PREVIEW
+			if (DebugWorld)
+			{
+				FVector V0(VertexBuffer[(StartVertOffset + 0) * 3+0], VertexBuffer[(StartVertOffset + 0) * 3+1], VertexBuffer[(StartVertOffset + 0) * 3+2]);
+				FVector V1(VertexBuffer[(StartVertOffset + j + 2) * 3+0], VertexBuffer[(StartVertOffset + j + 2) * 3+1], VertexBuffer[(StartVertOffset + j + 2) * 3+2]);
+				FVector V2(VertexBuffer[(StartVertOffset + j + 1) * 3+0], VertexBuffer[(StartVertOffset + j + 1) * 3+1], VertexBuffer[(StartVertOffset + j + 1) * 3+2]);
+
+				DrawDebugLine(DebugWorld, V0, V1, bNegX ? FColor::Red : FColor::Blue, true);
+				DrawDebugLine(DebugWorld, V1, V2, bNegX ? FColor::Red : FColor::Blue, true);
+				DrawDebugLine(DebugWorld, V2, V0, bNegX ? FColor::Red : FColor::Blue, true);
+			}
+#endif // SHOW_NAV_EXPORT_PREVIEW
+		}
+
+		StartVertOffset += Data.mNbVerts;
+	}
+}
+
+template<typename TIndicesType> 
+FORCEINLINE_DEBUGGABLE void ExportPxTriMesh(PxTriangleMesh const * const TriMesh, const FTransform& LocalToWorld,
+											TNavStatArray<float>& VertexBuffer, TNavStatArray<int32>& IndexBuffer,
+											FBox& UnrealBounds)
+{
+	if (TriMesh == NULL)
+	{
+		return;
+	}
+
+	int32 VertOffset = VertexBuffer.Num() / 3;
+	const PxVec3* PVerts = TriMesh->getVertices();
+	const PxU32 NumTris = TriMesh->getNbTriangles();
+
+	const TIndicesType* Indices = (TIndicesType*)TriMesh->getTriangles();;
+		
+	VertexBuffer.Reserve(VertexBuffer.Num() + NumTris*3);
+	IndexBuffer.Reserve(IndexBuffer.Num() + NumTris*3);
+	const bool bFlipCullMode = (LocalToWorld.GetDeterminant() < 0.f);
+	const int32 IndexOrder[3] = { bFlipCullMode ? 0 : 2, 1, bFlipCullMode ? 2 : 0 };
+
+#if SHOW_NAV_EXPORT_PREVIEW
+	UWorld* DebugWorld = FindEditorWorld();
+#endif // SHOW_NAV_EXPORT_PREVIEW
+
+	for(PxU32 TriIdx = 0; TriIdx < NumTris; ++TriIdx)
+	{
+		for (int32 i = 0; i < 3; i++)
+		{
+			const FVector UnrealCoords = LocalToWorld.TransformPosition(P2UVector(PVerts[Indices[i]]));
+			UnrealBounds += UnrealCoords;
+
+			VertexBuffer.Add(UnrealCoords.X);
+			VertexBuffer.Add(UnrealCoords.Y);
+			VertexBuffer.Add(UnrealCoords.Z);
+		}
+		Indices += 3;
+
+		IndexBuffer.Add(VertOffset + IndexOrder[0]);
+		IndexBuffer.Add(VertOffset + IndexOrder[1]);
+		IndexBuffer.Add(VertOffset + IndexOrder[2]);
+
+#if SHOW_NAV_EXPORT_PREVIEW
+		if (DebugWorld)
+		{
+			FVector V0(VertexBuffer[(VertOffset + IndexOrder[0]) * 3+0], VertexBuffer[(VertOffset + IndexOrder[0]) * 3+1], VertexBuffer[(VertOffset + IndexOrder[0]) * 3+2]);
+			FVector V1(VertexBuffer[(VertOffset + IndexOrder[1]) * 3+0], VertexBuffer[(VertOffset + IndexOrder[1]) * 3+1], VertexBuffer[(VertOffset + IndexOrder[1]) * 3+2]);
+			FVector V2(VertexBuffer[(VertOffset + IndexOrder[2]) * 3+0], VertexBuffer[(VertOffset + IndexOrder[2]) * 3+1], VertexBuffer[(VertOffset + IndexOrder[2]) * 3+2]);
+
+			DrawDebugLine(DebugWorld, V0, V1, bFlipCullMode ? FColor::Red : FColor::Blue, true);
+			DrawDebugLine(DebugWorld, V1, V2, bFlipCullMode ? FColor::Red : FColor::Blue, true);
+			DrawDebugLine(DebugWorld, V2, V0, bFlipCullMode ? FColor::Red : FColor::Blue, true);
+		}
+#endif // SHOW_NAV_EXPORT_PREVIEW
+
+		VertOffset += 3;
+	}
+}
+
+void ExportPxHeightField(PxHeightField const * const HeightField, const FTransform& LocalToWorld
+	, TNavStatArray<float>& VertexBuffer, TNavStatArray<int32>& IndexBuffer
+	, FBox& UnrealBounds)
+{
+	if (HeightField == NULL)
+	{
+		return;
+	}
+
+	QUICK_SCOPE_CYCLE_COUNTER(STAT_NavMesh_ExportPxHeightField);
+
+	const int32 NumRows = HeightField->getNbRows();
+	const int32 NumCols = HeightField->getNbColumns();
+	const int32 VertexCount = NumRows * NumCols;
+
+	// Unfortunately we have to use PxHeightField::saveCells instead PxHeightField::getHeight here 
+	// because current PxHeightField interface does not provide an access to a triangle material index by HF 2D coordinates
+	// PxHeightField::getTriangleMaterialIndex uses some internal adressing which does not match HF 2D coordinates
+	TArray<PxHeightFieldSample> HFSamples;
+	HFSamples.SetNumUninitialized(VertexCount);
+	{
+		QUICK_SCOPE_CYCLE_COUNTER(STAT_NavMesh_ExportPxHeightField_saveCells);
+		HeightField->saveCells(HFSamples.GetData(), VertexCount*HFSamples.GetTypeSize());
+	}
+
+	//
+	const int32 VertOffset = VertexBuffer.Num() / 3;
+	const int32 NumQuads = (NumRows - 1)*(NumCols - 1);
+
+	VertexBuffer.Reserve(VertexBuffer.Num() + VertexCount * 3);
+	IndexBuffer.Reserve(IndexBuffer.Num() + NumQuads * 6);
+
+	const bool bMirrored = (LocalToWorld.GetDeterminant() < 0.f);
+
+	for (int32 Y = 0; Y < NumRows; Y++)
+	{
+		for (int32 X = 0; X < NumCols; X++)
+		{
+			const int32 SampleIdx = (bMirrored ? X : (NumCols - X - 1))*NumCols + Y;
+
+			const PxHeightFieldSample& Sample = HFSamples[SampleIdx];
+			const FVector UnrealCoords = LocalToWorld.TransformPosition(FVector(X, Y, Sample.height));
+			UnrealBounds += UnrealCoords;
+
+			VertexBuffer.Add(UnrealCoords.X);
+			VertexBuffer.Add(UnrealCoords.Y);
+			VertexBuffer.Add(UnrealCoords.Z);
+		}
+	}
+
+	for (int32 Y = 0; Y < NumRows - 1; Y++)
+	{
+		for (int32 X = 0; X < NumCols - 1; X++)
+		{
+			const int32 SampleIdx = (bMirrored ? X : (NumCols - X - 1 - 1))*NumCols + Y;
+			const PxHeightFieldSample& Sample = HFSamples[SampleIdx];
+			const bool bIsHole = (Sample.materialIndex0 == PxHeightFieldMaterial::eHOLE);
+			if (bIsHole)
+			{
+				continue;
+			}
+
+			const int32 I00 = X + 0 + (Y + 0)*NumCols;
+			int32 I01 = X + 0 + (Y + 1)*NumCols;
+			int32 I10 = X + 1 + (Y + 0)*NumCols;
+			const int32 I11 = X + 1 + (Y + 1)*NumCols;
+
+			if (bMirrored)
+			{
+				Swap(I01, I10);
+			}
+
+			IndexBuffer.Add(VertOffset + I00);
+			IndexBuffer.Add(VertOffset + I11);
+			IndexBuffer.Add(VertOffset + I10);
+
+			IndexBuffer.Add(VertOffset + I00);
+			IndexBuffer.Add(VertOffset + I01);
+			IndexBuffer.Add(VertOffset + I11);
+		}
+	}
+}
+#endif // WITH_PHYSX
+
+void ExportHeightFieldSlice(const FNavHeightfieldSamples& PrefetchedHeightfieldSamples, const int32 NumRows, const int32 NumCols, const FTransform& LocalToWorld
+	, TNavStatArray<float>& VertexBuffer, TNavStatArray<int32>& IndexBuffer, const FBox& SliceBox
+	, FBox& UnrealBounds)
+{
+	QUICK_SCOPE_CYCLE_COUNTER(STAT_NavMesh_ExportHeightFieldSlice);
+
+#if WITH_PHYSX
+	static const uint32 SizeOfPx = sizeof(physx::PxI16);
+	static const uint32 SizeOfHeight = PrefetchedHeightfieldSamples.Heights.GetTypeSize();
+	ensure(SizeOfPx == SizeOfHeight);
+#endif // WITH_PHYSX
+
+	// calculate the actual start and number of columns we want
+	const FBox LocalBox = SliceBox.TransformBy(LocalToWorld.Inverse());
+	const bool bMirrored = (LocalToWorld.GetDeterminant() < 0.f);
+
+	const int32 MinX = FMath::Clamp(FMath::FloorToInt(LocalBox.Min.X) - 1, 0, NumCols);
+	const int32 MinY = FMath::Clamp(FMath::FloorToInt(LocalBox.Min.Y) - 1, 0, NumRows);
+	const int32 MaxX = FMath::Clamp(FMath::CeilToInt(LocalBox.Max.X) + 1, 0, NumCols);
+	const int32 MaxY = FMath::Clamp(FMath::CeilToInt(LocalBox.Max.Y) + 1, 0, NumRows);
+	const int32 SizeX = MaxX - MinX;
+	const int32 SizeY = MaxY - MinY;
+
+	if (SizeX <= 0 || SizeY <= 0)
+	{
+		// slice is outside bounds, skip
+		return;
+	}
+
+	const int32 VertOffset = VertexBuffer.Num() / 3;
+	const int32 NumVerts = SizeX * SizeY;
+	const int32 NumQuads = (SizeX - 1) * (SizeY - 1);
+	{
+		QUICK_SCOPE_CYCLE_COUNTER(STAT_RecastGeometryExport_AllocatingMemory);
+		VertexBuffer.Reserve(VertexBuffer.Num() + NumVerts * 3);
+		IndexBuffer.Reserve(IndexBuffer.Num() + NumQuads * 3 * 2);
+	}
+
+	for (int32 IdxY = 0; IdxY < SizeY; IdxY++)
+	{
+		for (int32 IdxX = 0; IdxX < SizeX; IdxX++)
+		{
+			const int32 CoordX = IdxX + MinX;
+			const int32 CoordY = IdxY + MinY;
+			const int32 SampleIdx = ((bMirrored ? CoordX : (NumCols - CoordX - 1)) * NumCols) + CoordY;
+
+			const FVector UnrealCoords = LocalToWorld.TransformPosition(FVector(CoordX, CoordY, PrefetchedHeightfieldSamples.Heights[SampleIdx]));
+			VertexBuffer.Add(UnrealCoords.X);
+			VertexBuffer.Add(UnrealCoords.Y);
+			VertexBuffer.Add(UnrealCoords.Z);
+		}
+	}
+
+	for (int32 IdxY = 0; IdxY < SizeY - 1; IdxY++)
+	{
+		for (int32 IdxX = 0; IdxX < SizeX - 1; IdxX++)
+		{
+			const int32 CoordX = IdxX + MinX;
+			const int32 CoordY = IdxY + MinY;
+			const int32 SampleIdx = ((bMirrored ? CoordX : (NumCols - CoordX - 1)) * NumCols) + CoordY;
+
+			const bool bIsHole = PrefetchedHeightfieldSamples.Holes[SampleIdx];
+			if (bIsHole)
+			{
+				continue;
+			}
+
+			const int32 I00 = (IdxX + 0) + (IdxY + 0)*SizeX;
+			int32 I01 = (IdxX + 0) + (IdxY + 1)*SizeX;
+			int32 I10 = (IdxX + 1) + (IdxY + 0)*SizeX;
+			const int32 I11 = (IdxX + 1) + (IdxY + 1)*SizeX;
+			if (bMirrored)
+			{
+				Swap(I01, I10);
+			}
+
+			IndexBuffer.Add(VertOffset + I00);
+			IndexBuffer.Add(VertOffset + I11);
+			IndexBuffer.Add(VertOffset + I10);
+
+			IndexBuffer.Add(VertOffset + I00);
+			IndexBuffer.Add(VertOffset + I01);
+			IndexBuffer.Add(VertOffset + I11);
+		}
+	}
+}
+
+void ExportCustomMesh(const FVector* InVertices, int32 NumVerts, const int32* InIndices, int32 NumIndices, const FTransform& LocalToWorld,
+					  TNavStatArray<float>& VertexBuffer, TNavStatArray<int32>& IndexBuffer, FBox& UnrealBounds)
+{
+	if (NumVerts <= 0 || NumIndices <= 0)
+	{
+		return;
+	}
+
+	int32 VertOffset = VertexBuffer.Num() / 3;
+	VertexBuffer.Reserve(VertexBuffer.Num() + NumVerts*3);
+	IndexBuffer.Reserve(IndexBuffer.Num() + NumIndices);
+
+	const bool bFlipCullMode = (LocalToWorld.GetDeterminant() < 0.f);
+	const int32 IndexOrder[3] = { bFlipCullMode ? 2 : 0, 1, bFlipCullMode ? 0 : 2 };
+
+#if SHOW_NAV_EXPORT_PREVIEW
+	UWorld* DebugWorld = FindEditorWorld();
+#endif // SHOW_NAV_EXPORT_PREVIEW
+
+	// Add vertices
+	for (int32 i = 0; i < NumVerts; ++i)
+	{
+		const FVector UnrealCoords = LocalToWorld.TransformPosition(InVertices[i]);
+		UnrealBounds += UnrealCoords;
+
+		VertexBuffer.Add(UnrealCoords.X);
+		VertexBuffer.Add(UnrealCoords.Y);
+		VertexBuffer.Add(UnrealCoords.Z);
+	}
+
+	// Add indices
+	for (int32 i = 0; i < NumIndices; i += 3)
+	{
+		IndexBuffer.Add(InIndices[i + IndexOrder[0]] + VertOffset);
+		IndexBuffer.Add(InIndices[i + IndexOrder[1]] + VertOffset);
+		IndexBuffer.Add(InIndices[i + IndexOrder[2]] + VertOffset);
+
+#if SHOW_NAV_EXPORT_PREVIEW
+		if (DebugWorld)
+		{
+			FVector V0(VertexBuffer[(VertOffset + InIndices[i + IndexOrder[0]]) * 3+0], VertexBuffer[(VertOffset + InIndices[i + IndexOrder[0]]) * 3+1], VertexBuffer[(VertOffset + InIndices[i + IndexOrder[0]]) * 3+2]);
+			FVector V1(VertexBuffer[(VertOffset + InIndices[i + IndexOrder[1]]) * 3+0], VertexBuffer[(VertOffset + InIndices[i + IndexOrder[1]]) * 3+1], VertexBuffer[(VertOffset + InIndices[i + IndexOrder[1]]) * 3+2]);
+			FVector V2(VertexBuffer[(VertOffset + InIndices[i + IndexOrder[2]]) * 3+0], VertexBuffer[(VertOffset + InIndices[i + IndexOrder[2]]) * 3+1], VertexBuffer[(VertOffset + InIndices[i + IndexOrder[2]]) * 3+2]);
+
+			DrawDebugLine(DebugWorld, V0, V1, bFlipCullMode ? FColor::Red : FColor::Blue, true);
+			DrawDebugLine(DebugWorld, V1, V2, bFlipCullMode ? FColor::Red : FColor::Blue, true);
+			DrawDebugLine(DebugWorld, V2, V0, bFlipCullMode ? FColor::Red : FColor::Blue, true);
+		}
+#endif // SHOW_NAV_EXPORT_PREVIEW
+	}
+}
+
+template<typename OtherAllocator>
+FORCEINLINE_DEBUGGABLE void AddFacesToRecast(TArray<FVector, OtherAllocator>& InVerts, TArray<int32, OtherAllocator>& InFaces,
+											 TNavStatArray<float>& OutVerts, TNavStatArray<int32>& OutIndices, FBox& UnrealBounds)
+{
+	// Add indices
+	int32 StartVertOffset = OutVerts.Num();
+	if (StartVertOffset > 0)
+	{
+		const int32 FirstIndex = OutIndices.AddUninitialized(InFaces.Num());
+		for (int32 Idx=0; Idx < InFaces.Num(); ++Idx)
+		{
+			OutIndices[FirstIndex + Idx] = InFaces[Idx]+StartVertOffset;
+		}
+	}
+	else
+	{
+		OutIndices.Append(InFaces);
+	}
+
+	// Add vertices
+	for (int32 i = 0; i < InVerts.Num(); i++)
+	{
+		const FVector& RecastCoords = InVerts[i];
+		OutVerts.Add(RecastCoords.X);
+		OutVerts.Add(RecastCoords.Y);
+		OutVerts.Add(RecastCoords.Z);
+
+		UnrealBounds += Recast2UnrealPoint(RecastCoords);
+	}
+}
+
+FORCEINLINE_DEBUGGABLE void ExportRigidBodyConvexElements(UBodySetup& BodySetup, TNavStatArray<float>& VertexBuffer, TNavStatArray<int32>& IndexBuffer,
+														  TNavStatArray<int32>& ShapeBuffer, FBox& UnrealBounds, const FTransform& LocalToWorld)
+{
+#if WITH_PHYSX
+	const int32 ConvexCount = BodySetup.AggGeom.ConvexElems.Num();
+	FKConvexElem const * ConvexElem = BodySetup.AggGeom.ConvexElems.GetData();
+
+	const FTransform NegXScale(FQuat::Identity, FVector::ZeroVector, FVector(-1, 1, 1));
+
+	for(int32 i=0; i< ConvexCount; ++i, ++ConvexElem)
+	{
+		// Store index of first vertex in shape buffer
+		ShapeBuffer.Add(VertexBuffer.Num() / 3);
+
+		// Get verts/triangles from this hull.
+		if (!ConvexElem->GetConvexMesh() && ConvexElem->GetMirroredConvexMesh())
+		{
+			// If there is only a NegX mesh (e.g. a mirrored volume), use it
+			ExportPxConvexMesh(ConvexElem->GetMirroredConvexMesh(), NegXScale * LocalToWorld, VertexBuffer, IndexBuffer, UnrealBounds);
+		}
+		else
+		{
+			// Otherwise use the regular mesh in the case that both exist
+			ExportPxConvexMesh(ConvexElem->GetConvexMesh(), LocalToWorld, VertexBuffer, IndexBuffer, UnrealBounds);
+		}
+	}
+#endif // WITH_PHYSX
+}
+
+FORCEINLINE_DEBUGGABLE void ExportRigidBodyTriMesh(UBodySetup& BodySetup, TNavStatArray<float>& VertexBuffer, TNavStatArray<int32>& IndexBuffer,
+												   FBox& UnrealBounds, const FTransform& LocalToWorld)
+{
+#if WITH_PHYSX
+	if (BodySetup.GetCollisionTraceFlag() == CTF_UseComplexAsSimple)
+	{
+		for(PxTriangleMesh* TriMesh : BodySetup.TriMeshes)
+		{
+			if (TriMesh->getTriangleMeshFlags() & PxTriangleMeshFlag::e16_BIT_INDICES)
+			{
+				ExportPxTriMesh<PxU16>(TriMesh, LocalToWorld, VertexBuffer, IndexBuffer, UnrealBounds);
+			}
+			else
+			{
+				ExportPxTriMesh<PxU32>(TriMesh, LocalToWorld, VertexBuffer, IndexBuffer, UnrealBounds);
+			}
+		}
+	}
+#endif // WITH_PHYSX
+}
+
+void ExportRigidBodyBoxElements(const FKAggregateGeom& AggGeom, TNavStatArray<float>& VertexBuffer, TNavStatArray<int32>& IndexBuffer,
+								TNavStatArray<int32>& ShapeBuffer, FBox& UnrealBounds, const FTransform& LocalToWorld, int32 NumExistingVerts)
+{
+	for (int32 i = 0; i < AggGeom.BoxElems.Num(); i++)
+	{
+		const FKBoxElem& BoxInfo = AggGeom.BoxElems[i];
+		const FMatrix ElemTM = BoxInfo.GetTransform().ToMatrixWithScale() * LocalToWorld.ToMatrixWithScale();
+		const FVector Extent(BoxInfo.X * 0.5f, BoxInfo.Y * 0.5f, BoxInfo.Z * 0.5f);
+
+		const int32 VertBase = (NumExistingVerts + VertexBuffer.Num()) / 3;
+		
+		// Store index of first vertex in shape buffer
+		ShapeBuffer.Add(VertBase);
+
+		// add box vertices
+		FVector UnrealVerts[] = {
+			ElemTM.TransformPosition(FVector(-Extent.X, -Extent.Y,  Extent.Z)),
+			ElemTM.TransformPosition(FVector( Extent.X, -Extent.Y,  Extent.Z)),
+			ElemTM.TransformPosition(FVector(-Extent.X, -Extent.Y, -Extent.Z)),
+			ElemTM.TransformPosition(FVector( Extent.X, -Extent.Y, -Extent.Z)),
+			ElemTM.TransformPosition(FVector(-Extent.X,  Extent.Y,  Extent.Z)),
+			ElemTM.TransformPosition(FVector( Extent.X,  Extent.Y,  Extent.Z)),
+			ElemTM.TransformPosition(FVector(-Extent.X,  Extent.Y, -Extent.Z)),
+			ElemTM.TransformPosition(FVector( Extent.X,  Extent.Y, -Extent.Z))
+		};
+
+		for (int32 iv = 0; iv < ARRAY_COUNT(UnrealVerts); iv++)
+		{
+			UnrealBounds += UnrealVerts[iv];
+
+			VertexBuffer.Add(UnrealVerts[iv].X);
+			VertexBuffer.Add(UnrealVerts[iv].Y);
+			VertexBuffer.Add(UnrealVerts[iv].Z);
+		}
+		
+		IndexBuffer.Add(VertBase + 3); IndexBuffer.Add(VertBase + 2); IndexBuffer.Add(VertBase + 0);
+		IndexBuffer.Add(VertBase + 3); IndexBuffer.Add(VertBase + 0); IndexBuffer.Add(VertBase + 1);
+		IndexBuffer.Add(VertBase + 7); IndexBuffer.Add(VertBase + 3); IndexBuffer.Add(VertBase + 1);
+		IndexBuffer.Add(VertBase + 7); IndexBuffer.Add(VertBase + 1); IndexBuffer.Add(VertBase + 5);
+		IndexBuffer.Add(VertBase + 6); IndexBuffer.Add(VertBase + 7); IndexBuffer.Add(VertBase + 5);
+		IndexBuffer.Add(VertBase + 6); IndexBuffer.Add(VertBase + 5); IndexBuffer.Add(VertBase + 4);
+		IndexBuffer.Add(VertBase + 2); IndexBuffer.Add(VertBase + 6); IndexBuffer.Add(VertBase + 4);
+		IndexBuffer.Add(VertBase + 2); IndexBuffer.Add(VertBase + 4); IndexBuffer.Add(VertBase + 0);
+		IndexBuffer.Add(VertBase + 1); IndexBuffer.Add(VertBase + 0); IndexBuffer.Add(VertBase + 4);
+		IndexBuffer.Add(VertBase + 1); IndexBuffer.Add(VertBase + 4); IndexBuffer.Add(VertBase + 5);
+		IndexBuffer.Add(VertBase + 7); IndexBuffer.Add(VertBase + 6); IndexBuffer.Add(VertBase + 2);
+		IndexBuffer.Add(VertBase + 7); IndexBuffer.Add(VertBase + 2); IndexBuffer.Add(VertBase + 3);
+	}
+}
+
+void ExportRigidBodySphylElements(const FKAggregateGeom& AggGeom, TNavStatArray<float>& VertexBuffer, TNavStatArray<int32>& IndexBuffer,
+								  TNavStatArray<int32>& ShapeBuffer, FBox& UnrealBounds, const FTransform& LocalToWorld, int32 NumExistingVerts)
+{
+	TArray<FVector> ArcVerts;
+
+	for (int32 i = 0; i < AggGeom.SphylElems.Num(); i++)
+	{
+		const FKSphylElem& SphylInfo = AggGeom.SphylElems[i];
+		const FMatrix ElemTM = SphylInfo.GetTransform().ToMatrixWithScale() * LocalToWorld.ToMatrixWithScale();
+
+		const int32 VertBase = (NumExistingVerts + VertexBuffer.Num()) / 3;
+
+		// Store index of first vertex in shape buffer
+		ShapeBuffer.Add(VertBase);
+
+		const int32 NumSides = 16;
+		const int32 NumRings = (NumSides/2) + 1;
+
+		// The first/last arc are on top of each other.
+		const int32 NumVerts = (NumSides+1) * (NumRings+1);
+
+		ArcVerts.Reset();
+		ArcVerts.AddZeroed(NumRings+1);
+		for (int32 RingIdx=0; RingIdx<NumRings+1; RingIdx++)
+		{
+			float Angle;
+			float ZOffset;
+			if (RingIdx <= NumSides/4)
+			{
+				Angle = ((float)RingIdx/(NumRings-1)) * PI;
+				ZOffset = 0.5 * SphylInfo.Length;
+			}
+			else
+			{
+				Angle = ((float)(RingIdx-1)/(NumRings-1)) * PI;
+				ZOffset = -0.5 * SphylInfo.Length;
+			}
+
+			// Note- unit sphere, so position always has mag of one. We can just use it for normal!		
+			FVector SpherePos;
+			SpherePos.X = 0.0f;
+			SpherePos.Y = SphylInfo.Radius * FMath::Sin(Angle);
+			SpherePos.Z = SphylInfo.Radius * FMath::Cos(Angle);
+
+			ArcVerts[RingIdx] = SpherePos + FVector(0,0,ZOffset);
+		}
+
+		// Then rotate this arc NumSides+1 times.
+		for (int32 SideIdx=0; SideIdx<NumSides+1; SideIdx++)
+		{
+			const FRotator ArcRotator(0, 360.f * ((float)SideIdx/NumSides), 0);
+			const FRotationMatrix ArcRot( ArcRotator );
+			const FMatrix ArcTM = ArcRot * ElemTM;
+
+			for(int32 VertIdx=0; VertIdx<NumRings+1; VertIdx++)
+			{
+				const FVector UnrealVert = ArcTM.TransformPosition(ArcVerts[VertIdx]);
+				UnrealBounds += UnrealVert;
+
+				VertexBuffer.Add(UnrealVert.X);
+				VertexBuffer.Add(UnrealVert.Y);
+				VertexBuffer.Add(UnrealVert.Z);
+			}
+		}
+
+		// Add all of the triangles to the mesh.
+		for (int32 SideIdx=0; SideIdx<NumSides; SideIdx++)
+		{
+			const int32 a0start = VertBase + ((SideIdx+0) * (NumRings+1));
+			const int32 a1start = VertBase + ((SideIdx+1) * (NumRings+1));
+
+			for (int32 RingIdx=0; RingIdx<NumRings; RingIdx++)
+			{
+				IndexBuffer.Add(a0start + RingIdx + 0); IndexBuffer.Add(a1start + RingIdx + 0); IndexBuffer.Add(a0start + RingIdx + 1);
+				IndexBuffer.Add(a1start + RingIdx + 0); IndexBuffer.Add(a1start + RingIdx + 1); IndexBuffer.Add(a0start + RingIdx + 1);
+			}
+		}
+	}
+}
+
+void ExportRigidBodySphereElements(const FKAggregateGeom& AggGeom, TNavStatArray<float>& VertexBuffer, TNavStatArray<int32>& IndexBuffer,
+								   TNavStatArray<int32>& ShapeBuffer, FBox& UnrealBounds, const FTransform& LocalToWorld, int32 NumExistingVerts)
+{
+	TArray<FVector> ArcVerts;
+
+	for (int32 i = 0; i < AggGeom.SphereElems.Num(); i++)
+	{
+		const FKSphereElem& SphereInfo = AggGeom.SphereElems[i];
+		const FMatrix ElemTM = SphereInfo.GetTransform().ToMatrixWithScale() * LocalToWorld.ToMatrixWithScale();
+
+		const int32 VertBase = (NumExistingVerts + VertexBuffer.Num()) / 3;
+
+		// Store index of first vertex in shape buffer
+		ShapeBuffer.Add(VertBase);
+
+		const int32 NumSides = 16;
+		const int32 NumRings = (NumSides/2) + 1;
+
+		// The first/last arc are on top of each other.
+		const int32 NumVerts = (NumSides+1) * (NumRings+1);
+
+		ArcVerts.Reset();
+		ArcVerts.AddZeroed(NumRings+1);
+		for (int32 RingIdx=0; RingIdx<NumRings+1; RingIdx++)
+		{
+			float Angle = ((float)RingIdx/NumRings) * PI;
+
+			// Note- unit sphere, so position always has mag of one. We can just use it for normal!			
+			FVector& ArcVert = ArcVerts[RingIdx];
+			ArcVert.X = 0.0f;
+			ArcVert.Y = SphereInfo.Radius * FMath::Sin(Angle);
+			ArcVert.Z = SphereInfo.Radius * FMath::Cos(Angle);
+		}
+
+		// Then rotate this arc NumSides+1 times.
+		for (int32 SideIdx=0; SideIdx<NumSides+1; SideIdx++)
+		{
+			const FRotator ArcRotator(0, 360.f * ((float)SideIdx/NumSides), 0);
+			const FRotationMatrix ArcRot( ArcRotator );
+			const FMatrix ArcTM = ArcRot * ElemTM;
+
+			for(int32 VertIdx=0; VertIdx<NumRings+1; VertIdx++)
+			{
+				const FVector UnrealVert = ArcTM.TransformPosition(ArcVerts[VertIdx]);
+				UnrealBounds += UnrealVert;
+
+				VertexBuffer.Add(UnrealVert.X);
+				VertexBuffer.Add(UnrealVert.Y);
+				VertexBuffer.Add(UnrealVert.Z);
+			}
+		}
+
+		// Add all of the triangles to the mesh.
+		for (int32 SideIdx=0; SideIdx<NumSides; SideIdx++)
+		{
+			const int32 a0start = VertBase + ((SideIdx+0) * (NumRings+1));
+			const int32 a1start = VertBase + ((SideIdx+1) * (NumRings+1));
+
+			for (int32 RingIdx=0; RingIdx<NumRings; RingIdx++)
+			{
+				IndexBuffer.Add(a0start + RingIdx + 0); IndexBuffer.Add(a1start + RingIdx + 0); IndexBuffer.Add(a0start + RingIdx + 1);
+				IndexBuffer.Add(a1start + RingIdx + 0); IndexBuffer.Add(a1start + RingIdx + 1); IndexBuffer.Add(a0start + RingIdx + 1);
+			}
+		}
+	}
+}
+
+FORCEINLINE_DEBUGGABLE void ExportRigidBodySetup(UBodySetup& BodySetup, TNavStatArray<float>& VertexBuffer, TNavStatArray<int32>& IndexBuffer,
+												 FBox& UnrealBounds, const FTransform& LocalToWorld)
+{
+	// Make sure meshes are created before we try and export them
+	BodySetup.CreatePhysicsMeshes();
+
+	static TNavStatArray<int32> TemporaryShapeBuffer;
+	const int32 NumExistingsVerts = VertexBuffer.Num();
+
+	ExportRigidBodyTriMesh(BodySetup, VertexBuffer, IndexBuffer, UnrealBounds, LocalToWorld);
+	ExportRigidBodyConvexElements(BodySetup, VertexBuffer, IndexBuffer, TemporaryShapeBuffer, UnrealBounds, LocalToWorld);
+	ExportRigidBodyBoxElements(BodySetup.AggGeom, VertexBuffer, IndexBuffer, TemporaryShapeBuffer, UnrealBounds, LocalToWorld, NumExistingsVerts);
+	ExportRigidBodySphylElements(BodySetup.AggGeom, VertexBuffer, IndexBuffer, TemporaryShapeBuffer, UnrealBounds, LocalToWorld, NumExistingsVerts);
+	ExportRigidBodySphereElements(BodySetup.AggGeom, VertexBuffer, IndexBuffer, TemporaryShapeBuffer, UnrealBounds, LocalToWorld, NumExistingsVerts);
+
+	TemporaryShapeBuffer.Reset();
+}
+
+FORCEINLINE_DEBUGGABLE void ExportComponent(UActorComponent* Component, FRecastGeometryExport& GeomExport, const FBox* ClipBounds=NULL)
+{
+#if WITH_PHYSX
+	bool bHasData = false;
+
+	UPrimitiveComponent* PrimComp = Cast<UPrimitiveComponent>(Component);
+	if (PrimComp && PrimComp->IsNavigationRelevant() && (PrimComp->HasCustomNavigableGeometry() != EHasCustomNavigableGeometry::DontExport))
+	{
+		if ((PrimComp->HasCustomNavigableGeometry() != EHasCustomNavigableGeometry::Type::No) && !PrimComp->DoCustomNavigableGeometryExport(GeomExport))
+		{
+			bHasData = true;
+		}
+
+		UBodySetup* BodySetup = PrimComp->GetBodySetup();
+		if (BodySetup)
+		{
+			if (!bHasData)
+			{
+				ExportRigidBodySetup(*BodySetup, GeomExport.VertexBuffer, GeomExport.IndexBuffer, GeomExport.Data->Bounds, PrimComp->GetComponentTransform());
+				bHasData = true;
+			}
+
+			GeomExport.SlopeOverride = BodySetup->WalkableSlopeOverride;
+		}
+	}
+#endif // WITH_PHYSX
+}
+
+FORCEINLINE void TransformVertexSoupToRecast(const TArray<FVector>& VertexSoup, TNavStatArray<FVector>& Verts, TNavStatArray<int32>& Faces)
+{
+	if (VertexSoup.Num() == 0)
+	{
+		return;
+	}
+
+	check(VertexSoup.Num() % 3 == 0);
+
+	const int32 StaticFacesCount = VertexSoup.Num() / 3;
+	int32 VertsCount = Verts.Num();
+	const FVector* Vertex = VertexSoup.GetData();
+
+	for (int32 k = 0; k < StaticFacesCount; ++k, Vertex += 3)
+	{
+		Verts.Add(Unreal2RecastPoint(Vertex[0]));
+		Verts.Add(Unreal2RecastPoint(Vertex[1]));
+		Verts.Add(Unreal2RecastPoint(Vertex[2]));
+		Faces.Add(VertsCount + 2);
+		Faces.Add(VertsCount + 1);
+		Faces.Add(VertsCount + 0);
+			
+		VertsCount += 3;
+	}
+}
+
+FORCEINLINE void CovertCoordDataToRecast(TNavStatArray<float>& Coords)
+{
+	float* CoordPtr = Coords.GetData();
+	const int32 MaxIt = Coords.Num() / 3;
+	for (int32 i = 0; i < MaxIt; i++)
+	{
+		CoordPtr[0] = -CoordPtr[0];
+
+		const float TmpV = -CoordPtr[1];
+		CoordPtr[1] = CoordPtr[2];
+		CoordPtr[2] = TmpV;
+
+		CoordPtr += 3;
+	}
+}
+
+void ExportVertexSoup(const TArray<FVector>& VertexSoup, TNavStatArray<float>& VertexBuffer, TNavStatArray<int32>& IndexBuffer, FBox& UnrealBounds)
+{
+	if (VertexSoup.Num())
+	{
+		check(VertexSoup.Num() % 3 == 0);
+		
+		int32 VertBase = VertexBuffer.Num() / 3;
+		VertexBuffer.Reserve(VertexSoup.Num() * 3);
+		IndexBuffer.Reserve(VertexSoup.Num() / 3);
+
+		const int32 NumVerts = VertexSoup.Num();
+		for (int32 i = 0; i < NumVerts; i++)
+		{
+			const FVector& UnrealCoords = VertexSoup[i];
+			UnrealBounds += UnrealCoords;
+
+			const FVector RecastCoords = Unreal2RecastPoint(UnrealCoords);
+			VertexBuffer.Add(RecastCoords.X);
+			VertexBuffer.Add(RecastCoords.Y);
+			VertexBuffer.Add(RecastCoords.Z);
+		}
+
+		const int32 NumFaces = VertexSoup.Num() / 3;
+		for (int32 i = 0; i < NumFaces; i++)
+		{
+			IndexBuffer.Add(VertBase + 2);
+			IndexBuffer.Add(VertBase + 1);
+			IndexBuffer.Add(VertBase + 0);
+			VertBase += 3;
+		}
+	}
+}
+
+} // namespace RecastGeometryExport
+
+#if WITH_PHYSX
+void FRecastGeometryExport::ExportPxTriMesh16Bit(physx::PxTriangleMesh const * const TriMesh, const FTransform& LocalToWorld)
+{
+	RecastGeometryExport::ExportPxTriMesh<PxU16>(TriMesh, LocalToWorld, VertexBuffer, IndexBuffer, Data->Bounds);
+}
+
+void FRecastGeometryExport::ExportPxTriMesh32Bit(physx::PxTriangleMesh const * const TriMesh, const FTransform& LocalToWorld)
+{
+	RecastGeometryExport::ExportPxTriMesh<PxU32>(TriMesh, LocalToWorld, VertexBuffer, IndexBuffer, Data->Bounds);
+}
+
+void FRecastGeometryExport::ExportPxConvexMesh(physx::PxConvexMesh const * const ConvexMesh, const FTransform& LocalToWorld)
+{
+	RecastGeometryExport::ExportPxConvexMesh(ConvexMesh, LocalToWorld, VertexBuffer, IndexBuffer, Data->Bounds);
+}
+
+void FRecastGeometryExport::ExportPxHeightField(physx::PxHeightField const * const HeightField, const FTransform& LocalToWorld)
+{
+	RecastGeometryExport::ExportPxHeightField(HeightField, LocalToWorld, VertexBuffer, IndexBuffer, Data->Bounds);
+}
+#endif // WITH_PHYSX
+
+void FRecastGeometryExport::ExportHeightFieldSlice(const FNavHeightfieldSamples& PrefetchedHeightfieldSamples, const int32 NumRows, const int32 NumCols, const FTransform& LocalToWorld, const FBox& SliceBox)
+{
+	RecastGeometryExport::ExportHeightFieldSlice(PrefetchedHeightfieldSamples, NumRows, NumCols, LocalToWorld, VertexBuffer, IndexBuffer, SliceBox, Data->Bounds);
+}
+
+void FRecastGeometryExport::ExportCustomMesh(const FVector* InVertices, int32 NumVerts, const int32* InIndices, int32 NumIndices, const FTransform& LocalToWorld)
+{
+	RecastGeometryExport::ExportCustomMesh(InVertices, NumVerts, InIndices, NumIndices, LocalToWorld, VertexBuffer, IndexBuffer, Data->Bounds);
+}
+
+void FRecastGeometryExport::ExportRigidBodySetup(UBodySetup& BodySetup, const FTransform& LocalToWorld)
+{
+	RecastGeometryExport::ExportRigidBodySetup(BodySetup, VertexBuffer, IndexBuffer, Data->Bounds, LocalToWorld);
+}
+
+void FRecastGeometryExport::AddNavModifiers(const FCompositeNavModifier& Modifiers)
+{
+	Data->Modifiers.Add(Modifiers);
+}
+
+void FRecastGeometryExport::SetNavDataPerInstanceTransformDelegate(const FNavDataPerInstanceTransformDelegate& InDelegate)
+{
+	Data->NavDataPerInstanceTransformDelegate = InDelegate;
+}
+
+static void PartialTransformConvexHull(FConvexNavAreaData& ConvexData, const FTransform& LocalToWorld)
+{
+	FVector ScaleXY = LocalToWorld.GetScale3D().GetAbs();
+	ScaleXY.Z = 1.f;
+
+	FVector TranslationXY = LocalToWorld.GetLocation();
+	TranslationXY.Z = 0.f;
+
+	for (FVector& Point : ConvexData.Points)
+	{
+		Point = Point*ScaleXY + TranslationXY;
+	}
+
+	ConvexData.MaxZ+= LocalToWorld.GetLocation().Z;
+	ConvexData.MinZ+= LocalToWorld.GetLocation().Z;
+}
+
+FORCEINLINE void GrowConvexHull(const float ExpandBy, const TArray<FVector>& Verts, TArray<FVector>& OutResult)
+{
+	if (Verts.Num() < 3)
+	{
+		return;
+	}
+
+	struct FSimpleLine
+	{
+		FVector P1, P2;
+
+		FSimpleLine() {}
+
+		FSimpleLine(FVector Point1, FVector Point2) 
+			: P1(Point1), P2(Point2) 
+		{
+
+		}
+		static FVector Intersection(const FSimpleLine& Line1, const FSimpleLine& Line2)
+		{
+			const float A1 = Line1.P2.X - Line1.P1.X;
+			const float B1 = Line2.P1.X - Line2.P2.X;
+			const float C1 = Line2.P1.X - Line1.P1.X;
+
+			const float A2 = Line1.P2.Y - Line1.P1.Y;
+			const float B2 = Line2.P1.Y - Line2.P2.Y;
+			const float C2 = Line2.P1.Y - Line1.P1.Y;
+
+			const float Denominator = A2*B1 - A1*B2;
+			if (Denominator != 0)
+			{
+				const float t = (B1*C2 - B2*C1) / Denominator;
+				return Line1.P1 + t * (Line1.P2 - Line1.P1);
+			}
+
+			return FVector::ZeroVector;
+		}
+	};
+
+	TArray<FVector> AllVerts(Verts);
+	AllVerts.Add(Verts[0]);
+	AllVerts.Add(Verts[1]);
+
+	const int32 VertsCount = AllVerts.Num();
+	const FQuat Rotation90(FVector(0, 0, 1), FMath::DegreesToRadians(90));
+
+	float RotationAngle = MAX_FLT;
+	for (int32 Index = 0; Index < VertsCount - 2; ++Index)
+	{
+		const FVector& V1 = AllVerts[Index + 0];
+		const FVector& V2 = AllVerts[Index + 1];
+		const FVector& V3 = AllVerts[Index + 2];
+
+		const FVector V01 = (V1 - V2).GetSafeNormal();
+		const FVector V12 = (V2 - V3).GetSafeNormal();
+		const FVector NV1 = Rotation90.RotateVector(V01);
+		const float d = FVector::DotProduct(NV1, V12);
+
+		if (d < 0)
+		{
+			// CW
+			RotationAngle = -90;
+			break;
+		}
+		else if (d > 0)
+		{
+			//CCW
+			RotationAngle = 90;
+			break;
+		}
+	}
+
+	// check if we detected CW or CCW direction
+	if (RotationAngle >= BIG_NUMBER)
+	{
+		return;
+	}
+
+	const float ExpansionThreshold = 2 * ExpandBy;
+	const float ExpansionThresholdSQ = ExpansionThreshold * ExpansionThreshold;
+	const FQuat Rotation(FVector(0, 0, 1), FMath::DegreesToRadians(RotationAngle));
+	FSimpleLine PreviousLine;
+	OutResult.Reserve(Verts.Num());
+	for (int32 Index = 0; Index < VertsCount-2; ++Index)
+	{
+		const FVector& V1 = AllVerts[Index + 0];
+		const FVector& V2 = AllVerts[Index + 1];
+		const FVector& V3 = AllVerts[Index + 2];
+
+		FSimpleLine Line1;
+		if (Index > 0)
+		{
+			Line1 = PreviousLine;
+		}
+		else
+		{
+			const FVector V01 = (V1 - V2).GetSafeNormal();
+			const FVector N1 = Rotation.RotateVector(V01).GetSafeNormal();
+			const FVector MoveDir1 = N1 * ExpandBy;
+			Line1 = FSimpleLine(V1 + MoveDir1, V2 + MoveDir1);
+		}
+
+		const FVector V12 = (V2 - V3).GetSafeNormal();
+		const FVector N2 = Rotation.RotateVector(V12).GetSafeNormal();
+		const FVector MoveDir2 = N2 * ExpandBy;
+		const FSimpleLine Line2(V2 + MoveDir2, V3 + MoveDir2);
+
+		const FVector NewPoint = FSimpleLine::Intersection(Line1, Line2);
+		if (NewPoint == FVector::ZeroVector)
+		{
+			// both lines are parallel so just move our point by expansion distance
+			OutResult.Add(V2 + MoveDir2);
+		}
+		else
+		{
+			const FVector VectorToNewPoint = NewPoint - V2;
+			const float DistToNewVector = VectorToNewPoint.SizeSquared2D();
+			if (DistToNewVector > ExpansionThresholdSQ)
+			{
+				//clamp our point to not move to far from original location
+				const FVector HelpPos = V2 + VectorToNewPoint.GetSafeNormal2D() * ExpandBy * 1.4142;
+				OutResult.Add(HelpPos);
+			}
+			else
+			{
+				OutResult.Add(NewPoint);
+			}
+		}
+
+		PreviousLine = Line2;
+	}
+}
+
+//----------------------------------------------------------------------//
+
+struct FOffMeshData
+{
+	TArray<dtOffMeshLinkCreateParams> LinkParams;
+	const TMap<const UClass*, int32>* AreaClassToIdMap;
+	const ARecastNavMesh::FNavPolyFlags* FlagsPerArea;
+
+	FOffMeshData() : AreaClassToIdMap(NULL), FlagsPerArea(NULL) {}
+
+	FORCEINLINE void Reserve(const uint32 ElementsCount)
+	{
+		LinkParams.Reserve(ElementsCount);
+	}
+
+	void AddLinks(const TArray<FNavigationLink>& Links, const FTransform& LocalToWorld, int32 AgentIndex, float DefaultSnapHeight)
+	{
+		for (int32 LinkIndex = 0; LinkIndex < Links.Num(); ++LinkIndex)
+		{
+			const FNavigationLink& Link = Links[LinkIndex];
+			if (!Link.SupportedAgents.Contains(AgentIndex))
+			{
+				continue;
+			}
+
+			dtOffMeshLinkCreateParams NewInfo;
+			FMemory::Memzero(NewInfo);
+
+			// not doing anything to link's points order - should be already ordered properly by link processor
+			StoreUnrealPoint(NewInfo.vertsA0, LocalToWorld.TransformPosition(Link.Left));
+			StoreUnrealPoint(NewInfo.vertsB0, LocalToWorld.TransformPosition(Link.Right));
+
+			NewInfo.type = DT_OFFMESH_CON_POINT | 
+				(Link.Direction == ENavLinkDirection::BothWays ? DT_OFFMESH_CON_BIDIR : 0) |
+				(Link.bSnapToCheapestArea ? DT_OFFMESH_CON_CHEAPAREA : 0);
+
+			NewInfo.snapRadius = Link.SnapRadius;
+			NewInfo.snapHeight = Link.bUseSnapHeight ? Link.SnapHeight : DefaultSnapHeight;
+			NewInfo.userID = Link.UserId;
+
+			UClass* AreaClass = Link.GetAreaClass();
+			const int32* AreaID = AreaClassToIdMap->Find(AreaClass);
+			if (AreaID != NULL)
+			{
+				NewInfo.area = *AreaID;
+				NewInfo.polyFlag = FlagsPerArea[*AreaID];
+			}
+			else
+			{
+				UE_LOG(LogNavigation, Warning, TEXT("FRecastTileGenerator: Trying to use undefined area class while defining Off-Mesh links! (%s)"), *GetNameSafe(AreaClass));
+			}
+
+			// snap area is currently not supported for regular (point-point) offmesh links
+
+			LinkParams.Add(NewInfo);
+		}
+	}
+	void AddSegmentLinks(const TArray<FNavigationSegmentLink>& Links, const FTransform& LocalToWorld, int32 AgentIndex, float DefaultSnapHeight)
+	{
+		for (int32 LinkIndex = 0; LinkIndex < Links.Num(); ++LinkIndex)
+		{
+			const FNavigationSegmentLink& Link = Links[LinkIndex];
+			if (!Link.SupportedAgents.Contains(AgentIndex))
+			{
+				continue;
+			}
+
+			dtOffMeshLinkCreateParams NewInfo;
+			FMemory::Memzero(NewInfo);
+
+			// not doing anything to link's points order - should be already ordered properly by link processor
+			StoreUnrealPoint(NewInfo.vertsA0, LocalToWorld.TransformPosition(Link.LeftStart));
+			StoreUnrealPoint(NewInfo.vertsA1, LocalToWorld.TransformPosition(Link.LeftEnd));
+			StoreUnrealPoint(NewInfo.vertsB0, LocalToWorld.TransformPosition(Link.RightStart));
+			StoreUnrealPoint(NewInfo.vertsB1, LocalToWorld.TransformPosition(Link.RightEnd));
+
+			NewInfo.type = DT_OFFMESH_CON_SEGMENT | (Link.Direction == ENavLinkDirection::BothWays ? DT_OFFMESH_CON_BIDIR : 0);
+			NewInfo.snapRadius = Link.SnapRadius;
+			NewInfo.snapHeight = Link.bUseSnapHeight ? Link.SnapHeight : DefaultSnapHeight;
+			NewInfo.userID = Link.UserId;
+
+			UClass* AreaClass = Link.GetAreaClass();
+			const int32* AreaID = AreaClassToIdMap->Find(AreaClass);
+			if (AreaID != NULL)
+			{
+				NewInfo.area = *AreaID;
+				NewInfo.polyFlag = FlagsPerArea[*AreaID];
+			}
+			else
+			{
+				UE_LOG(LogNavigation, Warning, TEXT("FRecastTileGenerator: Trying to use undefined area class while defining Off-Mesh links! (%s)"), *GetNameSafe(AreaClass));
+			}
+
+			LinkParams.Add(NewInfo);
+		}
+	}
+
+protected:
+
+	void StoreUnrealPoint(float* dest, const FVector& UnrealPt)
+	{
+		const FVector RecastPt = Unreal2RecastPoint(UnrealPt);
+		dest[0] = RecastPt.X;
+		dest[1] = RecastPt.Y;
+		dest[2] = RecastPt.Z;
+	}
+};
+
+//----------------------------------------------------------------------//
+// FNavMeshBuildContext
+// A navmesh building reporting helper
+//----------------------------------------------------------------------//
+class FNavMeshBuildContext : public rcContext, public dtTileCacheLogContext
+{
+public:
+	FNavMeshBuildContext()
+		: rcContext(true)
+	{
+	}
+protected:
+	/// Logs a message.
+	///  @param[in]		category	The category of the message.
+	///  @param[in]		msg			The formatted message.
+	///  @param[in]		len			The length of the formatted message.
+	virtual void doLog(const rcLogCategory category, const char* Msg, const int32 /*len*/) 
+	{
+		switch (category) 
+		{
+		case RC_LOG_ERROR:
+			UE_LOG(LogNavigation, Error, TEXT("Recast: %s"), ANSI_TO_TCHAR( Msg ) );
+			break;
+		case RC_LOG_WARNING:
+			UE_LOG(LogNavigation, Log, TEXT("Recast: %s"), ANSI_TO_TCHAR( Msg ) );
+			break;
+		default:
+			UE_LOG(LogNavigation, Verbose, TEXT("Recast: %s"), ANSI_TO_TCHAR( Msg ) );
+			break;
+		}
+	}
+
+	virtual void doDtLog(const char* Msg, const int32 /*len*/)
+	{
+		UE_LOG(LogNavigation, Error, TEXT("Recast: %s"), ANSI_TO_TCHAR(Msg));
+	}
+};
+
+//----------------------------------------------------------------------//
+struct FTileCacheCompressor : public dtTileCacheCompressor
+{
+	struct FCompressedCacheHeader
+	{
+		int32 UncompressedSize;
+	};
+
+	virtual int32 maxCompressedSize(const int32 bufferSize)
+	{
+		return FMath::TruncToInt(bufferSize * 1.1f) + sizeof(FCompressedCacheHeader);
+	}
+
+	virtual dtStatus compress(const uint8* buffer, const int32 bufferSize,
+		uint8* compressed, const int32 maxCompressedSize, int32* compressedSize)
+	{
+		const int32 HeaderSize = sizeof(FCompressedCacheHeader);
+
+		FCompressedCacheHeader DataHeader;
+		DataHeader.UncompressedSize = bufferSize;
+		FMemory::Memcpy((void*)compressed, &DataHeader, HeaderSize);
+
+		uint8* DataPtr = compressed + HeaderSize;		
+		int32 DataSize = maxCompressedSize - HeaderSize;
+
+		FCompression::CompressMemory((ECompressionFlags)(COMPRESS_ZLIB | COMPRESS_BiasMemory),
+			(void*)DataPtr, DataSize, (const void*)buffer, bufferSize);
+
+		*compressedSize = DataSize + HeaderSize;
+		return DT_SUCCESS;
+	}
+
+	virtual dtStatus decompress(const uint8* compressed, const int32 compressedSize,
+		uint8* buffer, const int32 maxBufferSize, int32* bufferSize)
+	{
+		const int32 HeaderSize = sizeof(FCompressedCacheHeader);
+		
+		FCompressedCacheHeader DataHeader;
+		FMemory::Memcpy(&DataHeader, (void*)compressed, HeaderSize);
+
+		const uint8* DataPtr = compressed + HeaderSize;		
+		const int32 DataSize = compressedSize - HeaderSize;
+
+		FCompression::UncompressMemory((ECompressionFlags)(COMPRESS_ZLIB),
+			(void*)buffer, DataHeader.UncompressedSize, (const void*)DataPtr, DataSize);
+
+		*bufferSize = DataHeader.UncompressedSize;
+		return DT_SUCCESS;
+	}
+};
+
+struct FTileCacheAllocator : public dtTileCacheAlloc
+{
+	virtual void reset()
+	{
+		 check(0 && "dtTileCacheAlloc.reset() is not supported!");
+	}
+
+	virtual void* alloc(const int32 Size)
+	{
+		return dtAlloc(Size, DT_ALLOC_TEMP);
+	}
+
+	virtual void free(void* Data)
+	{
+		dtFree(Data);
+	}
+};
+
+//----------------------------------------------------------------------//
+// FVoxelCacheRasterizeContext
+//----------------------------------------------------------------------//
+
+struct FVoxelCacheRasterizeContext
+{
+	FVoxelCacheRasterizeContext()
+	{
+		RasterizeHF = NULL;
+	}
+
+	~FVoxelCacheRasterizeContext()
+	{
+		rcFreeHeightField(RasterizeHF);
+		RasterizeHF = 0;
+	}
+
+	void Create(int32 FieldSize, float CellSize, float CellHeight)
+	{
+		if (RasterizeHF == NULL)
+		{
+			const float DummyBounds[3] = { 0 };
+
+			RasterizeHF = rcAllocHeightfield();
+			rcCreateHeightfield(NULL, *RasterizeHF, FieldSize, FieldSize, DummyBounds, DummyBounds, CellSize, CellHeight);
+		}
+	}
+
+	void Reset()
+	{
+		rcResetHeightfield(*RasterizeHF);
+	}
+
+	void SetupForTile(const float* TileBMin, const float* TileBMax, const float RasterizationPadding)
+	{
+		Reset();
+
+		rcVcopy(RasterizeHF->bmin, TileBMin);
+		rcVcopy(RasterizeHF->bmax, TileBMax);
+
+		RasterizeHF->bmin[0] -= RasterizationPadding;
+		RasterizeHF->bmin[2] -= RasterizationPadding;
+		RasterizeHF->bmax[0] += RasterizationPadding;
+		RasterizeHF->bmax[2] += RasterizationPadding;
+	}
+
+	rcHeightfield* RasterizeHF;
+};
+
+static FVoxelCacheRasterizeContext VoxelCacheContext;
+
+uint32 GetTileCacheSizeHelper(TArray<FNavMeshTileData>& CompressedTiles)
+{
+	uint32 TotalMemory = 0;
+	for (int32 i = 0; i < CompressedTiles.Num(); i++)
+	{
+		TotalMemory += CompressedTiles[i].DataSize;
+	}
+
+	return TotalMemory;
+}
+
+static FBox CalculateTileBounds(int32 X, int32 Y, const FVector& NavMeshOrigin, const FBox& TotalNavBounds, float TileSizeInWorldUnits)
+{
+	const FVector RcNavMeshOrigin = Unreal2RecastPoint(NavMeshOrigin);
+	FBox TileBox(
+		RcNavMeshOrigin + (FVector(X + 0, 0, Y + 0) * TileSizeInWorldUnits),
+		RcNavMeshOrigin + (FVector(X + 1, 0, Y + 1) * TileSizeInWorldUnits)
+		);
+
+	TileBox = Recast2UnrealBox(TileBox);
+	TileBox.Min.Z = TotalNavBounds.Min.Z;
+	TileBox.Max.Z = TotalNavBounds.Max.Z;
+
+	// unreal coord space
+	return TileBox;
+}
+
+//----------------------------------------------------------------------//
+// FRecastTileGenerator
+//----------------------------------------------------------------------//
+
+FRecastTileGenerator::FRecastTileGenerator(FRecastNavMeshGenerator& ParentGenerator, const FIntPoint& Location)
+{
+	bSucceeded = false;
+	bUpdateGeometry = true;
+	bHasLowAreaModifiers = false;
+
+	TileX = Location.X;
+	TileY = Location.Y;
+
+	TileConfig = ParentGenerator.GetConfig();
+	Version = ParentGenerator.GetVersion();
+	AdditionalCachedData = ParentGenerator.GetAdditionalCachedData();
+
+	ParentGeneratorWeakPtr = ParentGenerator.AsShared();
+}
+
+FRecastTileGenerator::~FRecastTileGenerator()
+{
+}
+
+void FRecastTileGenerator::Setup(const FRecastNavMeshGenerator& ParentGenerator, const TArray<FBox>& DirtyAreas)
+{
+	const FVector NavMeshOrigin = FVector::ZeroVector;
+	const FBox NavTotalBounds = ParentGenerator.GetTotalBounds();
+	const float TileCellSize = (TileConfig.tileSize * TileConfig.cs);
+
+	TileBB = CalculateTileBounds(TileX, TileY, NavMeshOrigin, NavTotalBounds, TileCellSize);
+	TileBBExpandedForAgent = TileBB.ExpandBy(NavDataConfig.AgentRadius * 2 + TileConfig.cs);
+	const FBox RCBox = Unreal2RecastBox(TileBB);
+	rcVcopy(TileConfig.bmin, &RCBox.Min.X);
+	rcVcopy(TileConfig.bmax, &RCBox.Max.X);
+			
+	// from passed in boxes pick the ones overlapping with tile bounds
+	bFullyEncapsulatedByInclusionBounds = true;
+	const TNavStatArray<FBox>& ParentBounds = ParentGenerator.GetInclusionBounds();
+	if (ParentBounds.Num() > 0)
+	{
+		bFullyEncapsulatedByInclusionBounds = false;
+		InclusionBounds.Reserve(ParentBounds.Num());
+		for (const FBox& Bounds : ParentBounds)
+		{
+			if (Bounds.Intersect(TileBB))
+			{
+				InclusionBounds.Add(Bounds);
+				bFullyEncapsulatedByInclusionBounds = DoesBoxContainBox(Bounds, TileBB);
+			}
+		}
+	}
+
+	const bool bGeometryChanged = (DirtyAreas.Num() == 0);
+	if (!bGeometryChanged)
+	{
+		// Get compressed tile cache layers if they exist for this location
+		CompressedLayers = ParentGenerator.GetOwner()->GetTileCacheLayers(TileX, TileY);
+		for (FNavMeshTileData& LayerData : CompressedLayers)
+		{
+			// we don't want to modify shared state inside async task, so make sure we are unique owner
+			LayerData.MakeUnique();
+		}
+	}
+
+	// We have to regenerate layers data in case geometry is changed or tile cache is missing
+	bRegenerateCompressedLayers = (bGeometryChanged || CompressedLayers.Num() == 0);
+	
+	// Gather geometry for tile if it inside navigable bounds
+	if (InclusionBounds.Num())
+	{
+		if (!bRegenerateCompressedLayers)
+		{
+			// Mark layers that needs to be updated
+			DirtyLayers.Init(false, CompressedLayers.Num());
+			for (const FNavMeshTileData& LayerData : CompressedLayers)
+			{
+				for (FBox DirtyBox : DirtyAreas)
+				{
+					if (DirtyBox.Intersect(LayerData.LayerBBox))
+					{
+						DirtyLayers[LayerData.LayerIndex] = true;
+					}
+				}
+			}
+		}
+		
+		if (ParentGenerator.GatherGeometryOnGameThread())
+		{
+			GatherGeometry(ParentGenerator, bRegenerateCompressedLayers);
+		}
+		else
+		{
+			PrepareGeometrySources(ParentGenerator, bRegenerateCompressedLayers);
+		}
+	}
+	
+	//
+	UsedMemoryOnStartup = GetUsedMemCount() + sizeof(FRecastTileGenerator);
+}
+
+bool FRecastTileGenerator::HasDataToBuild() const
+{
+	return
+		CompressedLayers.Num()
+		|| Modifiers.Num()
+		|| OffmeshLinks.Num()
+		|| RawGeometry.Num()
+		|| (InclusionBounds.Num() && NavigationRelevantData.Num() > 0);
+}
+
+void FRecastTileGenerator::DoWork()
+{
+	TSharedPtr<FNavDataGenerator, ESPMode::ThreadSafe> ParentGenerator = ParentGeneratorWeakPtr.Pin();
+
+	if (ParentGenerator.IsValid())
+	{
+		FRecastNavMeshGenerator* RecastParentGenerator = static_cast<FRecastNavMeshGenerator*>(ParentGenerator.Get());
+
+		if (InclusionBounds.Num())
+		{
+			DoAsyncGeometryGathering();
+		}
+
+		bSucceeded = GenerateTile();
+	}
+
+	DumpAsyncData();
+}
+
+void FRecastTileGenerator::DumpAsyncData()
+{
+	RawGeometry.Empty();
+	Modifiers.Empty();
+	OffmeshLinks.Empty();
+
+	NavigationRelevantData.Empty();
+	NavOctree = nullptr;
+}
+
+void FRecastTileGenerator::DoAsyncGeometryGathering()
+{
+	QUICK_SCOPE_CYCLE_COUNTER(STAT_RecastNavMeshGenerator_PrepareGeometrySources);
+	
+	for (auto& ElementData : NavigationRelevantData)
+	{
+		if (ElementData->GetOwner() == nullptr)
+		{
+			UE_LOG(LogNavigation, Warning, TEXT("DoAsyncGeometryGathering: skipping an element with no longer valid Owner"));
+			continue;
+		}
+
+		bool bDumpGeometryData = false;
+		if (ElementData->IsPendingLazyGeometryGathering() && ElementData->SupportsGatheringGeometrySlices())
+		{
+			QUICK_SCOPE_CYCLE_COUNTER(STAT_RecastNavMeshGenerator_LandscapeSlicesExporting);
+
+			FRecastGeometryExport GeomExport(*ElementData);
+
+			INavRelevantInterface* NavRelevant = Cast<INavRelevantInterface>(ElementData->GetOwner());
+			if(NavRelevant)
+			{
+				NavRelevant->PrepareGeometryExportSync();
+				// adding a small bump to avoid special case of zero-expansion when tile bounds
+				// overlap landscape's tile bounds
+				NavRelevant->GatherGeometrySlice(GeomExport, TileBBExpandedForAgent);
+
+				RecastGeometryExport::CovertCoordDataToRecast(GeomExport.VertexBuffer);
+				RecastGeometryExport::StoreCollisionCache(GeomExport);
+				bDumpGeometryData = true;
+			}
+			else
+			{
+				UE_LOG(LogNavigation, Error, TEXT("DoAsyncGeometryGathering: got an invalid NavRelevant instance!"));
+			}
+		}
+
+		if (ElementData->IsPendingLazyGeometryGathering() || ElementData->IsPendingLazyModifiersGathering())
+		{
+			NavOctree->DemandLazyDataGathering(*ElementData);
+		}
+
+		const bool bExportGeometry = bUpdateGeometry && ElementData->HasGeometry();
+		if (bExportGeometry)
+		{
+			if (ARecastNavMesh::IsVoxelCacheEnabled())
+			{
+				TNavStatArray<rcSpanCache> SpanData;
+				rcSpanCache* CachedVoxels = 0;
+				int32 NumCachedVoxels = 0;
+
+				DECLARE_SCOPE_CYCLE_COUNTER(TEXT("Rasterization: prepare voxel cache"), Stat_RecastRasterCachePrep, STATGROUP_Navigation);
+
+				if (!HasVoxelCache(ElementData->VoxelData, CachedVoxels, NumCachedVoxels))
+				{
+					// rasterize
+					PrepareVoxelCache(ElementData->CollisionData, SpanData);
+					CachedVoxels = SpanData.GetData();
+					NumCachedVoxels = SpanData.Num();
+
+					// encode
+					const int32 PrevElementMemory = ElementData->GetAllocatedSize();
+					FNavigationRelevantData* ModData = (FNavigationRelevantData*)&ElementData;
+					AddVoxelCache(ModData->VoxelData, CachedVoxels, NumCachedVoxels);
+
+					const int32 NewElementMemory = ElementData->GetAllocatedSize();
+					const int32 ElementMemoryDelta = NewElementMemory - PrevElementMemory;
+					INC_MEMORY_STAT_BY(STAT_Navigation_CollisionTreeMemory, ElementMemoryDelta);
+				}
+			}
+			else
+			{
+				ValidateAndAppendGeometry(ElementData);
+			}
+
+			if (bDumpGeometryData)
+			{
+				const_cast<FNavigationRelevantData&>(*ElementData).CollisionData.Empty();
+			}
+		}
+
+		const FCompositeNavModifier ModifierInstance = ElementData->Modifiers.HasMetaAreas() ? ElementData->Modifiers.GetInstantiatedMetaModifier(&NavDataConfig, ElementData->SourceObject) : ElementData->Modifiers;
+		if (ModifierInstance.IsEmpty() == false)
+		{
+			AppendModifier(ModifierInstance, ElementData->NavDataPerInstanceTransformDelegate);
+		}
+	}
+}
+
+void FRecastTileGenerator::PrepareGeometrySources(const FRecastNavMeshGenerator& ParentGenerator, bool bGeometryChanged)
+{
+	QUICK_SCOPE_CYCLE_COUNTER(STAT_RecastNavMeshGenerator_PrepareGeometrySources);
+
+	UNavigationSystemV1* NavSys = FNavigationSystem::GetCurrent<UNavigationSystemV1>(ParentGenerator.GetWorld());
+	FNavigationOctree* NavOctreeInstance = NavSys ? NavSys->GetMutableNavOctree() : nullptr;
+	check(NavOctreeInstance);
+	NavDataConfig = ParentGenerator.GetOwner()->GetConfig();
+	NavigationRelevantData.Reset();
+	NavOctree = NavOctreeInstance->AsShared();
+	bUpdateGeometry = bGeometryChanged;
+
+	for (FNavigationOctree::TConstElementBoxIterator<FNavigationOctree::DefaultStackAllocator> It(*NavOctreeInstance, ParentGenerator.GrowBoundingBox(TileBB, /*bIncludeAgentHeight*/ false));
+		It.HasPendingElements();
+		It.Advance())
+	{
+		const FNavigationOctreeElement& Element = It.GetCurrentElement();
+		const bool bShouldUse = Element.ShouldUseGeometry(NavDataConfig);
+		if (bShouldUse)
+		{
+			const bool bExportGeometry = bGeometryChanged && (Element.Data->HasGeometry() || Element.Data->IsPendingLazyGeometryGathering());
+			if (bExportGeometry
+				|| (Element.Data->IsPendingLazyModifiersGathering() || Element.Data->Modifiers.HasMetaAreas() == true || Element.Data->Modifiers.IsEmpty() == false))
+			{
+				NavigationRelevantData.Add(Element.Data);
+			}
+		}
+	}
+}
+
+void FRecastTileGenerator::GatherGeometry(const FRecastNavMeshGenerator& ParentGenerator, bool bGeometryChanged)
+{
+	QUICK_SCOPE_CYCLE_COUNTER(STAT_RecastNavMeshGenerator_GatherGeometry);
+
+	UNavigationSystemV1* NavSys = FNavigationSystem::GetCurrent<UNavigationSystemV1>(ParentGenerator.GetWorld());
+	FNavigationOctree* NavigationOctree = NavSys ? NavSys->GetMutableNavOctree() : nullptr;
+	if (NavigationOctree == nullptr)
+	{
+		return;
+	}
+	const FNavDataConfig& OwnerNavDataConfig = ParentGenerator.GetOwner()->GetConfig();
+
+	for (FNavigationOctree::TConstElementBoxIterator<FNavigationOctree::DefaultStackAllocator> It(*NavigationOctree, ParentGenerator.GrowBoundingBox(TileBB, /*bIncludeAgentHeight*/ false));
+		It.HasPendingElements();
+		It.Advance())
+	{
+		const FNavigationOctreeElement& Element = It.GetCurrentElement();
+		const bool bShouldUse = Element.ShouldUseGeometry(OwnerNavDataConfig);
+		if (bShouldUse)
+		{
+			bool bDumpGeometryData = false;
+			if (Element.Data->IsPendingLazyGeometryGathering() || Element.Data->IsPendingLazyModifiersGathering())
+			{
+				const bool bSupportsSlices = Element.Data->SupportsGatheringGeometrySlices();
+				if (bSupportsSlices == false || Element.Data->IsPendingLazyModifiersGathering() == true)
+				{
+					QUICK_SCOPE_CYCLE_COUNTER(STAT_RecastNavMeshGenerator_LazyGeometryExport);
+					NavigationOctree->DemandLazyDataGathering(Element);
+				}
+				
+				if (bSupportsSlices == true)
+				{
+					QUICK_SCOPE_CYCLE_COUNTER(STAT_RecastNavMeshGenerator_LandscapeSlicesExporting);
+
+					FRecastGeometryExport GeomExport(const_cast<FNavigationRelevantData&>(*Element.Data));
+
+					INavRelevantInterface* NavRelevant = const_cast<INavRelevantInterface*>(Cast<const INavRelevantInterface>(Element.GetOwner()));
+					if (NavRelevant)
+					{
+						NavRelevant->PrepareGeometryExportSync();
+						// adding a small bump to avoid special case of zero-expansion when tile bounds
+						// overlap landscape's tile bounds
+						NavRelevant->GatherGeometrySlice(GeomExport, TileBBExpandedForAgent);
+
+						RecastGeometryExport::CovertCoordDataToRecast(GeomExport.VertexBuffer);
+						RecastGeometryExport::StoreCollisionCache(GeomExport);
+						bDumpGeometryData = true;
+					}
+					else
+					{
+						UE_LOG(LogNavigation, Error, TEXT("GatherGeometry: got an invalid NavRelevant instance!"));
+					}
+				}
+			}
+
+			const bool bExportGeometry = bGeometryChanged && Element.Data->HasGeometry();
+			if (bExportGeometry)
+			{
+				if (ARecastNavMesh::IsVoxelCacheEnabled())
+				{
+					TNavStatArray<rcSpanCache> SpanData;
+					rcSpanCache* CachedVoxels = 0;
+					int32 NumCachedVoxels = 0;
+
+					DECLARE_SCOPE_CYCLE_COUNTER(TEXT("Rasterization: prepare voxel cache"), Stat_RecastRasterCachePrep, STATGROUP_Navigation);
+
+					if (!HasVoxelCache(Element.Data->VoxelData, CachedVoxels, NumCachedVoxels))
+					{
+						// rasterize
+						PrepareVoxelCache(Element.Data->CollisionData, SpanData);
+						CachedVoxels = SpanData.GetData();
+						NumCachedVoxels = SpanData.Num();
+
+						// encode
+						const int32 PrevElementMemory = Element.Data->GetAllocatedSize();
+						FNavigationRelevantData* ModData = (FNavigationRelevantData*)&Element.Data;
+						AddVoxelCache(ModData->VoxelData, CachedVoxels, NumCachedVoxels);
+
+						const int32 NewElementMemory = Element.Data->GetAllocatedSize();
+						const int32 ElementMemoryDelta = NewElementMemory - PrevElementMemory;
+						INC_MEMORY_STAT_BY(STAT_Navigation_CollisionTreeMemory, ElementMemoryDelta);
+					}
+				}
+				else
+				{
+					ValidateAndAppendGeometry(Element.Data);
+				}
+
+				if (bDumpGeometryData)
+				{
+					const_cast<FNavigationRelevantData&>(*Element.Data).CollisionData.Empty();
+				}
+			}
+						
+			const FCompositeNavModifier ModifierInstance = Element.GetModifierForAgent(&OwnerNavDataConfig);
+			if (ModifierInstance.IsEmpty() == false)
+			{
+				AppendModifier(ModifierInstance, Element.Data->NavDataPerInstanceTransformDelegate);
+			}
+		}
+	}
+}
+
+void FRecastTileGenerator::ApplyVoxelFilter(rcHeightfield* HF, float WalkableRadius)
+{
+	SCOPE_CYCLE_COUNTER(STAT_Navigation_TileVoxelFilteringAsync);
+
+	if (HF != NULL)
+	{
+		const int32 Width = HF->width;
+		const int32 Height = HF->height;
+		const float CellSize = HF->cs;
+		const float CellHeight = HF->ch;
+		const float BottomX = HF->bmin[0];
+		const float BottomZ = HF->bmin[1];
+		const float BottomY = HF->bmin[2];
+		const int32 SpansCount = Width*Height;
+		// we need to expand considered bounding boxes so that
+		// it doesn't create "fake cliffs"
+		const float ExpandBBBy = WalkableRadius*CellSize;
+
+		const FBox* BBox = InclusionBounds.GetData();
+		// optimized common case of single box
+		if (InclusionBounds.Num() == 1)
+		{
+			const FBox BB = BBox->ExpandBy(ExpandBBBy);
+
+			rcSpan** Span = HF->spans;
+
+			for (int32 y = 0; y < Height; ++y)
+			{
+				for (int32 x = 0; x < Width; ++x)
+				{
+					const float SpanX = -(BottomX + x * CellSize);
+					const float SpanY = -(BottomY + y * CellSize);
+
+					// mark all spans outside of InclusionBounds as unwalkable
+					for (rcSpan* s = *Span; s; s = s->next)
+					{
+						if (s->data.area == RC_WALKABLE_AREA)
+						{
+							const float SpanMin = CellHeight * s->data.smin + BottomZ;
+							const float SpanMax = CellHeight * s->data.smax + BottomZ;
+
+							const FVector SpanMinV(SpanX-CellSize, SpanY-CellSize, SpanMin);
+							const FVector SpanMaxV(SpanX, SpanY, SpanMax);
+
+							if (BB.IsInside(SpanMinV) == false && BB.IsInside(SpanMaxV) == false)
+							{
+								s->data.area = RC_NULL_AREA;
+							}
+						}
+					}
+					++Span;
+				}
+			}
+		}
+		else
+		{
+			TArray<FBox> Bounds;
+			Bounds.Reserve(InclusionBounds.Num());
+
+			for (int32 i = 0; i < InclusionBounds.Num(); ++i, ++BBox)
+			{	
+				Bounds.Add(BBox->ExpandBy(ExpandBBBy));
+			}
+			const int32 BoundsCount = Bounds.Num();
+
+			rcSpan** Span = HF->spans;
+
+			for (int32 y = 0; y < Height; ++y)
+			{
+				for (int32 x = 0; x < Width; ++x)
+				{
+					const float SpanX = -(BottomX + x * CellSize);
+					const float SpanY = -(BottomY + y * CellSize);
+
+					// mark all spans outside of InclusionBounds as unwalkable
+					for (rcSpan* s = *Span; s; s = s->next)
+					{
+						if (s->data.area == RC_WALKABLE_AREA)
+						{
+							const float SpanMin = CellHeight * s->data.smin + BottomZ;
+							const float SpanMax = CellHeight * s->data.smax + BottomZ;
+
+							const FVector SpanMinV(SpanX-CellSize, SpanY-CellSize, SpanMin);
+							const FVector SpanMaxV(SpanX, SpanY, SpanMax);
+
+							bool bIsInsideAnyBB = false;
+							const FBox* BB = Bounds.GetData();
+							for (int32 BoundIndex = 0; BoundIndex < BoundsCount; ++BoundIndex, ++BB)
+							{
+								if (BB->IsInside(SpanMinV) || BB->IsInside(SpanMaxV))
+								{
+									bIsInsideAnyBB = true;
+									break;
+								}
+							}
+
+							if (bIsInsideAnyBB == false)
+							{
+								s->data.area = RC_NULL_AREA;
+							}
+						}
+					}
+					++Span;
+				}
+			}
+		}
+	}
+}
+
+void FRecastTileGenerator::PrepareVoxelCache(const TNavStatArray<uint8>& RawCollisionCache, TNavStatArray<rcSpanCache>& SpanData)
+{
+	// tile's geometry: voxel cache (only for synchronous rebuilds)
+	const int32 WalkableClimbVX = TileConfig.walkableClimb;
+	const float WalkableSlopeCos = FMath::Cos(FMath::DegreesToRadians(TileConfig.walkableSlopeAngle));
+	const float RasterizationPadding = TileConfig.borderSize * TileConfig.cs;
+
+	FRecastGeometryCache CachedCollisions(RawCollisionCache.GetData());
+
+	VoxelCacheContext.SetupForTile(TileConfig.bmin, TileConfig.bmax, RasterizationPadding);
+
+	float SlopeCosPerActor = WalkableSlopeCos;
+	CachedCollisions.Header.SlopeOverride.ModifyWalkableFloorZ(SlopeCosPerActor);
+
+	// rasterize triangle soup
+	TNavStatArray<uint8> TriAreas;
+	TriAreas.AddZeroed(CachedCollisions.Header.NumFaces);
+
+	rcMarkWalkableTrianglesCos(0, SlopeCosPerActor,
+		CachedCollisions.Verts, CachedCollisions.Header.NumVerts,
+		CachedCollisions.Indices, CachedCollisions.Header.NumFaces,
+		TriAreas.GetData());
+
+	rcRasterizeTriangles(0, CachedCollisions.Verts, CachedCollisions.Header.NumVerts,
+		CachedCollisions.Indices, TriAreas.GetData(), CachedCollisions.Header.NumFaces,
+		*VoxelCacheContext.RasterizeHF, WalkableClimbVX);
+
+	const int32 NumSpans = rcCountSpans(0, *VoxelCacheContext.RasterizeHF);
+	if (NumSpans > 0)
+	{
+		SpanData.AddZeroed(NumSpans);
+		rcCacheSpans(0, *VoxelCacheContext.RasterizeHF, SpanData.GetData());
+	}
+}
+
+bool FRecastTileGenerator::HasVoxelCache(const TNavStatArray<uint8>& RawVoxelCache, rcSpanCache*& CachedVoxels, int32& NumCachedVoxels) const
+{
+	FRecastVoxelCache VoxelCache(RawVoxelCache.GetData());
+	for (FRecastVoxelCache::FTileInfo* iTile = VoxelCache.Tiles; iTile; iTile = iTile->NextTile)
+	{
+		if (iTile->TileX == TileX && iTile->TileY == TileY)
+		{
+			CachedVoxels = iTile->SpanData;
+			NumCachedVoxels = iTile->NumSpans;
+			return true;
+		}
+	}
+	
+	return false;
+}
+
+void FRecastTileGenerator::AddVoxelCache(TNavStatArray<uint8>& RawVoxelCache, const rcSpanCache* CachedVoxels, const int32 NumCachedVoxels) const
+{
+	if (RawVoxelCache.Num() == 0)
+	{
+		RawVoxelCache.AddZeroed(sizeof(int32));
+	}
+
+	int32* NumTiles = (int32*)RawVoxelCache.GetData();
+	*NumTiles = *NumTiles + 1;
+
+	const int32 NewCacheIdx = RawVoxelCache.Num();
+	const int32 HeaderSize = sizeof(FRecastVoxelCache::FTileInfo);
+	const int32 VoxelsSize = sizeof(rcSpanCache) * NumCachedVoxels;
+	const int32 EntrySize = HeaderSize + VoxelsSize;
+	RawVoxelCache.AddZeroed(EntrySize);
+
+	FRecastVoxelCache::FTileInfo* TileInfo = (FRecastVoxelCache::FTileInfo*)(RawVoxelCache.GetData() + NewCacheIdx);
+	TileInfo->TileX = TileX;
+	TileInfo->TileY = TileY;
+	TileInfo->NumSpans = NumCachedVoxels;
+
+	FMemory::Memcpy(RawVoxelCache.GetData() + NewCacheIdx + HeaderSize, CachedVoxels, VoxelsSize);
+}
+
+void FRecastTileGenerator::AppendModifier(const FCompositeNavModifier& Modifier, const FNavDataPerInstanceTransformDelegate& InTransformsDelegate)
+{
+	// append all offmesh links (not included in compress layers)
+	OffmeshLinks.Append(Modifier.GetSimpleLinks());
+
+	// evaluate custom links
+	const FCustomLinkNavModifier* LinkModifier = Modifier.GetCustomLinks().GetData();
+	for (int32 i = 0; i < Modifier.GetCustomLinks().Num(); i++, LinkModifier++)
+	{
+		FSimpleLinkNavModifier SimpleLinkCollection(UNavLinkDefinition::GetLinksDefinition(LinkModifier->GetNavLinkClass()), LinkModifier->LocalToWorld);
+		OffmeshLinks.Add(SimpleLinkCollection);
+	}
+
+	if (Modifier.GetAreas().Num() == 0)
+	{
+		return;
+	}
+
+	bHasLowAreaModifiers = bHasLowAreaModifiers || Modifier.HasLowAreaModifiers();
+	
+	FRecastAreaNavModifierElement ModifierElement;
+
+	// Gather per instance transforms if any
+	if (InTransformsDelegate.IsBound())
+	{
+		InTransformsDelegate.Execute(TileBBExpandedForAgent, ModifierElement.PerInstanceTransform);
+		// skip this modifier in case there is no instances for this tile
+		if (ModifierElement.PerInstanceTransform.Num() == 0)
+		{
+			return;
+		}
+	}
+		
+	ModifierElement.Areas = Modifier.GetAreas();
+	Modifiers.Add(MoveTemp(ModifierElement));
+}
+
+void FRecastTileGenerator::ValidateAndAppendGeometry(TSharedRef<FNavigationRelevantData, ESPMode::ThreadSafe> ElementData)
+{
+	const FNavigationRelevantData& DataRef = ElementData.Get();
+	if (DataRef.IsCollisionDataValid())
+	{
+		AppendGeometry(DataRef.CollisionData, DataRef.NavDataPerInstanceTransformDelegate);
+	}
+}
+
+void FRecastTileGenerator::AppendGeometry(const TNavStatArray<uint8>& RawCollisionCache, const FNavDataPerInstanceTransformDelegate& InTransformsDelegate)
+{	
+	if (RawCollisionCache.Num() == 0)
+	{
+		return;
+	}
+	
+	FRecastRawGeometryElement GeometryElement;
+	FRecastGeometryCache CollisionCache(RawCollisionCache.GetData());
+	
+	// Gather per instance transforms
+	if (InTransformsDelegate.IsBound())
+	{
+		InTransformsDelegate.Execute(TileBBExpandedForAgent, GeometryElement.PerInstanceTransform);
+		if (GeometryElement.PerInstanceTransform.Num() == 0)
+		{
+			return;
+		}
+	}
+	
+	const int32 NumCoords = CollisionCache.Header.NumVerts * 3;
+	const int32 NumIndices = CollisionCache.Header.NumFaces * 3;
+	if (NumIndices > 0)
+	{
+		GeometryElement.GeomCoords.SetNumUninitialized(NumCoords);
+		GeometryElement.GeomIndices.SetNumUninitialized(NumIndices);
+
+		FMemory::Memcpy(GeometryElement.GeomCoords.GetData(), CollisionCache.Verts, sizeof(float) * NumCoords);
+		FMemory::Memcpy(GeometryElement.GeomIndices.GetData(), CollisionCache.Indices, sizeof(int32) * NumIndices);
+
+		RawGeometry.Add(MoveTemp(GeometryElement));
+	}	
+}
+
+bool FRecastTileGenerator::GenerateTile()
+{
+	bool bSuccess = true;
+	FNavMeshBuildContext BuildContext;
+
+	if (bRegenerateCompressedLayers)
+	{
+		CompressedLayers.Reset();
+		
+		bSuccess = GenerateCompressedLayers(BuildContext);
+
+		if (bSuccess)
+		{
+			// Mark all layers as dirty
+			DirtyLayers.Init(true, CompressedLayers.Num());
+		}
+	}
+
+	if (bSuccess)
+	{
+		bSuccess = GenerateNavigationData(BuildContext);
+	}
+
+	// it's possible to have valid generation with empty resulting tile (no navigable geometry in tile)
+	return bSuccess;
+}
+
+struct FTileRasterizationContext
+{
+	FTileRasterizationContext() : SolidHF(0), LayerSet(0), CompactHF(0)
+	{
+	}
+
+	~FTileRasterizationContext()
+	{
+		rcFreeHeightField(SolidHF);
+		rcFreeHeightfieldLayerSet(LayerSet);
+		rcFreeCompactHeightfield(CompactHF);
+	}
+
+	struct rcHeightfield* SolidHF;
+	struct rcHeightfieldLayerSet* LayerSet;
+	struct rcCompactHeightfield* CompactHF;
+	TArray<FNavMeshTileData> Layers;
+};
+
+static void RasterizeGeometry(
+	FNavMeshBuildContext& BuildContext, const FRecastBuildConfig& TileConfig, 
+	const TArray<float>& Coords, const TArray<int32>& Indices, FTileRasterizationContext& RasterContext)
+{
+	const int32 NumFaces = Indices.Num() / 3;
+	const int32 NumVerts = Coords.Num() / 3;
+
+	TNavStatArray<uint8> TriAreas;
+	TriAreas.Reserve(NumFaces);
+	TriAreas.AddZeroed(NumFaces);
+
+	rcMarkWalkableTriangles(&BuildContext, TileConfig.walkableSlopeAngle,
+		Coords.GetData(), NumVerts, Indices.GetData(), NumFaces,
+		TriAreas.GetData());
+
+	rcRasterizeTriangles(&BuildContext,
+		Coords.GetData(), NumVerts, 
+		Indices.GetData(), TriAreas.GetData(), NumFaces,
+		*RasterContext.SolidHF, TileConfig.walkableClimb);
+}
+
+static void RasterizeGeometry(
+	FNavMeshBuildContext& BuildContext,const FRecastBuildConfig& TileConfig, 
+	const TArray<float>& Coords, const TArray<int32>& Indices, const FTransform& LocalToWorld, FTileRasterizationContext& RasterContext)
+{
+	TArray<float> WorldRecastCoords;
+	WorldRecastCoords.SetNumUninitialized(Coords.Num());
+	
+	FMatrix LocalToRecastWorld = LocalToWorld.ToMatrixWithScale()*Unreal2RecastMatrix();
+	// Convert geometry to recast world space
+	for (int32 i = 0; i < Coords.Num(); i+=3)
+	{
+		// collision cache stores coordinates in recast space, convert them to unreal and transform to recast world space
+		FVector WorldRecastCoord = LocalToRecastWorld.TransformPosition(Recast2UnrealPoint(&Coords[i])); 
+
+		WorldRecastCoords[i+0] = WorldRecastCoord.X;
+		WorldRecastCoords[i+1] = WorldRecastCoord.Y;
+		WorldRecastCoords[i+2] = WorldRecastCoord.Z;
+	}
+
+	RasterizeGeometry(BuildContext, TileConfig, WorldRecastCoords, Indices, RasterContext);
+}
+
+bool FRecastTileGenerator::GenerateCompressedLayers(FNavMeshBuildContext& BuildContext)
+{
+	RECAST_STAT(STAT_Navigation_Async_Recast_Rasterize);
+	SCOPE_CYCLE_COUNTER(STAT_Navigation_RecastBuildCompressedLayers);
+	
+	TileConfig.width = TileConfig.tileSize + TileConfig.borderSize*2;
+	TileConfig.height = TileConfig.tileSize + TileConfig.borderSize*2;
+
+	const float BBoxPadding = TileConfig.borderSize * TileConfig.cs;
+	TileConfig.bmin[0] -= BBoxPadding;
+	TileConfig.bmin[2] -= BBoxPadding;
+	TileConfig.bmax[0] += BBoxPadding;
+	TileConfig.bmax[2] += BBoxPadding;
+
+	BuildContext.log(RC_LOG_PROGRESS, "GenerateCompressedLayers:");
+	BuildContext.log(RC_LOG_PROGRESS, " - %d x %d cells", TileConfig.width, TileConfig.height);
+
+	FTileRasterizationContext RasterContext;
+	const bool bHasGeometry = RawGeometry.Num() > 0;
+
+	// Allocate voxel heightfield where we rasterize our input data to.
+	if (bHasGeometry)
+	{
+		RasterContext.SolidHF = rcAllocHeightfield();
+		if (RasterContext.SolidHF == NULL)
+		{
+			BuildContext.log(RC_LOG_ERROR, "GenerateCompressedLayers: Out of memory 'SolidHF'.");
+			return false;
+		}
+		if (!rcCreateHeightfield(&BuildContext, *RasterContext.SolidHF, TileConfig.width, TileConfig.height, TileConfig.bmin, TileConfig.bmax, TileConfig.cs, TileConfig.ch))
+		{
+			BuildContext.log(RC_LOG_ERROR, "GenerateCompressedLayers: Could not create solid heightfield.");
+			return false;
+		}
+
+		// Rasterize geometry
+		DECLARE_SCOPE_CYCLE_COUNTER(TEXT("Rasterization: without voxel cache"), Stat_RecastRasterNoCache, STATGROUP_Navigation);
+		RECAST_STAT(STAT_Navigation_RasterizeTriangles)
+		
+		for (const FRecastRawGeometryElement& Element : RawGeometry)
+		{
+			for (const FTransform& InstanceTransform : Element.PerInstanceTransform)
+			{
+				RasterizeGeometry(BuildContext, TileConfig, Element.GeomCoords, Element.GeomIndices, InstanceTransform, RasterContext);
+			}
+			
+			if (Element.PerInstanceTransform.Num() == 0)
+			{
+				RasterizeGeometry(BuildContext, TileConfig, Element.GeomCoords, Element.GeomIndices, RasterContext);
+			}
+		}
+	}
+	
+	if (!RasterContext.SolidHF || RasterContext.SolidHF->pools == 0)
+	{
+		BuildContext.log(RC_LOG_WARNING, "GenerateCompressedLayers: empty tile - aborting");
+		return true;
+	}
+
+	// Reject voxels outside generation boundaries
+	if (TileConfig.bPerformVoxelFiltering && !bFullyEncapsulatedByInclusionBounds)
+	{
+		ApplyVoxelFilter(RasterContext.SolidHF, TileConfig.walkableRadius);
+	}
+
+	// TileConfig.walkableHeight is set to 1 when marking low spans, calculate real value for filtering
+	const int32 FilterWalkableHeight = FMath::CeilToInt(TileConfig.AgentHeight / TileConfig.ch);
+
+	{
+		RECAST_STAT(STAT_Navigation_Async_Recast_Filter);
+		// Once all geometry is rasterized, we do initial pass of filtering to
+		// remove unwanted overhangs caused by the conservative rasterization
+		// as well as filter spans where the character cannot possibly stand.
+		rcFilterLowHangingWalkableObstacles(&BuildContext, TileConfig.walkableClimb, *RasterContext.SolidHF);
+		rcFilterLedgeSpans(&BuildContext, TileConfig.walkableHeight, TileConfig.walkableClimb, *RasterContext.SolidHF);
+		if (!TileConfig.bMarkLowHeightAreas)
+		{
+			rcFilterWalkableLowHeightSpans(&BuildContext, TileConfig.walkableHeight, *RasterContext.SolidHF);
+		}
+		else if (TileConfig.bFilterLowSpanFromTileCache)
+		{
+			// TODO: investigate if creating detailed 2D map from active modifiers is cheap enough
+			// for now, switch on presence of those modifiers, will save memory as long as they are sparse (should be)
+
+			if (TileConfig.bFilterLowSpanSequences && bHasLowAreaModifiers)
+			{
+				rcFilterWalkableLowHeightSpansSequences(&BuildContext, FilterWalkableHeight, *RasterContext.SolidHF);
+			}
+			else
+			{
+				rcFilterWalkableLowHeightSpans(&BuildContext, FilterWalkableHeight, *RasterContext.SolidHF);
+			}
+		}
+	}
+
+	{
+		RECAST_STAT(STAT_Navigation_Async_Recast_BuildCompact);
+		// Compact the heightfield so that it is faster to handle from now on.
+		// This will result more cache coherent data as well as the neighbors
+		// between walkable cells will be calculated.
+		RasterContext.CompactHF = rcAllocCompactHeightfield();
+		if (RasterContext.CompactHF == NULL)
+		{
+			BuildContext.log(RC_LOG_ERROR, "GenerateCompressedLayers: Out of memory 'CompactHF'.");
+			return false;
+		}
+		if (!rcBuildCompactHeightfield(&BuildContext, TileConfig.walkableHeight, TileConfig.walkableClimb, *RasterContext.SolidHF, *RasterContext.CompactHF))
+		{
+			BuildContext.log(RC_LOG_ERROR, "GenerateCompressedLayers: Could not build compact data.");
+			return false;
+		}
+	}
+
+	{
+		RECAST_STAT(STAT_Navigation_Async_Recast_Erode);
+
+		if (TileConfig.walkableRadius > RECAST_VERY_SMALL_AGENT_RADIUS)
+		{
+			uint8 FilterFlags = 0;
+			if (TileConfig.bFilterLowSpanSequences)
+			{
+				FilterFlags = RC_LOW_FILTER_POST_PROCESS | (TileConfig.bFilterLowSpanFromTileCache ? 0 : RC_LOW_FILTER_SEED_SPANS);
+			}
+
+			const bool bEroded = TileConfig.bMarkLowHeightAreas ?
+				rcErodeWalkableAndLowAreas(&BuildContext, TileConfig.walkableRadius, FilterWalkableHeight, RECAST_LOW_AREA, FilterFlags, *RasterContext.CompactHF) :
+				rcErodeWalkableArea(&BuildContext, TileConfig.walkableRadius, *RasterContext.CompactHF);
+
+			if (!bEroded)
+			{
+				BuildContext.log(RC_LOG_ERROR, "GenerateCompressedLayers: Could not erode.");
+				return false;
+			}
+		}
+		else if (TileConfig.bMarkLowHeightAreas)
+		{
+			rcMarkLowAreas(&BuildContext, FilterWalkableHeight, RECAST_LOW_AREA, *RasterContext.CompactHF);
+		}
+	}
+
+	// Build layers
+	{
+		RECAST_STAT(STAT_Navigation_Async_Recast_Layers);
+		
+		RasterContext.LayerSet = rcAllocHeightfieldLayerSet();
+		if (RasterContext.LayerSet == NULL)
+		{
+			BuildContext.log(RC_LOG_ERROR, "GenerateCompressedLayers: Out of memory 'LayerSet'.");
+			return false;
+		}
+
+		if (TileConfig.regionPartitioning == RC_REGION_MONOTONE)
+		{
+			if (!rcBuildHeightfieldLayersMonotone(&BuildContext, *RasterContext.CompactHF, TileConfig.borderSize, TileConfig.walkableHeight, *RasterContext.LayerSet))
+			{
+				BuildContext.log(RC_LOG_ERROR, "GenerateCompressedLayers: Could not build heightfield layers.");
+				return 0;
+			}
+		}
+		else if (TileConfig.regionPartitioning == RC_REGION_WATERSHED)
+		{
+			if (!rcBuildDistanceField(&BuildContext, *RasterContext.CompactHF))
+			{
+				BuildContext.log(RC_LOG_ERROR, "GenerateCompressedLayers: Could not build distance field.");
+				return 0;
+			}
+
+			if (!rcBuildHeightfieldLayers(&BuildContext, *RasterContext.CompactHF, TileConfig.borderSize, TileConfig.walkableHeight, *RasterContext.LayerSet))
+			{
+				BuildContext.log(RC_LOG_ERROR, "GenerateCompressedLayers: Could not build heightfield layers.");
+				return 0;
+			}
+		}
+		else
+		{
+			if (!rcBuildHeightfieldLayersChunky(&BuildContext, *RasterContext.CompactHF, TileConfig.borderSize, TileConfig.walkableHeight, TileConfig.regionChunkSize, *RasterContext.LayerSet))
+			{
+				BuildContext.log(RC_LOG_ERROR, "GenerateCompressedLayers: Could not build heightfield layers.");
+				return 0;
+			}
+		}
+
+		const int32 NumLayers = RasterContext.LayerSet->nlayers;
+	
+		// use this to expand vertically layer's bounds
+		// this is needed to allow off-mesh connections that are not quite
+		// touching tile layer still connect with it.
+		const float StepHeights = TileConfig.AgentMaxClimb;
+
+		FTileCacheCompressor TileCompressor;
+		for (int32 i = 0; i < NumLayers; i++)
+		{
+			const rcHeightfieldLayer* layer = &RasterContext.LayerSet->layers[i];
+
+			// Store header
+			dtTileCacheLayerHeader header;
+			header.magic = DT_TILECACHE_MAGIC;
+			header.version = DT_TILECACHE_VERSION;
+
+			// Tile layer location in the navmesh.
+			header.tx = TileX;
+			header.ty = TileY;
+			header.tlayer = i;
+			dtVcopy(header.bmin, layer->bmin);
+			dtVcopy(header.bmax, layer->bmax);
+
+			// Tile info.
+			header.width = (unsigned short)layer->width;
+			header.height = (unsigned short)layer->height;
+			header.minx = (unsigned short)layer->minx;
+			header.maxx = (unsigned short)layer->maxx;
+			header.miny = (unsigned short)layer->miny;
+			header.maxy = (unsigned short)layer->maxy;
+			header.hmin = (unsigned short)layer->hmin;
+			header.hmax = (unsigned short)layer->hmax;
+
+			// Layer bounds in unreal coords
+			FBox LayerBBox = Recast2UnrealBox(header.bmin, header.bmax);
+			LayerBBox.Min.Z -= StepHeights;
+			LayerBBox.Max.Z += StepHeights;
+			
+			// Compress tile layer
+			uint8* TileData = NULL;
+			int32 TileDataSize = 0;
+			const dtStatus status = dtBuildTileCacheLayer(&TileCompressor, &header, layer->heights, layer->areas, layer->cons, &TileData, &TileDataSize);
+			if (dtStatusFailed(status))
+			{
+				dtFree(TileData);
+				BuildContext.log(RC_LOG_ERROR, "GenerateCompressedLayers: failed to build layer.");
+				return false;
+			}
+
+			// copy compressed data to new buffer in rasterization context
+			// (TileData allocates a lots of space, but only first TileDataSize bytes hold compressed data)
+
+			uint8* CompressedData = (uint8*)dtAlloc(TileDataSize * sizeof(uint8), DT_ALLOC_PERM);
+			if (CompressedData == NULL)
+			{
+				dtFree(TileData);
+				BuildContext.log(RC_LOG_ERROR, "GenerateCompressedLayers: Out of memory 'CompressedData'.");
+				return false;
+			}
+
+			FMemory::Memcpy(CompressedData, TileData, TileDataSize);
+			RasterContext.Layers.Add(FNavMeshTileData(CompressedData, TileDataSize, i, LayerBBox));
+
+			dtFree(TileData);
+
+			const int32 UncompressedSize = ((sizeof(dtTileCacheLayerHeader)+3) & ~3) + (3 * header.width * header.height);
+			const float Inv1kB = 1.0f / 1024.0f;
+			BuildContext.log(RC_LOG_PROGRESS, ">> Cache[%d,%d:%d] = %.2fkB (full:%.2fkB rate:%.2f%%)", TileX, TileY, i,
+				TileDataSize * Inv1kB, UncompressedSize * Inv1kB, 1.0f * TileDataSize / UncompressedSize);
+		}
+	}
+
+	// Transfer final data
+	CompressedLayers = RasterContext.Layers;
+	return true;
+}
+
+struct FTileGenerationContext
+{
+	FTileGenerationContext(dtTileCacheAlloc* MyAllocator) :
+		Allocator(MyAllocator), Layer(0), DistanceField(0), ContourSet(0), ClusterSet(0), PolyMesh(0), DetailMesh(0)
+	{
+	}
+
+	~FTileGenerationContext()
+	{
+		ResetIntermediateData();
+	}
+
+	void ResetIntermediateData()
+	{
+		dtFreeTileCacheLayer(Allocator, Layer);
+		Layer = 0;
+		dtFreeTileCacheDistanceField(Allocator, DistanceField);
+		DistanceField = 0;
+		dtFreeTileCacheContourSet(Allocator, ContourSet);
+		ContourSet = 0;
+		dtFreeTileCacheClusterSet(Allocator, ClusterSet);
+		ClusterSet = 0;
+		dtFreeTileCachePolyMesh(Allocator, PolyMesh);
+		PolyMesh = 0;
+		dtFreeTileCachePolyMeshDetail(Allocator, DetailMesh);
+		DetailMesh = 0;
+		// don't clear NavigationData here!
+	}
+
+	struct dtTileCacheAlloc* Allocator;
+	struct dtTileCacheLayer* Layer;
+	struct dtTileCacheDistanceField* DistanceField;
+	struct dtTileCacheContourSet* ContourSet;
+	struct dtTileCacheClusterSet* ClusterSet;
+	struct dtTileCachePolyMesh* PolyMesh;
+	struct dtTileCachePolyMeshDetail* DetailMesh;
+	TArray<FNavMeshTileData> NavigationData;
+};
+
+bool FRecastTileGenerator::GenerateNavigationData(FNavMeshBuildContext& BuildContext)
+{
+	RECAST_STAT(STAT_Navigation_Async_Recast_Generate);
+	SCOPE_CYCLE_COUNTER(STAT_Navigation_RecastBuildNavigation);
+		
+	FTileCacheAllocator MyAllocator;
+	FTileCacheCompressor TileCompressor;
+	
+	FTileGenerationContext GenerationContext(&MyAllocator);
+	GenerationContext.NavigationData.Reserve(CompressedLayers.Num());
+
+	dtStatus status = DT_SUCCESS;
+
+	for (int32 iLayer = 0; iLayer < CompressedLayers.Num(); iLayer++)
+	{
+		if (DirtyLayers[iLayer] == false || !CompressedLayers[iLayer].IsValid())
+		{
+			// skip layers not marked for rebuild
+			continue;
+		}
+				
+		FNavMeshTileData& CompressedData = CompressedLayers[iLayer];
+		const dtTileCacheLayerHeader* TileHeader = (const dtTileCacheLayerHeader*)CompressedData.GetData();
+		GenerationContext.ResetIntermediateData();
+
+		// Decompress tile layer data. 
+		status = dtDecompressTileCacheLayer(&MyAllocator, &TileCompressor, (unsigned char*)CompressedData.GetData(), CompressedData.DataSize, &GenerationContext.Layer);
+		if (dtStatusFailed(status))
+		{
+			BuildContext.log(RC_LOG_ERROR, "GenerateNavigationData: failed to decompress layer.");
+			return false;
+		}
+
+		// Rasterize obstacles.
+		MarkDynamicAreas(*GenerationContext.Layer);
+
+		{
+			RECAST_STAT(STAT_Navigation_Async_Recast_BuildRegions)
+			// Build regions
+			if (TileConfig.TileCachePartitionType == RC_REGION_MONOTONE)
+			{
+				status = dtBuildTileCacheRegionsMonotone(&MyAllocator, TileConfig.minRegionArea, TileConfig.mergeRegionArea, *GenerationContext.Layer);
+			}
+			else if (TileConfig.TileCachePartitionType == RC_REGION_WATERSHED)
+			{
+				GenerationContext.DistanceField = dtAllocTileCacheDistanceField(&MyAllocator);
+				if (GenerationContext.DistanceField == NULL)
+				{
+					BuildContext.log(RC_LOG_ERROR, "GenerateNavigationData: Out of memory 'DistanceField'.");
+					return false;
+				}
+
+				status = dtBuildTileCacheDistanceField(&MyAllocator, *GenerationContext.Layer, *GenerationContext.DistanceField);
+				if (dtStatusFailed(status))
+				{
+					BuildContext.log(RC_LOG_ERROR, "GenerateNavigationData: Failed to build distance field.");
+					return false;
+				}
+
+				status = dtBuildTileCacheRegions(&MyAllocator, TileConfig.minRegionArea, TileConfig.mergeRegionArea, *GenerationContext.Layer, *GenerationContext.DistanceField);
+			}
+			else
+			{
+				status = dtBuildTileCacheRegionsChunky(&MyAllocator, TileConfig.minRegionArea, TileConfig.mergeRegionArea, *GenerationContext.Layer, TileConfig.TileCacheChunkSize);
+			}
+
+			if (dtStatusFailed(status))
+			{
+				BuildContext.log(RC_LOG_ERROR, "GenerateNavigationData: Failed to build regions.");
+				return false;
+			}
+
+			// skip empty layer
+			if (GenerationContext.Layer->regCount <= 0)
+			{
+				continue;
+			}
+		}
+	
+		{
+			RECAST_STAT(STAT_Navigation_Async_Recast_BuildContours);
+			// Build contour set
+			GenerationContext.ContourSet = dtAllocTileCacheContourSet(&MyAllocator);
+			if (GenerationContext.ContourSet == NULL)
+			{
+				BuildContext.log(RC_LOG_ERROR, "GenerateNavigationData: Out of memory 'ContourSet'.");
+				return false;
+			}
+
+			GenerationContext.ClusterSet = dtAllocTileCacheClusterSet(&MyAllocator);
+			if (GenerationContext.ClusterSet == NULL)
+			{
+				BuildContext.log(RC_LOG_ERROR, "GenerateNavigationData: Out of memory 'ClusterSet'.");
+				return false;
+			}
+
+			status = dtBuildTileCacheContours(&MyAllocator, *GenerationContext.Layer,
+				TileConfig.walkableClimb, TileConfig.maxSimplificationError, TileConfig.cs, TileConfig.ch,
+				*GenerationContext.ContourSet, *GenerationContext.ClusterSet);
+			if (dtStatusFailed(status))
+			{
+				BuildContext.log(RC_LOG_ERROR, "GenerateNavigationData: Failed to generate contour set (0x%08X).", status);
+				return false;
+			}
+
+			// skip empty layer, sometimes there are regions assigned but all flagged as empty (id=0)
+			if (GenerationContext.ContourSet->nconts <= 0)
+			{
+				continue;
+			}
+		}
+
+		{
+			RECAST_STAT(STAT_Navigation_Async_Recast_BuildPolyMesh);
+			// Build poly mesh
+			GenerationContext.PolyMesh = dtAllocTileCachePolyMesh(&MyAllocator);
+			if (GenerationContext.PolyMesh == NULL)
+			{
+				BuildContext.log(RC_LOG_ERROR, "GenerateNavigationData: Out of memory 'PolyMesh'.");
+				return false;
+			}
+
+			status = dtBuildTileCachePolyMesh(&MyAllocator, &BuildContext, *GenerationContext.ContourSet, *GenerationContext.PolyMesh);
+			if (dtStatusFailed(status))
+			{
+				BuildContext.log(RC_LOG_ERROR, "GenerateNavigationData: Failed to generate poly mesh.");
+				return false;
+			}
+
+			status = dtBuildTileCacheClusters(&MyAllocator, *GenerationContext.ClusterSet, *GenerationContext.PolyMesh);
+			if (dtStatusFailed(status))
+			{
+				BuildContext.log(RC_LOG_ERROR, "GenerateNavigationData: Failed to update cluster set.");
+				return false;
+			}
+		}
+
+		// Build detail mesh
+		if (TileConfig.bGenerateDetailedMesh)
+		{
+			RECAST_STAT(STAT_Navigation_Async_Recast_BuildPolyDetail);
+
+			// Build detail mesh.
+			GenerationContext.DetailMesh = dtAllocTileCachePolyMeshDetail(&MyAllocator);
+			if (GenerationContext.DetailMesh == NULL)
+			{
+				BuildContext.log(RC_LOG_ERROR, "GenerateNavigationData: Out of memory 'DetailMesh'.");
+				return false;
+			}
+
+			status = dtBuildTileCachePolyMeshDetail(&MyAllocator, TileConfig.cs, TileConfig.ch, TileConfig.detailSampleDist, TileConfig.detailSampleMaxError,
+				*GenerationContext.Layer, *GenerationContext.PolyMesh, *GenerationContext.DetailMesh);
+			if (dtStatusFailed(status))
+			{
+				BuildContext.log(RC_LOG_ERROR, "GenerateNavigationData: Failed to generate poly detail mesh.");
+				return false;
+			}
+		}
+
+		unsigned char* NavData = 0;
+		int32 NavDataSize = 0;
+
+		if (TileConfig.maxVertsPerPoly <= DT_VERTS_PER_POLYGON &&
+			GenerationContext.PolyMesh->npolys > 0 && GenerationContext.PolyMesh->nverts > 0)
+		{
+			ensure(GenerationContext.PolyMesh->npolys <= TileConfig.MaxPolysPerTile && "Polys per Tile limit exceeded!");
+			if (GenerationContext.PolyMesh->nverts >= 0xffff)
+			{
+				// The vertex indices are ushorts, and cannot point to more than 0xffff vertices.
+				BuildContext.log(RC_LOG_ERROR, "Too many vertices per tile %d (max: %d).", GenerationContext.PolyMesh->nverts, 0xffff);
+				return false;
+			}
+
+			// if we didn't failed already then it's hight time we created data for off-mesh links
+			FOffMeshData OffMeshData;
+			if (OffmeshLinks.Num() > 0)
+			{
+				RECAST_STAT(STAT_Navigation_Async_GatherOffMeshData);
+
+				OffMeshData.Reserve(OffmeshLinks.Num());
+				OffMeshData.AreaClassToIdMap = &AdditionalCachedData.AreaClassToIdMap;
+				OffMeshData.FlagsPerArea = AdditionalCachedData.FlagsPerOffMeshLinkArea;
+				const FSimpleLinkNavModifier* LinkModifier = OffmeshLinks.GetData();
+				const float DefaultSnapHeight = TileConfig.walkableClimb * TileConfig.ch;
+
+				for (int32 LinkModifierIndex = 0; LinkModifierIndex < OffmeshLinks.Num(); ++LinkModifierIndex, ++LinkModifier)
+				{
+					OffMeshData.AddLinks(LinkModifier->Links, LinkModifier->LocalToWorld, TileConfig.AgentIndex, DefaultSnapHeight);
+#if GENERATE_SEGMENT_LINKS
+					OffMeshData.AddSegmentLinks(LinkModifier->SegmentLinks, LinkModifier->LocalToWorld, TileConfig.AgentIndex, DefaultSnapHeight);
+#endif // GENERATE_SEGMENT_LINKS
+				}
+			}
+
+			// fill flags, or else detour won't be able to find polygons
+			// Update poly flags from areas.
+			for (int32 i = 0; i < GenerationContext.PolyMesh->npolys; i++)
+			{
+				GenerationContext.PolyMesh->flags[i] = AdditionalCachedData.FlagsPerArea[GenerationContext.PolyMesh->areas[i]];
+			}
+
+			dtNavMeshCreateParams Params;
+			memset(&Params, 0, sizeof(Params));
+			Params.verts = GenerationContext.PolyMesh->verts;
+			Params.vertCount = GenerationContext.PolyMesh->nverts;
+			Params.polys = GenerationContext.PolyMesh->polys;
+			Params.polyAreas = GenerationContext.PolyMesh->areas;
+			Params.polyFlags = GenerationContext.PolyMesh->flags;
+			Params.polyCount = GenerationContext.PolyMesh->npolys;
+			Params.nvp = GenerationContext.PolyMesh->nvp;
+			if (TileConfig.bGenerateDetailedMesh)
+			{
+				Params.detailMeshes = GenerationContext.DetailMesh->meshes;
+				Params.detailVerts = GenerationContext.DetailMesh->verts;
+				Params.detailVertsCount = GenerationContext.DetailMesh->nverts;
+				Params.detailTris = GenerationContext.DetailMesh->tris;
+				Params.detailTriCount = GenerationContext.DetailMesh->ntris;
+			}
+			Params.offMeshCons = OffMeshData.LinkParams.GetData();
+			Params.offMeshConCount = OffMeshData.LinkParams.Num();
+			Params.walkableHeight = TileConfig.AgentHeight;
+			Params.walkableRadius = TileConfig.AgentRadius;
+			Params.walkableClimb = TileConfig.AgentMaxClimb;
+			Params.tileX = TileX;
+			Params.tileY = TileY;
+			Params.tileLayer = iLayer;
+			rcVcopy(Params.bmin, GenerationContext.Layer->header->bmin);
+			rcVcopy(Params.bmax, GenerationContext.Layer->header->bmax);
+			Params.cs = TileConfig.cs;
+			Params.ch = TileConfig.ch;
+			Params.buildBvTree = TileConfig.bGenerateBVTree;
+#if GENERATE_CLUSTER_LINKS
+			Params.clusterCount = GenerationContext.ClusterSet->nclusters;
+			Params.polyClusters = GenerationContext.ClusterSet->polyMap;
+#endif
+
+			RECAST_STAT(STAT_Navigation_Async_Recast_CreateNavMeshData);
+
+			if (!dtCreateNavMeshData(&Params, &NavData, &NavDataSize))
+			{
+				BuildContext.log(RC_LOG_ERROR, "Could not build Detour navmesh.");
+				return false;
+			}
+		}
+
+		GenerationContext.NavigationData.Add(FNavMeshTileData(NavData, NavDataSize, iLayer, CompressedData.LayerBBox));
+
+		const float ModkB = 1.0f / 1024.0f;
+		BuildContext.log(RC_LOG_PROGRESS, ">> Layer[%d] = Verts(%d) Polys(%d) Memory(%.2fkB) Cache(%.2fkB)",
+			iLayer, GenerationContext.PolyMesh->nverts, GenerationContext.PolyMesh->npolys,
+			GenerationContext.NavigationData.Last().DataSize * ModkB, CompressedLayers[iLayer].DataSize * ModkB);
+	}
+
+	// prepare navigation data of actually rebuild layers for transfer
+	NavigationData = GenerationContext.NavigationData;
+	return true;
+}
+
+void FRecastTileGenerator::MarkDynamicAreas(dtTileCacheLayer& Layer)
+{
+	RECAST_STAT(STAT_Navigation_Async_MarkAreas);
+
+	if (Modifiers.Num())
+	{
+		if (AdditionalCachedData.bUseSortFunction && AdditionalCachedData.ActorOwner && Modifiers.Num() > 1)
+		{
+			AdditionalCachedData.ActorOwner->SortAreasForGenerator(Modifiers);
+		}
+
+		// 1: if navmesh is using low areas, apply only low area replacements
+		if (TileConfig.bMarkLowHeightAreas)
+		{
+			const int32 LowAreaId = RECAST_LOW_AREA;
+			for (int32 ModIdx = 0; ModIdx < Modifiers.Num(); ModIdx++)
+			{
+				FRecastAreaNavModifierElement& Element = Modifiers[ModIdx];
+				for (int32 AreaIdx = Element.Areas.Num() - 1; AreaIdx >= 0; AreaIdx--)
+				{
+					const FAreaNavModifier& AreaMod = Element.Areas[AreaIdx];
+					if (AreaMod.GetApplyMode() == ENavigationAreaMode::ApplyInLowPass ||
+						AreaMod.GetApplyMode() == ENavigationAreaMode::ReplaceInLowPass)
+					{
+						const int32* AreaIDPtr = AdditionalCachedData.AreaClassToIdMap.Find(AreaMod.GetAreaClass());
+						// replace area will be fixed as LowAreaId during this pass, regardless settings in area modifier
+						const int32* ReplaceAreaIDPtr = (AreaMod.GetApplyMode() == ENavigationAreaMode::ReplaceInLowPass) ? &LowAreaId : nullptr;
+
+						if (AreaIDPtr != nullptr)
+						{
+							for (const FTransform& LocalToWorld : Element.PerInstanceTransform)
+							{
+								MarkDynamicArea(AreaMod, LocalToWorld, Layer, *AreaIDPtr, ReplaceAreaIDPtr);
+							}
+
+							if (Element.PerInstanceTransform.Num() == 0)
+							{
+								MarkDynamicArea(AreaMod, FTransform::Identity, Layer, *AreaIDPtr, ReplaceAreaIDPtr);
+							}
+						}
+					}
+				}
+			}
+
+			// 2. remove all low area marking
+			dtReplaceArea(Layer, RECAST_NULL_AREA, RECAST_LOW_AREA);
+		}
+
+		// 3. apply remaining modifiers
+		for (const FRecastAreaNavModifierElement& Element : Modifiers)
+		{
+			for (const FAreaNavModifier& Area : Element.Areas)
+			{
+				if (Area.GetApplyMode() == ENavigationAreaMode::ApplyInLowPass || Area.GetApplyMode() == ENavigationAreaMode::ReplaceInLowPass)
+				{
+					continue;
+				}
+
+				const int32* AreaIDPtr = AdditionalCachedData.AreaClassToIdMap.Find(Area.GetAreaClass());
+				const int32* ReplaceIDPtr = (Area.GetApplyMode() == ENavigationAreaMode::Replace) && Area.GetAreaClassToReplace() ?
+					AdditionalCachedData.AreaClassToIdMap.Find(Area.GetAreaClassToReplace()) : nullptr;
+				
+				if (AreaIDPtr)
+				{
+					for (const FTransform& LocalToWorld : Element.PerInstanceTransform)
+					{
+						MarkDynamicArea(Area, LocalToWorld, Layer, *AreaIDPtr, ReplaceIDPtr);
+					}
+
+					if (Element.PerInstanceTransform.Num() == 0)
+					{
+						MarkDynamicArea(Area, FTransform::Identity, Layer, *AreaIDPtr, ReplaceIDPtr);
+					}
+				}
+			}
+		}
+	}
+	else
+	{
+		if (TileConfig.bMarkLowHeightAreas)
+		{
+			dtReplaceArea(Layer, RECAST_NULL_AREA, RECAST_LOW_AREA);
+		}
+	}
+}
+
+void FRecastTileGenerator::MarkDynamicArea(const FAreaNavModifier& Modifier, const FTransform& LocalToWorld, dtTileCacheLayer& Layer)
+{
+	const int32* AreaIDPtr = AdditionalCachedData.AreaClassToIdMap.Find(Modifier.GetAreaClass());
+	const int32* ReplaceIDPtr = Modifier.GetAreaClassToReplace() ? AdditionalCachedData.AreaClassToIdMap.Find(Modifier.GetAreaClassToReplace()) : nullptr;
+	if (AreaIDPtr)
+	{
+		MarkDynamicArea(Modifier, LocalToWorld, Layer, *AreaIDPtr, ReplaceIDPtr);
+	}
+}
+
+void FRecastTileGenerator::MarkDynamicArea(const FAreaNavModifier& Modifier, const FTransform& LocalToWorld, dtTileCacheLayer& Layer, const int32 AreaID, const int32* ReplaceIDPtr)
+{
+	const float ExpandBy = TileConfig.AgentRadius;
+
+	// Expand by 1 cell height up and down to cover for voxel grid inaccuracy
+	const float OffsetZMax = TileConfig.ch;
+	const float OffsetZMin = TileConfig.ch + (Modifier.ShouldIncludeAgentHeight() ? TileConfig.AgentHeight : 0.0f);
+
+	// Check whether modifier affects this layer
+	const FBox LayerUnrealBounds = Recast2UnrealBox(Layer.header->bmin, Layer.header->bmax);
+	FBox ModifierBounds = Modifier.GetBounds().TransformBy(LocalToWorld);
+	ModifierBounds.Min -= FVector(ExpandBy, ExpandBy, OffsetZMin);
+	ModifierBounds.Max += FVector(ExpandBy, ExpandBy, OffsetZMax);
+
+	if (!LayerUnrealBounds.Intersect(ModifierBounds))
+	{
+		return;
+	}
+
+	const float* LayerRecastOrig = Layer.header->bmin;
+	switch (Modifier.GetShapeType())
+	{
+	case ENavigationShapeType::Cylinder:
+		{
+			FCylinderNavAreaData CylinderData;
+			Modifier.GetCylinder(CylinderData);
+
+			// Only scaling and translation
+			FVector Scale3D = LocalToWorld.GetScale3D().GetAbs();
+			CylinderData.Height *= Scale3D.Z;
+			CylinderData.Radius *= FMath::Max(Scale3D.X, Scale3D.Y);
+			CylinderData.Origin = LocalToWorld.TransformPosition(CylinderData.Origin);
+			
+			const float OffsetZMid = (OffsetZMax - OffsetZMin) * 0.5f;
+			CylinderData.Origin.Z += OffsetZMid;
+			CylinderData.Height += FMath::Abs(OffsetZMid) * 2.f;
+			CylinderData.Radius += ExpandBy;
+			
+			FVector RecastPos = Unreal2RecastPoint(CylinderData.Origin);
+
+			if (ReplaceIDPtr)
+			{
+				dtReplaceCylinderArea(Layer, LayerRecastOrig, TileConfig.cs, TileConfig.ch,
+					&(RecastPos.X), CylinderData.Radius, CylinderData.Height, AreaID, *ReplaceIDPtr);
+			}
+			else
+			{
+				dtMarkCylinderArea(Layer, LayerRecastOrig, TileConfig.cs, TileConfig.ch,
+					&(RecastPos.X), CylinderData.Radius, CylinderData.Height, AreaID);
+			}
+		}
+		break;
+
+	case ENavigationShapeType::Box:
+		{
+			FBoxNavAreaData BoxData;
+			Modifier.GetBox(BoxData);
+
+			FBox WorldBox = FBox::BuildAABB(BoxData.Origin, BoxData.Extent).TransformBy(LocalToWorld);
+			WorldBox = WorldBox.ExpandBy(FVector(ExpandBy, ExpandBy, 0));
+			WorldBox.Min.Z -= OffsetZMin;
+			WorldBox.Max.Z += OffsetZMax;
+
+			FBox RacastBox = Unreal2RecastBox(WorldBox);
+			FVector RecastPos;
+			FVector RecastExtent;
+			RacastBox.GetCenterAndExtents(RecastPos, RecastExtent);
+				
+			if (ReplaceIDPtr)
+			{
+				dtReplaceBoxArea(Layer, LayerRecastOrig, TileConfig.cs, TileConfig.ch,
+					&(RecastPos.X), &(RecastExtent.X), AreaID, *ReplaceIDPtr);
+			}
+			else
+			{
+				dtMarkBoxArea(Layer, LayerRecastOrig, TileConfig.cs, TileConfig.ch,
+					&(RecastPos.X), &(RecastExtent.X), AreaID);
+			}
+		}
+		break;
+
+	case ENavigationShapeType::Convex:
+		{
+			FConvexNavAreaData ConvexData;
+			Modifier.GetConvex(ConvexData);
+			// Only scaling and translation
+			PartialTransformConvexHull(ConvexData, LocalToWorld);
+
+			TArray<FVector> ConvexVerts;
+			GrowConvexHull(ExpandBy, ConvexData.Points, ConvexVerts);
+			ConvexData.MinZ -= OffsetZMin;
+			ConvexData.MaxZ += OffsetZMax;
+
+			if (ConvexVerts.Num())
+			{
+				TArray<float> ConvexCoords;
+				ConvexCoords.AddZeroed(ConvexVerts.Num() * 3);
+						
+				float* ItCoord = ConvexCoords.GetData();
+				for (int32 i = 0; i < ConvexVerts.Num(); i++)
+				{
+					const FVector RecastV = Unreal2RecastPoint(ConvexVerts[i]);
+					*ItCoord = RecastV.X; ItCoord++;
+					*ItCoord = RecastV.Y; ItCoord++;
+					*ItCoord = RecastV.Z; ItCoord++;
+				}
+
+				if (ReplaceIDPtr)
+				{
+					dtReplaceConvexArea(Layer, LayerRecastOrig, TileConfig.cs, TileConfig.ch,
+						ConvexCoords.GetData(), ConvexVerts.Num(), ConvexData.MinZ, ConvexData.MaxZ, AreaID, *ReplaceIDPtr);
+				}
+				else
+				{
+					dtMarkConvexArea(Layer, LayerRecastOrig, TileConfig.cs, TileConfig.ch,
+						ConvexCoords.GetData(), ConvexVerts.Num(), ConvexData.MinZ, ConvexData.MaxZ, AreaID);
+				}
+			}
+		}
+		break;
+
+	default: break;
+	}
+}
+
+uint32 FRecastTileGenerator::GetUsedMemCount() const
+{
+	uint32 TotalMemory = 0;
+	TotalMemory += InclusionBounds.GetAllocatedSize();
+	TotalMemory += Modifiers.GetAllocatedSize();
+	TotalMemory += OffmeshLinks.GetAllocatedSize();
+	TotalMemory += RawGeometry.GetAllocatedSize();
+	
+	for (const FRecastRawGeometryElement& Element : RawGeometry)
+	{
+		TotalMemory += Element.GeomCoords.GetAllocatedSize();
+		TotalMemory += Element.GeomIndices.GetAllocatedSize();
+		TotalMemory += Element.PerInstanceTransform.GetAllocatedSize();
+	}
+
+	for (const FRecastAreaNavModifierElement& Element : Modifiers)
+	{
+		TotalMemory += Element.Areas.GetAllocatedSize();
+		TotalMemory += Element.PerInstanceTransform.GetAllocatedSize();
+	}
+
+	const FSimpleLinkNavModifier* SimpleLink = OffmeshLinks.GetData();
+	for (int32 Index = 0; Index < OffmeshLinks.Num(); ++Index, ++SimpleLink)
+	{
+		TotalMemory += SimpleLink->Links.GetAllocatedSize();
+	}
+
+	TotalMemory += CompressedLayers.GetAllocatedSize();
+	for (int32 i = 0; i < CompressedLayers.Num(); i++)
+	{
+		TotalMemory += CompressedLayers[i].DataSize;
+	}
+
+	TotalMemory += NavigationData.GetAllocatedSize();
+	for (int32 i = 0; i < NavigationData.Num(); i++)
+	{
+		TotalMemory += NavigationData[i].DataSize;
+	}
+
+	return TotalMemory;
+}
+
+void FRecastTileGenerator::AddReferencedObjects(FReferenceCollector& Collector)
+{
+	for (auto& RelevantData : NavigationRelevantData)
+	{
+		UObject* Owner = RelevantData->GetOwner();
+		if (Owner)
+		{
+			Collector.AddReferencedObject(Owner);
+		}
+	}
+}
+
+static int32 CaclulateMaxTilesCount(const TNavStatArray<FBox>& NavigableAreas, float TileSizeinWorldUnits, float AvgLayersPerGridCell)
+{
+	int32 GridCellsCount = 0;
+	for (FBox AreaBounds : NavigableAreas)
+	{
+		// TODO: need more precise calculation, currently we don't take into account that volumes can be overlapped
+		FBox RCBox = Unreal2RecastBox(AreaBounds);
+		int32 XSize = FMath::CeilToInt(RCBox.GetSize().X/TileSizeinWorldUnits) + 1;
+		int32 YSize = FMath::CeilToInt(RCBox.GetSize().Z/TileSizeinWorldUnits) + 1;
+		GridCellsCount+= (XSize*YSize);
+	}
+	
+	return FMath::CeilToInt(GridCellsCount * AvgLayersPerGridCell);
+}
+
+// Whether navmesh is static, does not support rebuild from geometry
+static bool IsGameStaticNavMesh(ARecastNavMesh* InNavMesh)
+{
+	return (InNavMesh->GetWorld()->IsGameWorld() && InNavMesh->GetRuntimeGenerationMode() != ERuntimeGenerationType::Dynamic);
+}
+
+//----------------------------------------------------------------------//
+// FRecastNavMeshGenerator
+//----------------------------------------------------------------------//
+
+FRecastNavMeshGenerator::FRecastNavMeshGenerator(ARecastNavMesh& InDestNavMesh)
+	: NumActiveTiles(0)
+	, MaxTileGeneratorTasks(1)
+	, AvgLayersPerTile(8.0f)
+	, DestNavMesh(&InDestNavMesh)
+	, bInitialized(false)
+	, bRestrictBuildingToActiveTiles(false)
+	, Version(0)
+{
+	INC_DWORD_STAT_BY(STAT_NavigationMemory, sizeof(*this));
+
+	Init();
+
+	int32 MaxTiles = 0;
+	int32 MaxPolysPerTile = 0;
+
+	// recreate navmesh if no data was loaded, or when loaded data doesn't match current grid layout
+	bool bRecreateNavmesh = true;
+	if (DestNavMesh->HasValidNavmesh())
+	{
+		const bool bGameStaticNavMesh = IsGameStaticNavMesh(DestNavMesh);
+		const dtNavMeshParams* SavedNavParams = DestNavMesh->GetRecastNavMeshImpl()->DetourNavMesh->getParams();
+		if (SavedNavParams)
+		{
+			if (bGameStaticNavMesh)
+			{
+				bRecreateNavmesh = false;
+				MaxTiles = SavedNavParams->maxTiles;
+				MaxPolysPerTile = SavedNavParams->maxPolys;
+			}
+			else
+			{
+				const float TileDim = Config.tileSize * Config.cs;
+				if (SavedNavParams->tileHeight == TileDim && SavedNavParams->tileWidth == TileDim)
+				{
+					const FVector Orig = Recast2UnrealPoint(SavedNavParams->orig);
+					const FVector OrigError(FMath::Fmod(Orig.X, TileDim), FMath::Fmod(Orig.X, TileDim), FMath::Fmod(Orig.X, TileDim));
+					if (OrigError.IsNearlyZero())
+					{
+						bRecreateNavmesh = false;
+					}
+				}
+
+				// if new navmesh needs more tiles, force recreating
+				if (!bRecreateNavmesh)
+				{
+					CalcNavMeshProperties(MaxTiles, MaxPolysPerTile);
+					if (FMath::Log2(MaxTiles) != FMath::Log2(SavedNavParams->maxTiles))
+					{
+						bRecreateNavmesh = true;
+					}
+				}
+			}
+		};
+	}
+
+	if (bRecreateNavmesh)
+	{
+		// recreate navmesh from scratch if no data was loaded
+		ConstructTiledNavMesh();
+	}
+	else
+	{
+		// otherwise just update generator params
+		Config.MaxPolysPerTile = MaxPolysPerTile;
+		NumActiveTiles = GetTilesCountHelper(DestNavMesh->GetRecastNavMeshImpl()->DetourNavMesh);
+	}
+}
+
+FRecastNavMeshGenerator::~FRecastNavMeshGenerator()
+{
+	DEC_DWORD_STAT_BY( STAT_NavigationMemory, sizeof(*this) );
+}
+
+void FRecastNavMeshGenerator::Init()
+{
+	// @todo those variables should be tweakable per navmesh actor
+	const float CellSize = DestNavMesh->CellSize;
+	const float CellHeight = DestNavMesh->CellHeight;
+	const float AgentHeight = DestNavMesh->AgentHeight;
+	const float MaxAgentHeight = DestNavMesh->AgentMaxHeight;
+	const float AgentMaxSlope = DestNavMesh->AgentMaxSlope;
+	const float AgentMaxClimb = DestNavMesh->AgentMaxStepHeight;
+	const float AgentRadius = DestNavMesh->AgentRadius;
+
+	Config.Reset();
+	
+	Config.cs = CellSize;
+	Config.ch = CellHeight;
+	Config.walkableSlopeAngle = AgentMaxSlope;
+	Config.walkableHeight = (int32)ceilf(AgentHeight / CellHeight);
+	Config.walkableClimb = (int32)ceilf(AgentMaxClimb / CellHeight);
+	const float WalkableRadius = FMath::CeilToFloat(AgentRadius / CellSize);
+	Config.walkableRadius = WalkableRadius;
+	
+	// store original sizes
+	Config.AgentHeight = AgentHeight;
+	Config.AgentMaxClimb = AgentMaxClimb;
+	Config.AgentRadius = AgentRadius;
+
+	Config.borderSize = WalkableRadius + 3;
+	Config.maxEdgeLen = (int32)(1200.0f / CellSize);
+	Config.maxSimplificationError = 1.3f;
+	// hardcoded, but can be overridden by RecastNavMesh params later
+	Config.minRegionArea = (int32)rcSqr(0);
+	Config.mergeRegionArea = (int32)rcSqr(20.f);
+
+	Config.maxVertsPerPoly = (int32)MAX_VERTS_PER_POLY;
+	Config.detailSampleDist = 600.0f;
+	Config.detailSampleMaxError = 1.0f;
+	Config.PolyMaxHeight = (int32)ceilf(MaxAgentHeight / CellHeight);
+		
+	Config.minRegionArea = (int32)rcSqr(DestNavMesh->MinRegionArea / CellSize);
+	Config.mergeRegionArea = (int32)rcSqr(DestNavMesh->MergeRegionSize / CellSize);
+	Config.maxSimplificationError = DestNavMesh->MaxSimplificationError;
+	Config.bPerformVoxelFiltering = DestNavMesh->bPerformVoxelFiltering;
+	Config.bMarkLowHeightAreas = DestNavMesh->bMarkLowHeightAreas;
+	Config.bFilterLowSpanSequences = DestNavMesh->bFilterLowSpanSequences;
+	Config.bFilterLowSpanFromTileCache = DestNavMesh->bFilterLowSpanFromTileCache;
+	if (DestNavMesh->bMarkLowHeightAreas)
+	{
+		Config.walkableHeight = 1;
+	}
+	
+	AdditionalCachedData = FRecastNavMeshCachedData::Construct(DestNavMesh);
+
+	const UNavigationSystemV1* NavSys = FNavigationSystem::GetCurrent<UNavigationSystemV1>(GetWorld());
+	Config.AgentIndex = NavSys->GetSupportedAgentIndex(DestNavMesh);
+
+	Config.tileSize = FMath::TruncToInt(DestNavMesh->TileSizeUU / CellSize);
+
+	Config.regionChunkSize = Config.tileSize / DestNavMesh->LayerChunkSplits;
+	Config.TileCacheChunkSize = Config.tileSize / DestNavMesh->RegionChunkSplits;
+	Config.regionPartitioning = DestNavMesh->LayerPartitioning;
+	Config.TileCachePartitionType = DestNavMesh->RegionPartitioning;
+
+	UpdateNavigationBounds();
+
+	/** setup maximum number of active tile generator*/
+	const int32 NumberOfWorkerThreads = FTaskGraphInterface::Get().GetNumWorkerThreads();
+	MaxTileGeneratorTasks = FMath::Min(FMath::Max(NumberOfWorkerThreads * 2, 1), GetOwner() ? GetOwner()->GetMaxSimultaneousTileGenerationJobsCount() : INT_MAX);
+	UE_LOG(LogNavigation, Log, TEXT("Using max of %d workers to build navigation."), MaxTileGeneratorTasks);
+	NumActiveTiles = 0;
+
+	// prepare voxel cache if needed
+	if (ARecastNavMesh::IsVoxelCacheEnabled())
+	{
+		VoxelCacheContext.Create(Config.tileSize + Config.borderSize * 2, Config.cs, Config.ch);
+	}
+
+	bInitialized = true;
+}
+
+void FRecastNavMeshGenerator::UpdateNavigationBounds()
+{
+	const UNavigationSystemV1* NavSys = FNavigationSystem::GetCurrent<UNavigationSystemV1>(GetWorld());
+	if (NavSys)
+	{		
+		if (NavSys->ShouldGenerateNavigationEverywhere() == false)
+		{
+			FBox BoundsSum(ForceInit);
+			if (DestNavMesh)
+			{
+				TArray<FBox> SupportedBounds;
+				NavSys->GetNavigationBoundsForNavData(*DestNavMesh, SupportedBounds);
+				InclusionBounds.Reset(SupportedBounds.Num());
+
+				for (const FBox& Box : SupportedBounds)
+				{
+					InclusionBounds.Add(Box);
+					BoundsSum += Box;
+				}
+			}
+			TotalNavBounds = BoundsSum;
+		}
+		else
+		{
+			InclusionBounds.Reset(1);
+			TotalNavBounds = NavSys->GetWorldBounds();
+			if (!TotalNavBounds.IsValid)
+			{
+				InclusionBounds.Add(TotalNavBounds);
+			}
+		}
+	}
+	else
+	{
+		TotalNavBounds = FBox(ForceInit);
+	}
+}
+
+bool FRecastNavMeshGenerator::ConstructTiledNavMesh() 
+{
+	bool bSuccess = false;
+
+	// There is should not be any active build tasks
+	CancelBuild();
+
+	// create new Detour navmesh instance
+	dtNavMesh* DetourMesh = dtAllocNavMesh();	
+	if (DetourMesh)
+	{
+		++Version;
+		
+		dtNavMeshParams TiledMeshParameters;
+		FMemory::Memzero(TiledMeshParameters);	
+		rcVcopy(TiledMeshParameters.orig, &FVector::ZeroVector.X);
+		TiledMeshParameters.tileWidth = Config.tileSize * Config.cs;
+		TiledMeshParameters.tileHeight = Config.tileSize * Config.cs;
+
+		CalcNavMeshProperties(TiledMeshParameters.maxTiles, TiledMeshParameters.maxPolys);
+		Config.MaxPolysPerTile = TiledMeshParameters.maxPolys;
+
+		const dtStatus status = DetourMesh->init(&TiledMeshParameters);
+
+		if (dtStatusFailed(status))
+		{
+			UE_LOG(LogNavigation, Warning, TEXT("ConstructTiledNavMesh: Could not init navmesh.") );
+			bSuccess = false;
+		}
+		else
+		{
+			bSuccess = true;
+			NumActiveTiles = GetTilesCountHelper(DetourMesh);
+			DestNavMesh->GetRecastNavMeshImpl()->SetRecastMesh(DetourMesh);
+		}
+	}
+	else
+	{
+		UE_LOG(LogNavigation, Warning, TEXT("ConstructTiledNavMesh: Could not allocate navmesh.") );
+		bSuccess = false;
+	}
+	
+	return bSuccess;
+}
+
+void FRecastNavMeshGenerator::CalcPolyRefBits(ARecastNavMesh* NavMeshOwner, int32& MaxTileBits, int32& MaxPolyBits)
+{
+	static const int32 TotalBits = (sizeof(dtPolyRef) * 8);
+#if USE_64BIT_ADDRESS
+	MaxTileBits = NavMeshOwner ? FMath::CeilToFloat(FMath::Log2(NavMeshOwner->GetTileNumberHardLimit())) : 20;
+	MaxPolyBits = FMath::Min<int32>(32, (TotalBits - DT_MIN_SALT_BITS) - MaxTileBits);
+#else
+	MaxTileBits = 14;
+	MaxPolyBits = (TotalBits - DT_MIN_SALT_BITS) - MaxTileBits;
+#endif//USE_64BIT_ADDRESS
+}
+
+void FRecastNavMeshGenerator::CalcNavMeshProperties(int32& MaxTiles, int32& MaxPolys)
+{
+	int32 MaxTileBits = -1;
+	int32 MaxPolyBits = -1;
+
+	// limit max amount of tiles
+	CalcPolyRefBits(DestNavMesh, MaxTileBits, MaxPolyBits);
+	
+	const int32 MaxTilesFromMask = (1 << MaxTileBits);
+	int32 MaxRequestedTiles = 0;
+	if (DestNavMesh->IsResizable())
+	{
+		MaxRequestedTiles = CaclulateMaxTilesCount(InclusionBounds, Config.tileSize * Config.cs, AvgLayersPerTile);
+	}
+	else
+	{
+		MaxRequestedTiles = DestNavMesh->TilePoolSize;
+	}
+
+	if (MaxRequestedTiles < 0 || MaxRequestedTiles > MaxTilesFromMask)
+	{
+		UE_LOG(LogNavigation, Error, TEXT("Navmesh bounds are too large! Limiting requested tiles count (%d) to: (%d)"), MaxRequestedTiles, MaxTilesFromMask);
+		MaxRequestedTiles = MaxTilesFromMask;
+	}
+
+	// Max tiles and max polys affect how the tile IDs are calculated.
+	// There are (sizeof(dtPolyRef)*8 - DT_MIN_SALT_BITS) bits available for 
+	// identifying a tile and a polygon.
+#if USE_64BIT_ADDRESS
+	MaxPolys = (MaxPolyBits >= 32) ? INT_MAX : (1 << MaxPolyBits);
+#else
+	MaxPolys = 1 << ((sizeof(dtPolyRef) * 8 - DT_MIN_SALT_BITS) - MaxTileBits);
+#endif // USE_64BIT_ADDRESS
+	MaxTiles = MaxRequestedTiles;
+}
+
+bool FRecastNavMeshGenerator::RebuildAll()
+{
+	DestNavMesh->UpdateNavVersion();
+	
+	// Recreate recast navmesh
+	DestNavMesh->GetRecastNavMeshImpl()->ReleaseDetourNavMesh();
+	ConstructTiledNavMesh();
+	
+	// if rebuilding all no point in keeping "old" invalidated areas
+	TArray<FNavigationDirtyArea> DirtyAreas;
+	for (FBox AreaBounds : InclusionBounds)
+	{
+		FNavigationDirtyArea DirtyArea(AreaBounds, ENavigationDirtyFlag::All | ENavigationDirtyFlag::NavigationBounds);
+		DirtyAreas.Add(DirtyArea);
+	}
+
+	if (DirtyAreas.Num())
+	{
+		MarkDirtyTiles(DirtyAreas);
+	}
+	else
+	{
+		// There are no navigation bounds to build, probably navmesh was resized and we just need to update debug draw
+		DestNavMesh->RequestDrawingUpdate();
+	}
+
+	return true;
+}
+
+void FRecastNavMeshGenerator::EnsureBuildCompletion()
+{
+	const bool bHadTasks = GetNumRemaningBuildTasks() > 0;
+	
+	const bool bDoAsyncDataGathering = (GatherGeometryOnGameThread() == false);
+	do 
+	{
+		const int32 NumTasksToSubmit = (bDoAsyncDataGathering ? 1 : MaxTileGeneratorTasks) - RunningDirtyTiles.Num();
+		ProcessTileTasks(NumTasksToSubmit);
+		
+		// Block until tasks are finished
+		for (FRunningTileElement& Element : RunningDirtyTiles)
+		{
+			Element.AsyncTask->EnsureCompletion();
+		}
+	}
+	while (GetNumRemaningBuildTasks() > 0);
+
+	// Update navmesh drawing only if we had something to build
+	if (bHadTasks)
+	{
+		DestNavMesh->RequestDrawingUpdate();
+	}
+}
+
+void FRecastNavMeshGenerator::CancelBuild()
+{
+	DiscardCurrentBuildingTasks();
+
+#if	WITH_EDITOR	
+	RecentlyBuiltTiles.Empty();
+#endif//WITH_EDITOR
+}
+
+void FRecastNavMeshGenerator::TickAsyncBuild(float DeltaSeconds)
+{
+	bool bRequestDrawingUpdate = false;
+
+#if	WITH_EDITOR
+	// Remove expired tiles
+	{
+		const double Timestamp = FPlatformTime::Seconds();
+		const int32 NumPreRemove = RecentlyBuiltTiles.Num();
+		
+		RecentlyBuiltTiles.RemoveAllSwap([&](const FTileTimestamp& Tile) { return (Timestamp - Tile.Timestamp) > 0.5; });
+
+		const int32 NumPostRemove = RecentlyBuiltTiles.Num();
+		bRequestDrawingUpdate = (NumPreRemove != NumPostRemove);
+	}
+#endif//WITH_EDITOR
+
+	// Submit async tile build tasks in case we have dirty tiles and have room for them
+	const UNavigationSystemV1* NavSys = FNavigationSystem::GetCurrent<UNavigationSystemV1>(GetWorld());
+	check(NavSys);
+	const int32 NumRunningTasks = NavSys->GetNumRunningBuildTasks();
+	// this is a temp solution to enforce only one worker thread if GatherGeometryOnGameThread == false
+	// due to missing safety features
+	const bool bDoAsyncDataGathering = GatherGeometryOnGameThread() == false;
+
+	const int32 NumTasksToSubmit = (bDoAsyncDataGathering ? 1 : MaxTileGeneratorTasks) - NumRunningTasks;
+	TArray<uint32> UpdatedTileIndices = ProcessTileTasks(NumTasksToSubmit);
+			
+	if (UpdatedTileIndices.Num() > 0)
+	{
+		// Invalidate active paths that go through regenerated tiles
+		DestNavMesh->OnNavMeshTilesUpdated(UpdatedTileIndices);
+		bRequestDrawingUpdate = true;
+
+#if	WITH_EDITOR
+		// Store completed tiles with timestamps to have ability to distinguish during debug draw
+		const double Timestamp = FPlatformTime::Seconds();
+		RecentlyBuiltTiles.Reserve(RecentlyBuiltTiles.Num() + UpdatedTileIndices.Num());
+		for (uint32 TiledIdx : UpdatedTileIndices)
+		{
+			FTileTimestamp TileTimestamp;
+			TileTimestamp.TileIdx = TiledIdx;
+			TileTimestamp.Timestamp = Timestamp;
+			RecentlyBuiltTiles.Add(TileTimestamp);
+		}
+#endif//WITH_EDITOR
+	}
+
+	if (bRequestDrawingUpdate)
+	{
+		DestNavMesh->RequestDrawingUpdate();
+	}
+}
+
+void FRecastNavMeshGenerator::OnNavigationBoundsChanged()
+{
+	UpdateNavigationBounds();
+	
+	dtNavMesh* DetourMesh = DestNavMesh->GetRecastNavMeshImpl()->GetRecastMesh();
+	if (!IsGameStaticNavMesh(DestNavMesh) && DestNavMesh->IsResizable() && DetourMesh)
+	{
+		// Check whether Navmesh size needs to be changed
+		int32 MaxRequestedTiles = CaclulateMaxTilesCount(InclusionBounds, Config.tileSize * Config.cs, AvgLayersPerTile);
+		if (DetourMesh->getMaxTiles() != MaxRequestedTiles)
+		{
+			// Destroy current NavMesh, it will be allocated with a new size on next build request
+			DestNavMesh->GetRecastNavMeshImpl()->SetRecastMesh(nullptr);
+		}
+	}
+}
+
+void FRecastNavMeshGenerator::RebuildDirtyAreas(const TArray<FNavigationDirtyArea>& InDirtyAreas)
+{
+	dtNavMesh* DetourMesh = DestNavMesh->GetRecastNavMeshImpl()->GetRecastMesh();
+	if (DetourMesh == nullptr)
+	{
+		ConstructTiledNavMesh();
+	}
+	
+	MarkDirtyTiles(InDirtyAreas);
+}
+
+void FRecastNavMeshGenerator::OnAreaAdded(const UClass* AreaClass, int32 AreaID)
+{
+	AdditionalCachedData.OnAreaAdded(AreaClass, AreaID);
+}
+
+int32 FRecastNavMeshGenerator::FindInclusionBoundEncapsulatingBox(const FBox& Box) const
+{
+	for (int32 Index = 0; Index < InclusionBounds.Num(); ++Index)
+	{
+		if (DoesBoxContainBox(InclusionBounds[Index], Box))
+		{
+			return Index;
+		}
+	}
+
+	return INDEX_NONE;
+}
+
+void FRecastNavMeshGenerator::RestrictBuildingToActiveTiles(bool InRestrictBuildingToActiveTiles) 
+{ 
+	if (bRestrictBuildingToActiveTiles != InRestrictBuildingToActiveTiles)
+	{
+		bRestrictBuildingToActiveTiles = InRestrictBuildingToActiveTiles;
+		if (InRestrictBuildingToActiveTiles)
+		{
+			// gather non-empty tiles and add them to ActiveTiles
+
+			const dtNavMesh* DetourMesh = DestNavMesh->GetRecastNavMeshImpl()->GetRecastMesh();
+
+			if (DetourMesh != nullptr && DetourMesh->isEmpty() == false)
+			{
+				ActiveTiles.Reset();
+				int32 TileCount = DetourMesh->getMaxTiles();
+				for (int32 TileIndex = 0; TileIndex < TileCount; ++TileIndex)
+				{
+					const dtMeshTile* Tile = DetourMesh->getTile(TileIndex);
+					if (Tile != nullptr && Tile->header != nullptr && Tile->header->polyCount > 0)
+					{
+						ActiveTiles.AddUnique(FIntPoint(Tile->header->x, Tile->header->y));
+					}
+				}
+			}
+		}
+	}
+}
+
+bool FRecastNavMeshGenerator::IsInActiveSet(const FIntPoint& Tile) const
+{
+	// @TODO checking if given tile is in active tiles needs to be faster
+	return bRestrictBuildingToActiveTiles == false || ActiveTiles.Find(Tile) != INDEX_NONE;
+}
+
+void FRecastNavMeshGenerator::RemoveTiles(const TArray<FIntPoint>& Tiles)
+{
+	for (const FIntPoint& TileXY : Tiles)
+	{
+		RemoveTileLayers(TileXY.X, TileXY.Y);
+		if (PendingDirtyTiles.Num() > 0)
+		{
+			FPendingTileElement DirtyTile;
+			DirtyTile.Coord = TileXY;
+			PendingDirtyTiles.Remove(DirtyTile);
+		}
+	}
+}
+
+void FRecastNavMeshGenerator::ReAddTiles(const TArray<FIntPoint>& Tiles)
+{
+	static const FVector Expansion(1, 1, BIG_NUMBER);
+	// a little trick here - adding a dirty area so that navmesh building figures it out on its own
+	dtNavMesh* DetourMesh = DestNavMesh->GetRecastNavMeshImpl()->GetRecastMesh();
+	const dtNavMeshParams* SavedNavParams = DestNavMesh->GetRecastNavMeshImpl()->DetourNavMesh->getParams();
+	const float TileDim = Config.tileSize * Config.cs;
+
+	TSet<FPendingTileElement> DirtyTiles;
+
+	// @note we act on assumption all items in Tiles are unique
+	for (const FIntPoint& TileCoords : Tiles)
+	{
+		FPendingTileElement Element;
+		Element.Coord = TileCoords;
+		Element.bRebuildGeometry = true;
+		DirtyTiles.Add(Element);
+	}
+
+	int32 NumTilesMarked = DirtyTiles.Num();
+
+	// Merge all pending tiles into one container
+	for (const FPendingTileElement& Element : PendingDirtyTiles)
+	{
+		FPendingTileElement* ExistingElement = DirtyTiles.Find(Element);
+		if (ExistingElement)
+		{
+			ExistingElement->bRebuildGeometry |= Element.bRebuildGeometry;
+			// Append area bounds to existing list 
+			if (ExistingElement->bRebuildGeometry == false)
+			{
+				ExistingElement->DirtyAreas.Append(Element.DirtyAreas);
+			}
+			else
+			{
+				ExistingElement->DirtyAreas.Empty();
+			}
+		}
+		else
+		{
+			DirtyTiles.Add(Element);
+		}
+	}
+
+	// Dump results into array
+	PendingDirtyTiles.Empty(DirtyTiles.Num());
+	for (const FPendingTileElement& Element : DirtyTiles)
+	{
+		PendingDirtyTiles.Add(Element);
+	}
+
+	// Sort tiles by proximity to players 
+	if (NumTilesMarked > 0)
+	{
+		SortPendingBuildTiles();
+	}
+
+	/*TArray<FNavigationDirtyArea> DirtyAreasContainer;
+	DirtyAreasContainer.Reserve(Tiles.Num());
+
+	TSet<FPendingTileElement> DirtyTiles;
+
+	for (const FIntPoint& TileCoords : Tiles)
+	{
+		const FVector TileCenter = Recast2UnrealPoint(SavedNavParams->orig) + FVector(TileDim * float(TileCoords.X), TileDim * float(TileCoords.Y), 0);
+		
+		FNavigationDirtyArea DirtyArea(FBox(TileCenter - Expansion, TileCenter - 1), ENavigationDirtyFlag::All);
+		DirtyAreasContainer.Add(DirtyArea);
+	}
+
+	MarkDirtyTiles(DirtyAreasContainer);*/
+}
+
+namespace RecastTileVersionHelper
+{
+	inline uint32 GetUpdatedTileId(dtPolyRef& TileRef, dtNavMesh* DetourMesh)
+	{
+		uint32 DecodedTileId = 0, DecodedPolyId = 0, DecodedSaltId = 0;
+		DetourMesh->decodePolyId(TileRef, DecodedSaltId, DecodedTileId, DecodedPolyId);
+
+		DecodedSaltId = (DecodedSaltId + 1) & ((1 << DetourMesh->getSaltBits()) - 1);
+		if (DecodedSaltId == 0)
+		{
+			DecodedSaltId++;
+		}
+
+		TileRef = DetourMesh->encodePolyId(DecodedSaltId, DecodedTileId, DecodedPolyId);
+		return DecodedTileId;
+	}
+}
+
+TArray<uint32> FRecastNavMeshGenerator::RemoveTileLayers(const int32 TileX, const int32 TileY, TMap<int32, dtPolyRef>* OldLayerTileIdMap)
+{
+	dtNavMesh* DetourMesh = DestNavMesh->GetRecastNavMeshImpl()->GetRecastMesh();
+	TArray<uint32> UpdatedIndices;
+	
+	if (DetourMesh != nullptr && DetourMesh->isEmpty() == false)
+	{
+		const int32 NumLayers = DetourMesh->getTileCountAt(TileX, TileY);
+
+		if (NumLayers > 0)
+		{
+			TArray<dtMeshTile*> Tiles;
+			Tiles.AddZeroed(NumLayers);
+			DetourMesh->getTilesAt(TileX, TileY, (const dtMeshTile**)Tiles.GetData(), NumLayers);
+
+			for (int32 i = 0; i < NumLayers; i++)
+			{
+				const int32 LayerIndex = Tiles[i]->header->layer;
+				dtPolyRef TileRef = DetourMesh->getTileRef(Tiles[i]);
+
+				NumActiveTiles--;
+				UE_LOG(LogNavigation, Log, TEXT("%s> Tile (%d,%d:%d), removing TileRef: 0x%X (active:%d)"),
+					*DestNavMesh->GetName(), TileX, TileY, LayerIndex, TileRef, NumActiveTiles);
+
+				DetourMesh->removeTile(TileRef, nullptr, nullptr);
+
+				uint32 TileId = RecastTileVersionHelper::GetUpdatedTileId(TileRef, DetourMesh);
+				UpdatedIndices.AddUnique(TileId);
+
+				if (OldLayerTileIdMap)
+				{
+					OldLayerTileIdMap->Add(LayerIndex, TileRef);
+				}
+			}
+		}
+
+		// Remove compressed tile cache layers
+		DestNavMesh->RemoveTileCacheLayers(TileX, TileY);
+	}
+
+	return UpdatedIndices;
+}
+
+TArray<uint32> FRecastNavMeshGenerator::AddGeneratedTiles(FRecastTileGenerator& TileGenerator)
+{
+	SCOPE_CYCLE_COUNTER(STAT_Navigation_AddGeneratedTiles);
+	
+	TMap<int32, dtPolyRef> OldLayerTileIdMap;
+	TArray<uint32> ResultTileIndices;
+	const int32 TileX = TileGenerator.GetTileX();
+	const int32 TileY = TileGenerator.GetTileY();
+	
+	if (TileGenerator.IsFullyRegenerated())
+	{
+		// remove all layers
+		ResultTileIndices = RemoveTileLayers(TileX, TileY, &OldLayerTileIdMap);
+	}
+
+	dtNavMesh* DetourMesh = DestNavMesh->GetRecastNavMeshImpl()->GetRecastMesh();
+	const int32 FirstDirtyTileIndex = TileGenerator.GetDirtyLayersMask().Find(true);
+
+	if (DetourMesh != nullptr && IsInActiveSet(FIntPoint(TileX, TileY))
+		&& FirstDirtyTileIndex != INDEX_NONE)
+	{
+		TArray<FNavMeshTileData> TileLayers = TileGenerator.GetNavigationData();
+		ResultTileIndices.Reserve(TileLayers.Num());
+
+		struct FLayerIndexFinder
+		{
+			int32 LayerIndex;
+			explicit FLayerIndexFinder(const int32 InLayerIndex) : LayerIndex(InLayerIndex) {}
+			bool operator()(const FNavMeshTileData& LayerData) const
+			{
+				return LayerData.LayerIndex == LayerIndex;
+			}
+		};
+
+		for (int32 LayerIndex = FirstDirtyTileIndex; LayerIndex < TileGenerator.GetDirtyLayersMask().Num(); ++LayerIndex)
+		{
+			if (TileGenerator.IsLayerChanged(LayerIndex))
+			{
+				dtTileRef OldTileRef = DetourMesh->getTileRefAt(TileX, TileY, LayerIndex);
+				const int32 LayerDataIndex = TileLayers.IndexOfByPredicate(FLayerIndexFinder(LayerIndex));
+
+				if (LayerDataIndex != INDEX_NONE)
+				{
+					FNavMeshTileData& LayerData = TileLayers[LayerDataIndex];
+					if (OldTileRef)
+					{
+						QUICK_SCOPE_CYCLE_COUNTER(STAT_NavMesh_RemoveOldTIle);
+
+						NumActiveTiles--;
+						UE_LOG(LogNavigation, Log, TEXT("%s> Tile (%d,%d:%d), removing TileRef: 0x%X (active:%d)"),
+							*DestNavMesh->GetName(), TileX, TileY, LayerIndex, OldTileRef, NumActiveTiles);
+
+						DetourMesh->removeTile(OldTileRef, nullptr, nullptr);
+
+						const uint32 TileId = RecastTileVersionHelper::GetUpdatedTileId(OldTileRef, DetourMesh);
+						ResultTileIndices.AddUnique(TileId);
+					}
+					else
+					{
+						OldTileRef = OldLayerTileIdMap.FindRef(LayerIndex);
+					}
+
+					if (LayerData.IsValid())
+					{
+						bool bRejectNavmesh = false;
+						dtTileRef ResultTileRef = 0;
+
+						dtStatus status = 0;
+
+						{
+							QUICK_SCOPE_CYCLE_COUNTER(STAT_NavMesh_AddTileToDetourMesh);
+							// let navmesh know it's tile generator who owns the data
+							status = DetourMesh->addTile(LayerData.GetData(), LayerData.DataSize, DT_TILE_FREE_DATA, OldTileRef, &ResultTileRef);
+
+							// if tile index was already taken by other layer try adding it on first free entry (salt was already updated by whatever took that spot)
+							if (dtStatusFailed(status) && dtStatusDetail(status, DT_OUT_OF_MEMORY) && OldTileRef)
+							{
+								OldTileRef = 0;
+								status = DetourMesh->addTile(LayerData.GetData(), LayerData.DataSize, DT_TILE_FREE_DATA, OldTileRef, &ResultTileRef);
+							}
+						}
+
+						if (dtStatusFailed(status))
+						{
+							if (dtStatusDetail(status, DT_OUT_OF_MEMORY))
+							{
+								UE_LOG(LogNavigation, Error, TEXT("%s> Tile (%d,%d:%d), tile limit reached!! (%d)"),
+									*DestNavMesh->GetName(), TileX, TileY, LayerIndex, DetourMesh->getMaxTiles());
+							}
+						}
+						else
+						{
+							ResultTileIndices.AddUnique(DetourMesh->decodePolyIdTile(ResultTileRef));
+							NumActiveTiles++;
+
+							UE_LOG(LogNavigation, Log, TEXT("%s> Tile (%d,%d:%d), added TileRef: 0x%X (active:%d)"),
+								*DestNavMesh->GetName(), TileX, TileY, LayerIndex, ResultTileRef, NumActiveTiles);
+
+							{
+								QUICK_SCOPE_CYCLE_COUNTER(STAT_NavMesh_ReleaseTileLayers);
+								// NavMesh took the ownership of generated data, so we don't need to deallocate it
+								uint8* ReleasedData = LayerData.Release();
+							}
+						}
+					}
+				}
+				else
+				{
+					// remove the layer since it ended up empty
+					DetourMesh->removeTile(OldTileRef, nullptr, nullptr);
+					const uint32 TileId = RecastTileVersionHelper::GetUpdatedTileId(OldTileRef, DetourMesh);
+					ResultTileIndices.AddUnique(TileId);
+				}
+			}
+		}
+	}
+
+	return ResultTileIndices;
+}
+
+void FRecastNavMeshGenerator::DiscardCurrentBuildingTasks()
+{
+	PendingDirtyTiles.Empty();
+	
+	for (FRunningTileElement& Element : RunningDirtyTiles)
+	{
+		if (Element.AsyncTask)
+		{
+			Element.AsyncTask->EnsureCompletion();
+			delete Element.AsyncTask;
+			Element.AsyncTask = nullptr;
+		}
+	}
+	
+	RunningDirtyTiles.Empty();
+}
+
+bool FRecastNavMeshGenerator::HasDirtyTiles() const
+{
+	return (PendingDirtyTiles.Num() > 0 || RunningDirtyTiles.Num() > 0);
+}
+
+FBox FRecastNavMeshGenerator::GrowBoundingBox(const FBox& BBox, bool bIncludeAgentHeight) const
+{
+	const FVector BBoxGrowOffsetBoth = FVector(2.0f * Config.borderSize * Config.cs);
+	const FVector BBoxGrowOffsetMin = FVector(0, 0, bIncludeAgentHeight ? Config.AgentHeight : 0.0f);
+
+	return FBox(BBox.Min - BBoxGrowOffsetBoth - BBoxGrowOffsetMin, BBox.Max + BBoxGrowOffsetBoth);
+}
+
+static bool IntercestBounds(const FBox& TestBox, const TNavStatArray<FBox>& Bounds)
+{
+	for (const FBox& Box : Bounds)
+	{
+		if (Box.Intersect(TestBox))
+		{
+			return true;
+		}
+	}
+
+	return false;
+}
+
+namespace 
+{
+	FBox CalculateBoxIntercetion(const FBox& BoxA, const FBox& BoxB)
+	{
+		// assumes boxes overlap
+		ensure(BoxA.Intersect(BoxB));
+		return FBox(FVector(FMath::Max(BoxA.Min.X, BoxB.Min.X)
+							, FMath::Max(BoxA.Min.Y, BoxB.Min.Y)
+							, FMath::Max(BoxA.Min.Z, BoxB.Min.Z))
+					, FVector(FMath::Min(BoxA.Max.X, BoxB.Max.X)
+							, FMath::Min(BoxA.Max.Y, BoxB.Max.Y)
+							, FMath::Min(BoxA.Max.Z, BoxB.Max.Z))
+					);
+	}
+}
+
+void FRecastNavMeshGenerator::MarkDirtyTiles(const TArray<FNavigationDirtyArea>& DirtyAreas)
+{
+	QUICK_SCOPE_CYCLE_COUNTER(STAT_RecastNavMeshGenerator_MarkDirtyTiles);
+	
+	check(bInitialized);
+	const float TileSizeInWorldUnits = Config.tileSize * Config.cs;
+	check(TileSizeInWorldUnits > 0);
+	const FVector NavMeshOrigin = FVector::ZeroVector;
+	
+	const bool bGameStaticNavMesh = IsGameStaticNavMesh(DestNavMesh);
+		
+	// find all tiles that need regeneration
+	TSet<FPendingTileElement> DirtyTiles;
+	for (const FNavigationDirtyArea& DirtyArea : DirtyAreas)
+	{
+		// Static navmeshes accept only area modifiers updates
+		if (bGameStaticNavMesh && (!DirtyArea.HasFlag(ENavigationDirtyFlag::DynamicModifier) || DirtyArea.HasFlag(ENavigationDirtyFlag::NavigationBounds)))
+		{
+			continue;
+		}
+		
+		bool bDoTileInclusionTest = false;
+		FBox AdjustedAreaBounds = DirtyArea.Bounds;
+		
+		// if it's not expanding the navigatble area
+		if (DirtyArea.HasFlag(ENavigationDirtyFlag::NavigationBounds) == false)
+		{
+			// and is outside of current bounds
+			if (GetTotalBounds().Intersect(DirtyArea.Bounds) == false)
+			{
+				// skip it
+				continue;
+			}
+
+			const FBox CutDownArea = CalculateBoxIntercetion(GetTotalBounds(), DirtyArea.Bounds);
+			AdjustedAreaBounds = GrowBoundingBox(CutDownArea, DirtyArea.HasFlag(ENavigationDirtyFlag::UseAgentHeight));
+
+			// @todo this and the following test share some work in common
+			if (IntercestBounds(AdjustedAreaBounds, InclusionBounds) == false)
+			{
+				continue;
+			}
+
+			// check if any of inclusion volumes encapsulates this box
+			// using CutDownArea not AdjustedAreaBounds since if the area is on the border of navigable space
+			// then FindInclusionBoundEncapsulatingBox can produce false negative
+			bDoTileInclusionTest = FindInclusionBoundEncapsulatingBox(CutDownArea) == INDEX_NONE;
+		}
+				
+		const FBox RcAreaBounds = Unreal2RecastBox(AdjustedAreaBounds);
+		const int32 XMin = FMath::FloorToInt((RcAreaBounds.Min.X - NavMeshOrigin.X) / TileSizeInWorldUnits);
+		const int32 XMax = FMath::FloorToInt((RcAreaBounds.Max.X - NavMeshOrigin.X) / TileSizeInWorldUnits);
+		const int32 YMin = FMath::FloorToInt((RcAreaBounds.Min.Z - NavMeshOrigin.Z) / TileSizeInWorldUnits);
+		const int32 YMax = FMath::FloorToInt((RcAreaBounds.Max.Z - NavMeshOrigin.Z) / TileSizeInWorldUnits);
+
+		for (int32 TileY = YMin; TileY <= YMax; ++TileY)
+		{
+			for (int32 TileX = XMin; TileX <= XMax; ++TileX)
+			{
+				if (IsInActiveSet(FIntPoint(TileX, TileY)) == false)
+				{
+					continue;
+				}
+
+				if (DirtyArea.HasFlag(ENavigationDirtyFlag::NavigationBounds) == false && bDoTileInclusionTest == true)
+				{
+					const FBox TileBox = CalculateTileBounds(TileX, TileY, NavMeshOrigin, TotalNavBounds, TileSizeInWorldUnits);
+
+					// do per tile check since we can have lots of tiles inbetween navigable bounds volumes
+					if (IntercestBounds(TileBox, InclusionBounds) == false)
+					{
+						// Skip this tile
+						continue;
+					}
+				}
+												
+				FPendingTileElement Element;
+				Element.Coord = FIntPoint(TileX, TileY);
+				Element.bRebuildGeometry = DirtyArea.HasFlag(ENavigationDirtyFlag::Geometry) || DirtyArea.HasFlag(ENavigationDirtyFlag::NavigationBounds);
+				if (Element.bRebuildGeometry == false)
+				{
+					Element.DirtyAreas.Add(AdjustedAreaBounds);
+				}
+				
+				FPendingTileElement* ExistingElement = DirtyTiles.Find(Element);
+				if (ExistingElement)
+				{
+					ExistingElement->bRebuildGeometry|= Element.bRebuildGeometry;
+					// Append area bounds to existing list 
+					if (ExistingElement->bRebuildGeometry == false)
+					{
+						ExistingElement->DirtyAreas.Append(Element.DirtyAreas);
+					}
+					else
+					{
+						ExistingElement->DirtyAreas.Empty();
+					}
+				}
+				else
+				{
+					DirtyTiles.Add(Element);
+				}
+			}
+		}
+	}
+	
+	int32 NumTilesMarked = DirtyTiles.Num();
+
+	// Merge all pending tiles into one container
+	for (const FPendingTileElement& Element : PendingDirtyTiles)
+	{
+		FPendingTileElement* ExistingElement = DirtyTiles.Find(Element);
+		if (ExistingElement)
+		{
+			ExistingElement->bRebuildGeometry|= Element.bRebuildGeometry;
+			// Append area bounds to existing list 
+			if (ExistingElement->bRebuildGeometry == false)
+			{
+				ExistingElement->DirtyAreas.Append(Element.DirtyAreas);
+			}
+			else
+			{
+				ExistingElement->DirtyAreas.Empty();
+			}
+		}
+		else
+		{
+			DirtyTiles.Add(Element);
+		}
+	}
+	
+	// Dump results into array
+	PendingDirtyTiles.Empty(DirtyTiles.Num());
+	for(const FPendingTileElement& Element : DirtyTiles)
+	{
+		PendingDirtyTiles.Add(Element);
+	}
+
+	// Sort tiles by proximity to players 
+	if (NumTilesMarked > 0)
+	{
+		SortPendingBuildTiles();
+	}
+}
+
+void FRecastNavMeshGenerator::SortPendingBuildTiles()
+{
+	TArray<FVector2D> SeedLocations;
+	UWorld* CurWorld = GetWorld();
+	if (CurWorld == nullptr)
+	{
+		return;
+	}
+
+	// Collect players positions
+	for (FConstPlayerControllerIterator PlayerIt = CurWorld->GetPlayerControllerIterator(); PlayerIt; ++PlayerIt)
+	{
+		if (PlayerIt->IsValid() && PlayerIt->Get()->GetPawn() != NULL)
+		{
+			const FVector2D SeedLoc((*PlayerIt)->GetPawn()->GetActorLocation());
+			SeedLocations.Add(SeedLoc);
+		}
+	}
+
+	if (SeedLocations.Num() == 0)
+	{
+		// Use navmesh origin for sorting
+		SeedLocations.Add(FVector2D::ZeroVector);
+	}
+
+	if (SeedLocations.Num() > 0)
+	{
+		const float TileSizeInWorldUnits = Config.tileSize * Config.cs;
+		
+		// Calculate shortest distances between tiles and players
+		for (FPendingTileElement& Element : PendingDirtyTiles)
+		{
+			const FBox TileBox = CalculateTileBounds(Element.Coord.X, Element.Coord.Y, FVector::ZeroVector, TotalNavBounds, TileSizeInWorldUnits);
+			FVector2D TileCenter2D = FVector2D(TileBox.GetCenter());
+			for (FVector2D SeedLocation : SeedLocations)
+			{
+				const float DistSq = FVector2D::DistSquared(TileCenter2D, SeedLocation);
+				if (DistSq < Element.SeedDistance)
+				{
+					Element.SeedDistance = DistSq;
+				}
+			}
+		}
+
+		// nearest tiles should be at the end of the list
+		PendingDirtyTiles.Sort();
+	}
+}
+
+TSharedRef<FRecastTileGenerator> FRecastNavMeshGenerator::CreateTileGenerator(const FIntPoint& Coord, const TArray<FBox>& DirtyAreas)
+{
+	TSharedRef<FRecastTileGenerator> TileGenerator = MakeShareable(new FRecastTileGenerator(*this, Coord));
+	TileGenerator->Setup(*this, DirtyAreas);
+	return TileGenerator;
+}
+
+TArray<uint32> FRecastNavMeshGenerator::ProcessTileTasks(const int32 NumTasksToSubmit)
+{
+	QUICK_SCOPE_CYCLE_COUNTER(STAT_RecastNavMeshGenerator_ProcessTileTasks);
+	
+	TArray<uint32> UpdatedTiles;
+	const bool bHasTasksAtStart = GetNumRemaningBuildTasks() > 0;
+	const bool bGameStaticNavMesh = IsGameStaticNavMesh(DestNavMesh);
+			
+	int32 NumSubmittedTasks = 0;
+	// Submit pending tile elements
+	for (int32 ElementIdx = PendingDirtyTiles.Num()-1; ElementIdx >= 0 && NumSubmittedTasks < NumTasksToSubmit; ElementIdx--)
+	{
+		QUICK_SCOPE_CYCLE_COUNTER(STAT_RecastNavMeshGenerator_ProcessTileTasks_NewTasks);
+
+		FPendingTileElement& PendingElement = PendingDirtyTiles[ElementIdx];
+		FRunningTileElement RunningElement(PendingElement.Coord);
+		
+		// Make sure that we are not submitting generator for grid cell that is currently being regenerated
+		if (!RunningDirtyTiles.Contains(RunningElement))
+		{
+			// Spawn async task
+#if RECAST_ASYNC_REBUILDING
+			TUniquePtr<FRecastTileGeneratorTask> TileTask = MakeUnique<FRecastTileGeneratorTask>(CreateTileGenerator(PendingElement.Coord, PendingElement.DirtyAreas));
+
+			// Start it in background in case it has something to build
+			if (TileTask->GetTask().TileGenerator->HasDataToBuild())
+			{
+				RunningElement.AsyncTask = TileTask.Release();
+				RunningElement.AsyncTask->StartBackgroundTask();
+			
+				RunningDirtyTiles.Add(RunningElement);
+				NumSubmittedTasks++;
+			}
+#else
+			TSharedRef<FRecastTileGenerator> TileGenerator = CreateTileGenerator(PendingElement.Coord, PendingElement.DirtyAreas);
+			if (TileGenerator->HasDataToBuild())
+			{
+				FRecastTileGenerator& TileGeneratorRef = TileGenerator.Get();
+				TileGeneratorRef.DoWork();
+
+				TArray<uint32> UpdatedTileIndices = AddGeneratedTiles(TileGeneratorRef);
+				UpdatedTiles.Append(UpdatedTileIndices);
+
+				// Store compressed tile cache layers so it can be reused later
+				if (TileGeneratorRef.GetCompressedLayers().Num())
+				{
+					QUICK_SCOPE_CYCLE_COUNTER(STAT_RecastNavMeshGenerator_StoringCompressedLayers);
+					DestNavMesh->AddTileCacheLayers(PendingElement.Coord.X, PendingElement.Coord.Y, TileGeneratorRef.GetCompressedLayers());
+				}
+				else
+				{
+					DestNavMesh->MarkEmptyTileCacheLayers(PendingElement.Coord.X, PendingElement.Coord.Y);
+				}
+
+				NumSubmittedTasks++;
+			}
+#endif
+			else if (!bGameStaticNavMesh)
+			{
+				// If there is nothing to generate remove all tiles from navmesh at specified grid coordinates
+				UpdatedTiles.Append(
+					RemoveTileLayers(PendingElement.Coord.X, PendingElement.Coord.Y)
+				);
+				DestNavMesh->MarkEmptyTileCacheLayers(PendingElement.Coord.X, PendingElement.Coord.Y);
+					
+				// TODO: should we increment NumSubmittedTasks here? 
+				// We can count removing as a tasks to avoid hitches when there are large number of pending tiles to remove
+			}
+			
+			// Remove submitted element from pending list
+			PendingDirtyTiles.RemoveAt(ElementIdx);
+
+			// Release memory, list could be quite big after map load
+			if (PendingDirtyTiles.Num() == 0)
+			{
+				PendingDirtyTiles.Empty(32);
+			}
+		}
+	}
+	
+	// Collect completed tasks and apply generated data to navmesh
+	for (int32 Idx = RunningDirtyTiles.Num() - 1; Idx >=0; --Idx)
+	{
+		QUICK_SCOPE_CYCLE_COUNTER(STAT_RecastNavMeshGenerator_ProcessTileTasks_FinishedTasks);
+
+		FRunningTileElement& Element = RunningDirtyTiles[Idx];
+		check(Element.AsyncTask);
+
+		if (Element.AsyncTask->IsDone())
+		{
+			// Add generated tiles to navmesh
+			if (!Element.bShouldDiscard)
+			{
+				FRecastTileGenerator& TileGenerator = *(Element.AsyncTask->GetTask().TileGenerator);
+				TArray<uint32> UpdatedTileIndices = AddGeneratedTiles(TileGenerator);
+				UpdatedTiles.Append(UpdatedTileIndices);
+			
+				// Store compressed tile cache layers so it can be reused later
+				if (TileGenerator.GetCompressedLayers().Num())
+				{
+					QUICK_SCOPE_CYCLE_COUNTER(STAT_RecastNavMeshGenerator_StoringCompressedLayers);
+					DestNavMesh->AddTileCacheLayers(Element.Coord.X, Element.Coord.Y, TileGenerator.GetCompressedLayers());
+				}
+				else
+				{
+					DestNavMesh->MarkEmptyTileCacheLayers(Element.Coord.X, Element.Coord.Y);
+				}
+			}
+
+			{
+				QUICK_SCOPE_CYCLE_COUNTER(STAT_RecastNavMeshGenerator_TileGeneratorRemoval);
+
+				// Destroy tile generator task
+				delete Element.AsyncTask;
+				Element.AsyncTask = nullptr;
+				// Remove completed tile element from a list of running tasks
+				RunningDirtyTiles.RemoveAtSwap(Idx, 1, false);
+			}
+		}
+	}
+
+	// Notify owner in case all tasks has been completed
+	const bool bHasTasksAtEnd = GetNumRemaningBuildTasks() > 0;
+	if (bHasTasksAtStart && !bHasTasksAtEnd)
+	{
+		DestNavMesh->OnNavMeshGenerationFinished();
+	}
+	
+	return UpdatedTiles;
+}
+
+void FRecastNavMeshGenerator::ExportComponentGeometry(UActorComponent* Component, FNavigationRelevantData& Data)
+{
+	FRecastGeometryExport GeomExport(Data);
+	RecastGeometryExport::ExportComponent(Component, GeomExport);
+	RecastGeometryExport::CovertCoordDataToRecast(GeomExport.VertexBuffer);
+	RecastGeometryExport::StoreCollisionCache(GeomExport);
+}
+
+void FRecastNavMeshGenerator::ExportVertexSoupGeometry(const TArray<FVector>& Verts, FNavigationRelevantData& Data)
+{
+	FRecastGeometryExport GeomExport(Data);
+	RecastGeometryExport::ExportVertexSoup(Verts, GeomExport.VertexBuffer, GeomExport.IndexBuffer, GeomExport.Data->Bounds);
+	RecastGeometryExport::StoreCollisionCache(GeomExport);
+}
+
+void FRecastNavMeshGenerator::ExportRigidBodyGeometry(UBodySetup& BodySetup, TNavStatArray<FVector>& OutVertexBuffer, TNavStatArray<int32>& OutIndexBuffer, const FTransform& LocalToWorld)
+{
+	TNavStatArray<float> VertCoords;
+	FBox TempBounds;
+
+	RecastGeometryExport::ExportRigidBodySetup(BodySetup, VertCoords, OutIndexBuffer, TempBounds, LocalToWorld);
+
+	OutVertexBuffer.Reserve(OutVertexBuffer.Num() + (VertCoords.Num() / 3));
+	for (int32 i = 0; i < VertCoords.Num(); i += 3)
+	{
+		OutVertexBuffer.Add(FVector(VertCoords[i + 0], VertCoords[i + 1], VertCoords[i + 2]));
+	}
+}
+
+void FRecastNavMeshGenerator::ExportRigidBodyGeometry(UBodySetup& BodySetup, TNavStatArray<FVector>& OutTriMeshVertexBuffer, TNavStatArray<int32>& OutTriMeshIndexBuffer
+	, TNavStatArray<FVector>& OutConvexVertexBuffer, TNavStatArray<int32>& OutConvexIndexBuffer, TNavStatArray<int32>& OutShapeBuffer
+	, const FTransform& LocalToWorld)
+{
+	BodySetup.CreatePhysicsMeshes();
+
+	TNavStatArray<float> VertCoords;
+	FBox TempBounds;
+
+	VertCoords.Reset();
+	RecastGeometryExport::ExportRigidBodyTriMesh(BodySetup, VertCoords, OutTriMeshIndexBuffer, TempBounds, LocalToWorld);
+
+	OutTriMeshVertexBuffer.Reserve(OutTriMeshVertexBuffer.Num() + (VertCoords.Num() / 3));
+	for (int32 i = 0; i < VertCoords.Num(); i += 3)
+	{
+		OutTriMeshVertexBuffer.Add(FVector(VertCoords[i + 0], VertCoords[i + 1], VertCoords[i + 2]));
+	}
+
+	const int32 NumExistingVerts = OutConvexVertexBuffer.Num();
+	VertCoords.Reset();
+	RecastGeometryExport::ExportRigidBodyConvexElements(BodySetup, VertCoords, OutConvexIndexBuffer, OutShapeBuffer, TempBounds, LocalToWorld);
+	RecastGeometryExport::ExportRigidBodyBoxElements(BodySetup.AggGeom, VertCoords, OutConvexIndexBuffer, OutShapeBuffer, TempBounds, LocalToWorld, NumExistingVerts);
+	RecastGeometryExport::ExportRigidBodySphylElements(BodySetup.AggGeom, VertCoords, OutConvexIndexBuffer, OutShapeBuffer, TempBounds, LocalToWorld, NumExistingVerts);
+	RecastGeometryExport::ExportRigidBodySphereElements(BodySetup.AggGeom, VertCoords, OutConvexIndexBuffer, OutShapeBuffer, TempBounds, LocalToWorld, NumExistingVerts);
+	
+	OutConvexVertexBuffer.Reserve(OutConvexVertexBuffer.Num() + (VertCoords.Num() / 3));
+	for (int32 i = 0; i < VertCoords.Num(); i += 3)
+	{
+		OutConvexVertexBuffer.Add(FVector(VertCoords[i + 0], VertCoords[i + 1], VertCoords[i + 2]));
+	}
+}
+
+void FRecastNavMeshGenerator::ExportAggregatedGeometry(const FKAggregateGeom& AggGeom, TNavStatArray<FVector>& OutConvexVertexBuffer, TNavStatArray<int32>& OutConvexIndexBuffer, TNavStatArray<int32>& OutShapeBuffer, const FTransform& LocalToWorld)
+{
+	TNavStatArray<float> VertCoords;
+	FBox TempBounds;
+
+	const int32 NumExistingVerts = OutConvexVertexBuffer.Num();
+
+	// convex and tri mesh are NOT supported, since they require BodySetup.CreatePhysicsMeshes() call
+	// only simple shapes
+
+	RecastGeometryExport::ExportRigidBodyBoxElements(AggGeom, VertCoords, OutConvexIndexBuffer, OutShapeBuffer, TempBounds, LocalToWorld, NumExistingVerts);
+	RecastGeometryExport::ExportRigidBodySphylElements(AggGeom, VertCoords, OutConvexIndexBuffer, OutShapeBuffer, TempBounds, LocalToWorld, NumExistingVerts);
+	RecastGeometryExport::ExportRigidBodySphereElements(AggGeom, VertCoords, OutConvexIndexBuffer, OutShapeBuffer, TempBounds, LocalToWorld, NumExistingVerts);
+
+	OutConvexVertexBuffer.Reserve(OutConvexVertexBuffer.Num() + (VertCoords.Num() / 3));
+	for (int32 i = 0; i < VertCoords.Num(); i += 3)
+	{
+		OutConvexVertexBuffer.Add(FVector(VertCoords[i + 0], VertCoords[i + 1], VertCoords[i + 2]));
+	}
+}
+
+bool FRecastNavMeshGenerator::IsBuildInProgress(bool bCheckDirtyToo) const
+{
+	if (RunningDirtyTiles.Num() || (bCheckDirtyToo && PendingDirtyTiles.Num()))
+	{
+		return true;
+	}
+
+	return false;
+}
+
+int32 FRecastNavMeshGenerator::GetNumRemaningBuildTasks() const
+{
+	return RunningDirtyTiles.Num() + PendingDirtyTiles.Num();
+}
+
+int32 FRecastNavMeshGenerator::GetNumRunningBuildTasks() const
+{
+	return RunningDirtyTiles.Num();
+}
+
+bool FRecastNavMeshGenerator::GatherGeometryOnGameThread() const 
+{ 
+	return DestNavMesh == nullptr || DestNavMesh->ShouldGatherDataOnGameThread() == true;
+}
+
+bool FRecastNavMeshGenerator::IsTileChanged(int32 TileIdx) const
+{
+#if WITH_EDITOR	
+	// Check recently built tiles
+	if (TileIdx > 0)
+	{
+		FTileTimestamp TileTimestamp;
+		TileTimestamp.TileIdx = static_cast<uint32>(TileIdx);
+		if (RecentlyBuiltTiles.Contains(TileTimestamp))
+		{
+			return true;
+		}
+	}
+#endif//WITH_EDITOR
+
+	return false;
+}
+
+uint32 FRecastNavMeshGenerator::LogMemUsed() const 
+{
+	UE_LOG(LogNavigation, Display, TEXT("    FRecastNavMeshGenerator: self %d"), sizeof(FRecastNavMeshGenerator));
+	
+	uint32 GeneratorsMem = 0;
+	for (const FRunningTileElement& Element : RunningDirtyTiles)
+	{
+		GeneratorsMem += Element.AsyncTask->GetTask().TileGenerator->UsedMemoryOnStartup;
+	}
+
+	UE_LOG(LogNavigation, Display, TEXT("    FRecastNavMeshGenerator: Total Generator\'s size %u, count %d"), GeneratorsMem, RunningDirtyTiles.Num());
+
+	return GeneratorsMem + sizeof(FRecastNavMeshGenerator) + PendingDirtyTiles.GetAllocatedSize() + PendingDirtyTiles.GetAllocatedSize();
+}
+
+#if !(UE_BUILD_SHIPPING || UE_BUILD_TEST) && ENABLE_VISUAL_LOG
+void FRecastNavMeshGenerator::GrabDebugSnapshot(struct FVisualLogEntry* Snapshot, const FBox& BoundingBox, const FName& CategoryName, ELogVerbosity::Type LogVerbosity) const
+{
+	const UNavigationSystemV1* NavSys = FNavigationSystem::GetCurrent<UNavigationSystemV1>(GetWorld());
+	const FNavigationOctree* NavOctree = NavSys ? NavSys->GetNavOctree() : NULL;
+	if (Snapshot == nullptr)
+	{
+		return;
+	}
+
+	if (NavOctree == NULL)
+	{
+		UE_LOG(LogNavigation, Error, TEXT("Failed to vlog navigation data due to %s being NULL"), NavSys == NULL ? TEXT("NavigationSystem") : TEXT("NavOctree"));
+		return;
+	}
+
+	ELogVerbosity::Type NavAreaVerbosity = FMath::Clamp(ELogVerbosity::Type(LogVerbosity + 1), ELogVerbosity::NoLogging, ELogVerbosity::VeryVerbose);
+
+	for (int32 Index = 0; Index < NavSys->NavDataSet.Num(); ++Index)
+	{
+		TArray<FVector> CoordBuffer;
+		TArray<int32> Indices;
+		TNavStatArray<FVector> Faces;
+		const ARecastNavMesh* NavData = Cast<const ARecastNavMesh>(NavSys->NavDataSet[Index]);
+		if (NavData)
+		{
+			for (FNavigationOctree::TConstElementBoxIterator<FNavigationOctree::DefaultStackAllocator> It(*NavOctree, BoundingBox);
+				It.HasPendingElements();
+				It.Advance())
+			{
+				const FNavigationOctreeElement& Element = It.GetCurrentElement();
+				const bool bExportGeometry = Element.Data->HasGeometry() && Element.ShouldUseGeometry(DestNavMesh->GetConfig());
+
+				if (bExportGeometry && Element.Data->CollisionData.Num())
+				{
+					FRecastGeometryCache CachedGeometry(Element.Data->CollisionData.GetData());
+					
+					const uint32 NumVerts = CachedGeometry.Header.NumVerts;
+					CoordBuffer.Reset(NumVerts);
+					for (uint32 VertIdx = 0; VertIdx < NumVerts * 3; VertIdx += 3)
+					{
+						CoordBuffer.Add(Recast2UnrealPoint(&CachedGeometry.Verts[VertIdx]));
+					}
+
+					const uint32 NumIndices = CachedGeometry.Header.NumFaces * 3;
+					Indices.SetNum(NumIndices, false);
+					for (uint32 IndicesIdx = 0; IndicesIdx < NumIndices; ++IndicesIdx)
+					{
+						Indices[IndicesIdx] = CachedGeometry.Indices[IndicesIdx];
+					}
+
+					Snapshot->AddElement(CoordBuffer, Indices, CategoryName, LogVerbosity, FColorList::LightGrey.WithAlpha(255));
+				}
+				else
+				{
+					TArray<FVector> Verts;
+					const TArray<FAreaNavModifier>& AreaMods = Element.Data->Modifiers.GetAreas();
+					for (int32 i = 0; i < AreaMods.Num(); i++)
+					{
+						if (AreaMods[i].GetShapeType() == ENavigationShapeType::Unknown)
+						{
+							continue;
+						}
+
+						const uint8 AreaId = NavData->GetAreaID(AreaMods[i].GetAreaClass());
+						const UClass* AreaClass = NavData->GetAreaClass(AreaId);
+						const UNavArea* DefArea = AreaClass ? ((UClass*)AreaClass)->GetDefaultObject<UNavArea>() : NULL;
+						const FColor PolygonColor = AreaClass != UNavigationSystemV1::GetDefaultWalkableArea() ? (DefArea ? DefArea->DrawColor : NavData->GetConfig().Color) : FColorList::Cyan;
+
+						if (AreaMods[i].GetShapeType() == ENavigationShapeType::Box)
+						{
+							FBoxNavAreaData Box;
+							AreaMods[i].GetBox(Box);
+
+							Snapshot->AddElement(FBox::BuildAABB(Box.Origin, Box.Extent), FMatrix::Identity, CategoryName, NavAreaVerbosity, PolygonColor.WithAlpha(255));
+						}
+						else if (AreaMods[i].GetShapeType() == ENavigationShapeType::Cylinder)
+						{
+							FCylinderNavAreaData Cylinder;
+							AreaMods[i].GetCylinder(Cylinder);
+
+							Snapshot->AddElement(Cylinder.Origin, Cylinder.Origin + FVector(0, 0, Cylinder.Height), Cylinder.Radius, CategoryName, NavAreaVerbosity, PolygonColor.WithAlpha(255));
+						}
+						else
+						{
+							FConvexNavAreaData Convex;
+							AreaMods[i].GetConvex(Convex);
+							Verts.Reset();
+							GrowConvexHull(NavData->AgentRadius, Convex.Points, Verts);
+
+							Snapshot->AddElement(
+								Verts,
+								Convex.MinZ - NavData->CellHeight, 
+								Convex.MaxZ + NavData->CellHeight, 
+								CategoryName, NavAreaVerbosity, PolygonColor.WithAlpha(255));
+						}
+					}
+				}
+			}
+
+		}
+
+	}
+}
+#endif
+
+#if !(UE_BUILD_SHIPPING || UE_BUILD_TEST) && ENABLE_VISUAL_LOG
+void FRecastNavMeshGenerator::ExportNavigationData(const FString& FileName) const
+{
+	const UNavigationSystemV1* NavSys = FNavigationSystem::GetCurrent<UNavigationSystemV1>(GetWorld());
+	const FNavigationOctree* NavOctree = NavSys ? NavSys->GetNavOctree() : NULL;
+	if (NavOctree == NULL)
+	{
+		UE_LOG(LogNavigation, Error, TEXT("Failed to export navigation data due to %s being NULL"), NavSys == NULL ? TEXT("NavigationSystem") : TEXT("NavOctree"));
+		return;
+	}
+
+	const double StartExportTime = FPlatformTime::Seconds();
+
+	FString CurrentTimeStr = FDateTime::Now().ToString();
+	for (int32 Index = 0; Index < NavSys->NavDataSet.Num(); ++Index)
+	{
+		// feed data from octtree and mark for rebuild				
+		TNavStatArray<float> CoordBuffer;
+		TNavStatArray<int32> IndexBuffer;
+		const ARecastNavMesh* NavData = Cast<const ARecastNavMesh>(NavSys->NavDataSet[Index]);
+		if (NavData)
+		{
+			struct FAreaExportData
+			{
+				FConvexNavAreaData Convex;
+				uint8 AreaId;
+			};
+			TArray<FAreaExportData> AreaExport;
+
+			for(FNavigationOctree::TConstElementBoxIterator<FNavigationOctree::DefaultStackAllocator> It(*NavOctree, TotalNavBounds);
+				It.HasPendingElements();
+				It.Advance())
+			{
+				const FNavigationOctreeElement& Element = It.GetCurrentElement();
+				const bool bExportGeometry = Element.Data->HasGeometry() && Element.ShouldUseGeometry(DestNavMesh->GetConfig());
+
+				if (bExportGeometry && Element.Data->CollisionData.Num())
+				{
+					FRecastGeometryCache CachedGeometry(Element.Data->CollisionData.GetData());
+					IndexBuffer.Reserve( IndexBuffer.Num() + (CachedGeometry.Header.NumFaces * 3 ));
+					CoordBuffer.Reserve( CoordBuffer.Num() + (CachedGeometry.Header.NumVerts * 3 ));
+					for (int32 i = 0; i < CachedGeometry.Header.NumFaces * 3; i++)
+					{
+						IndexBuffer.Add(CachedGeometry.Indices[i] + CoordBuffer.Num() / 3);
+					}
+					for (int32 i = 0; i < CachedGeometry.Header.NumVerts * 3; i++)
+					{
+						CoordBuffer.Add(CachedGeometry.Verts[i]);
+					}
+				}
+				else
+				{
+					const TArray<FAreaNavModifier>& AreaMods = Element.Data->Modifiers.GetAreas();
+					for (int32 i = 0; i < AreaMods.Num(); i++)
+					{
+						FAreaExportData ExportInfo;
+						ExportInfo.AreaId = NavData->GetAreaID(AreaMods[i].GetAreaClass());
+
+						if (AreaMods[i].GetShapeType() == ENavigationShapeType::Convex)
+						{
+							AreaMods[i].GetConvex(ExportInfo.Convex);
+
+							TArray<FVector> ConvexVerts;
+							GrowConvexHull(NavData->AgentRadius, ExportInfo.Convex.Points, ConvexVerts);
+							ExportInfo.Convex.MinZ -= NavData->CellHeight;
+							ExportInfo.Convex.MaxZ += NavData->CellHeight;
+							ExportInfo.Convex.Points = ConvexVerts;
+
+							AreaExport.Add(ExportInfo);
+						}
+					}
+				}
+			}
+			
+			UWorld* NavigationWorld = GetWorld();
+			for (int32 LevelIndex = 0; LevelIndex < NavigationWorld->GetNumLevels(); ++LevelIndex) 
+			{
+				const ULevel* const Level =  NavigationWorld->GetLevel(LevelIndex);
+				if (Level == NULL)
+				{
+					continue;
+				}
+
+				const TArray<FVector>* LevelGeom = Level->GetStaticNavigableGeometry();
+				if (LevelGeom != NULL && LevelGeom->Num() > 0)
+				{
+					TNavStatArray<FVector> Verts;
+					TNavStatArray<int32> Faces;
+					// For every ULevel in World take its pre-generated static geometry vertex soup
+					RecastGeometryExport::TransformVertexSoupToRecast(*LevelGeom, Verts, Faces);
+
+					IndexBuffer.Reserve( IndexBuffer.Num() + Faces.Num() );
+					CoordBuffer.Reserve( CoordBuffer.Num() + Verts.Num() * 3);
+					for (int32 i = 0; i < Faces.Num(); i++)
+					{
+						IndexBuffer.Add(Faces[i] + CoordBuffer.Num() / 3);
+					}
+					for (int32 i = 0; i < Verts.Num(); i++)
+					{
+						CoordBuffer.Add(Verts[i].X);
+						CoordBuffer.Add(Verts[i].Y);
+						CoordBuffer.Add(Verts[i].Z);
+					}
+				}
+			}
+			
+			
+			FString AreaExportStr;
+			for (int32 i = 0; i < AreaExport.Num(); i++)
+			{
+				const FAreaExportData& ExportInfo = AreaExport[i];
+				AreaExportStr += FString::Printf(TEXT("\nAE %d %d %f %f\n"),
+					ExportInfo.AreaId, ExportInfo.Convex.Points.Num(), ExportInfo.Convex.MinZ, ExportInfo.Convex.MaxZ);
+
+				for (int32 iv = 0; iv < ExportInfo.Convex.Points.Num(); iv++)
+				{
+					FVector Pt = Unreal2RecastPoint(ExportInfo.Convex.Points[iv]);
+					AreaExportStr += FString::Printf(TEXT("Av %f %f %f\n"), Pt.X, Pt.Y, Pt.Z);
+				}
+			}
+			
+			FString AdditionalData;
+			
+			if (AreaExport.Num())
+			{
+				AdditionalData += "# Area export\n";
+				AdditionalData += AreaExportStr;
+				AdditionalData += "\n";
+			}
+
+			AdditionalData += "# RecastDemo specific data\n";
+	#if 0
+			// use this bounds to have accurate navigation data bounds
+			const FVector Center = Unreal2RecastPoint(NavData->GetBounds().GetCenter());
+			FVector Extent = FVector(NavData->GetBounds().GetExtent());
+			Extent = FVector(Extent.X, Extent.Z, Extent.Y);
+	#else
+			// this bounds match navigation bounds from level
+			FBox RCNavBounds = Unreal2RecastBox(TotalNavBounds);
+			const FVector Center = RCNavBounds.GetCenter();
+			const FVector Extent = RCNavBounds.GetExtent();
+	#endif
+			const FBox Box = FBox::BuildAABB(Center, Extent);
+			AdditionalData += FString::Printf(
+				TEXT("rd_bbox %7.7f %7.7f %7.7f %7.7f %7.7f %7.7f\n"), 
+				Box.Min.X, Box.Min.Y, Box.Min.Z, 
+				Box.Max.X, Box.Max.Y, Box.Max.Z
+			);
+			
+			const FRecastNavMeshGenerator* CurrentGen = static_cast<const FRecastNavMeshGenerator*>(NavData->GetGenerator());
+			check(CurrentGen);
+			AdditionalData += FString::Printf(TEXT("# AgentHeight\n"));
+			AdditionalData += FString::Printf(TEXT("rd_agh %5.5f\n"), CurrentGen->Config.AgentHeight);
+			AdditionalData += FString::Printf(TEXT("# AgentRadius\n"));
+			AdditionalData += FString::Printf(TEXT("rd_agr %5.5f\n"), CurrentGen->Config.AgentRadius);
+
+			AdditionalData += FString::Printf(TEXT("# Cell Size\n"));
+			AdditionalData += FString::Printf(TEXT("rd_cs %5.5f\n"), CurrentGen->Config.cs);
+			AdditionalData += FString::Printf(TEXT("# Cell Height\n"));
+			AdditionalData += FString::Printf(TEXT("rd_ch %5.5f\n"), CurrentGen->Config.ch);
+
+			AdditionalData += FString::Printf(TEXT("# Agent max climb\n"));
+			AdditionalData += FString::Printf(TEXT("rd_amc %d\n"), (int)CurrentGen->Config.AgentMaxClimb);
+			AdditionalData += FString::Printf(TEXT("# Agent max slope\n"));
+			AdditionalData += FString::Printf(TEXT("rd_ams %5.5f\n"), CurrentGen->Config.walkableSlopeAngle);
+
+			AdditionalData += FString::Printf(TEXT("# Region min size\n"));
+			AdditionalData += FString::Printf(TEXT("rd_rmis %d\n"), (uint32)FMath::Sqrt(CurrentGen->Config.minRegionArea));
+			AdditionalData += FString::Printf(TEXT("# Region merge size\n"));
+			AdditionalData += FString::Printf(TEXT("rd_rmas %d\n"), (uint32)FMath::Sqrt(CurrentGen->Config.mergeRegionArea));
+
+			AdditionalData += FString::Printf(TEXT("# Max edge len\n"));
+			AdditionalData += FString::Printf(TEXT("rd_mel %d\n"), CurrentGen->Config.maxEdgeLen);
+
+			AdditionalData += FString::Printf(TEXT("# Perform Voxel Filtering\n"));
+			AdditionalData += FString::Printf(TEXT("rd_pvf %d\n"), CurrentGen->Config.bPerformVoxelFiltering);
+			AdditionalData += FString::Printf(TEXT("# Generate Detailed Mesh\n"));
+			AdditionalData += FString::Printf(TEXT("rd_gdm %d\n"), CurrentGen->Config.bGenerateDetailedMesh);
+			AdditionalData += FString::Printf(TEXT("# MaxPolysPerTile\n"));
+			AdditionalData += FString::Printf(TEXT("rd_mppt %d\n"), CurrentGen->Config.MaxPolysPerTile);
+			AdditionalData += FString::Printf(TEXT("# maxVertsPerPoly\n"));
+			AdditionalData += FString::Printf(TEXT("rd_mvpp %d\n"), CurrentGen->Config.maxVertsPerPoly);
+			AdditionalData += FString::Printf(TEXT("# Tile size\n"));
+			AdditionalData += FString::Printf(TEXT("rd_ts %d\n"), CurrentGen->Config.tileSize);
+
+			AdditionalData += FString::Printf(TEXT("\n"));
+			
+			const FString FilePathName = FileName + FString::Printf(TEXT("_NavDataSet%d_%s.obj"), Index, *CurrentTimeStr) ;
+			ExportGeomToOBJFile(FilePathName, CoordBuffer, IndexBuffer, AdditionalData);
+		}
+	}
+	UE_LOG(LogNavigation, Log, TEXT("ExportNavigation time: %.3f sec ."), FPlatformTime::Seconds() - StartExportTime);
+}
+#endif
+
+static class FNavigationGeomExec : private FSelfRegisteringExec
+{
+public:
+	/** Console commands, see embeded usage statement **/
+	virtual bool Exec( UWorld* InWorld, const TCHAR* Cmd, FOutputDevice& Ar ) override
+	{
+#if ALLOW_DEBUG_FILES && !(UE_BUILD_SHIPPING || UE_BUILD_TEST)
+		bool bCorrectCmd = FParse::Command(&Cmd, TEXT("ExportNavigation"));
+		if (bCorrectCmd && !InWorld)
+		{
+			UE_LOG(LogNavigation, Error, TEXT("Failed to export navigation data due to missing UWorld"));
+		}
+		else if (InWorld && bCorrectCmd)
+		{
+			UNavigationSystemV1* NavSys = FNavigationSystem::GetCurrent<UNavigationSystemV1>(InWorld);
+			if (NavSys)
+			{
+				const ANavigationData* NavData = NavSys->GetDefaultNavDataInstance();
+				if (NavData)
+				{
+					if (const FNavDataGenerator* Generator = NavData->GetGenerator())
+					{
+						const FString Name = NavData->GetName();
+						Generator->ExportNavigationData( FString::Printf( TEXT("%s/%s"), *FPaths::ProjectSavedDir(), *Name ));
+						return true;
+					}
+					else
+					{
+						UE_LOG(LogNavigation, Error, TEXT("Failed to export navigation data due to missing generator"));
+					}
+				}
+				else
+				{
+					UE_LOG(LogNavigation, Error, TEXT("Failed to export navigation data due to navigation data"));
+				}
+			}
+			else
+			{
+				UE_LOG(LogNavigation, Error, TEXT("Failed to export navigation data due to missing navigation system"));
+			}
+		}
+#endif // ALLOW_DEBUG_FILES && WITH_EDITOR
+		return false;
+	}
+} NavigationGeomExec;
+
+#endif // WITH_RECAST
