Copying //UE4/Dev-Framework to //UE4/Dev-Main (Source: //UE4/Dev-Framework @ 3967517)

diff --git a/Engine/Source/Runtime/NavigationSystem/Private/NavModifierComponent.cpp b/Engine/Source/Runtime/NavigationSystem/Private/NavModifierComponent.cpp
new file mode 100644
--- /dev/null
+++ b/Engine/Source/Runtime/NavigationSystem/Private/NavModifierComponent.cpp
@@ -0,0 +1,121 @@
+// Copyright 1998-2018 Epic Games, Inc. All Rights Reserved.
+
+#include "NavModifierComponent.h"
+#include "GameFramework/Actor.h"
+#include "Components/PrimitiveComponent.h"
+#include "AI/NavigationModifier.h"
+#include "NavAreas/NavArea_Null.h"
+#include "PhysicsEngine/BodySetup.h"
+
+UNavModifierComponent::UNavModifierComponent(const FObjectInitializer& ObjectInitializer) : Super(ObjectInitializer)
+{
+	AreaClass = UNavArea_Null::StaticClass();
+	FailsafeExtent = FVector(100, 100, 100);
+	bIncludeAgentHeight = true;
+}
+
+void UNavModifierComponent::CalcAndCacheBounds() const
+{
+	const AActor* MyOwner = GetOwner();
+	if (MyOwner)
+	{
+		TInlineComponentArray<UPrimitiveComponent*> PrimComponents;
+		MyOwner->GetComponents(PrimComponents);
+
+		Bounds = FBox(ForceInit);
+		ComponentBounds.Reset();
+		for (int32 Idx = 0; Idx < PrimComponents.Num(); Idx++)
+		{
+			UPrimitiveComponent* PrimComp = PrimComponents[Idx];
+			if (PrimComp && PrimComp->IsRegistered() && PrimComp->IsCollisionEnabled())
+			{
+				UBodySetup* BodySetup = PrimComp->GetBodySetup();
+				if (BodySetup)
+				{
+					FTransform ParentTM = PrimComp->GetComponentTransform();
+					const FVector Scale3D = ParentTM.GetScale3D();
+					ParentTM.RemoveScaling();
+					Bounds += PrimComp->Bounds.GetBox();
+
+					for (int32 SphereIdx = 0; SphereIdx < BodySetup->AggGeom.SphereElems.Num(); SphereIdx++)
+					{
+						const FKSphereElem& ElemInfo = BodySetup->AggGeom.SphereElems[SphereIdx];
+						FTransform ElemTM = ElemInfo.GetTransform();
+						ElemTM.ScaleTranslation(Scale3D);
+						ElemTM *= ParentTM;
+
+						const FBox SphereBounds = FBox::BuildAABB(ElemTM.GetLocation(), ElemInfo.Radius * Scale3D);
+						ComponentBounds.Add(FRotatedBox(SphereBounds, ElemTM.GetRotation()));
+					}
+
+					for (int32 BoxIdx = 0; BoxIdx < BodySetup->AggGeom.BoxElems.Num(); BoxIdx++)
+					{
+						const FKBoxElem& ElemInfo = BodySetup->AggGeom.BoxElems[BoxIdx];
+						FTransform ElemTM = ElemInfo.GetTransform();
+						ElemTM.ScaleTranslation(Scale3D);
+						ElemTM *= ParentTM;
+
+						const FBox BoxBounds = FBox::BuildAABB(ElemTM.GetLocation(), FVector(ElemInfo.X, ElemInfo.Y, ElemInfo.Z) * Scale3D * 0.5f);
+						ComponentBounds.Add(FRotatedBox(BoxBounds, ElemTM.GetRotation()));
+					}
+
+					for (int32 SphylIdx = 0; SphylIdx < BodySetup->AggGeom.SphylElems.Num(); SphylIdx++)
+					{
+						const FKSphylElem& ElemInfo = BodySetup->AggGeom.SphylElems[SphylIdx];
+						FTransform ElemTM = ElemInfo.GetTransform();
+						ElemTM.ScaleTranslation(Scale3D);
+						ElemTM *= ParentTM;
+
+						const FBox SphylBounds = FBox::BuildAABB(ElemTM.GetLocation(), FVector(ElemInfo.Radius, ElemInfo.Radius, ElemInfo.Length) * Scale3D);
+						ComponentBounds.Add(FRotatedBox(SphylBounds, ElemTM.GetRotation()));
+					}
+
+					for (int32 ConvexIdx = 0; ConvexIdx < BodySetup->AggGeom.ConvexElems.Num(); ConvexIdx++)
+					{
+						const FKConvexElem& ElemInfo = BodySetup->AggGeom.ConvexElems[ConvexIdx];
+						FTransform ElemTM = ElemInfo.GetTransform();
+						ElemTM.ScaleTranslation(Scale3D);
+						ElemTM *= FTransform(ParentTM.GetLocation());
+
+						const FBox ConvexBounds = ElemInfo.CalcAABB(ElemTM, Scale3D);
+						ComponentBounds.Add(FRotatedBox(ConvexBounds, ElemTM.GetRotation() * ParentTM.GetRotation()));
+					}
+				}
+			}
+		}
+
+		if (ComponentBounds.Num() == 0)
+		{
+			Bounds = FBox::BuildAABB(MyOwner->GetActorLocation(), FailsafeExtent);
+			ComponentBounds.Add(FRotatedBox(Bounds, MyOwner->GetActorQuat()));
+		}
+		else
+		{
+			for (int32 Idx = 0; Idx < ComponentBounds.Num(); Idx++)
+			{
+				const FVector BoxOrigin = ComponentBounds[Idx].Box.GetCenter();
+				const FVector BoxExtent = ComponentBounds[Idx].Box.GetExtent();
+				
+				const FVector NavModBoxOrigin = FTransform(ComponentBounds[Idx].Quat).InverseTransformPosition(BoxOrigin);
+				ComponentBounds[Idx].Box = FBox::BuildAABB(NavModBoxOrigin, BoxExtent);
+			}
+		}
+	}
+}
+
+void UNavModifierComponent::GetNavigationData(FNavigationRelevantData& Data) const
+{
+	for (int32 Idx = 0; Idx < ComponentBounds.Num(); Idx++)
+	{
+		Data.Modifiers.Add(FAreaNavModifier(ComponentBounds[Idx].Box, FTransform(ComponentBounds[Idx].Quat), AreaClass).SetIncludeAgentHeight(bIncludeAgentHeight));
+	}
+}
+
+void UNavModifierComponent::SetAreaClass(TSubclassOf<UNavArea> NewAreaClass)
+{
+	if (AreaClass != NewAreaClass)
+	{
+		AreaClass = NewAreaClass;
+		RefreshNavigationModifiers();
+	}
+}
diff --git a/Engine/Source/Runtime/NavigationSystem/Private/NavigationInterfaces.cpp b/Engine/Source/Runtime/NavigationSystem/Private/NavigationInterfaces.cpp
new file mode 100644
--- /dev/null
+++ b/Engine/Source/Runtime/NavigationSystem/Private/NavigationInterfaces.cpp
@@ -0,0 +1,10 @@
+// Copyright 1998-2019 Epic Games, Inc. All Rights Reserved.
+
+#include "NavNodeInterface.h"
+
+
+UNavNodeInterface::UNavNodeInterface(const FObjectInitializer& ObjectInitializer)
+	: Super(ObjectInitializer)
+{
+}
+
diff --git a/Engine/Source/Runtime/NavigationSystem/Private/NavigationOctree.cpp b/Engine/Source/Runtime/NavigationSystem/Private/NavigationOctree.cpp
new file mode 100644
--- /dev/null
+++ b/Engine/Source/Runtime/NavigationSystem/Private/NavigationOctree.cpp
@@ -0,0 +1,227 @@
+// Copyright 1998-2018 Epic Games, Inc. All Rights Reserved.
+
+#include "NavigationOctree.h"
+#include "AI/Navigation/NavRelevantInterface.h"
+#include "NavigationSystem.h"
+
+//----------------------------------------------------------------------//
+// FNavigationOctree
+//----------------------------------------------------------------------//
+FNavigationOctree::FNavigationOctree(const FVector& Origin, float Radius)
+	: TOctree<FNavigationOctreeElement, FNavigationOctreeSemantics>(Origin, Radius)
+	, DefaultGeometryGatheringMode(ENavDataGatheringMode::Instant)
+	, bGatherGeometry(false)
+	, NodesMemory(0)
+{
+	INC_DWORD_STAT_BY( STAT_NavigationMemory, sizeof(*this) );
+}
+
+FNavigationOctree::~FNavigationOctree()
+{
+	DEC_DWORD_STAT_BY( STAT_NavigationMemory, sizeof(*this) );
+	DEC_MEMORY_STAT_BY(STAT_Navigation_CollisionTreeMemory, NodesMemory);
+}
+
+void FNavigationOctree::SetDataGatheringMode(ENavDataGatheringModeConfig Mode)
+{
+	check(Mode != ENavDataGatheringModeConfig::Invalid);
+	DefaultGeometryGatheringMode = ENavDataGatheringMode(Mode);
+}
+
+void FNavigationOctree::SetNavigableGeometryStoringMode(ENavGeometryStoringMode NavGeometryMode)
+{
+	bGatherGeometry = (NavGeometryMode == FNavigationOctree::StoreNavGeometry);
+}
+
+void FNavigationOctree::DemandLazyDataGathering(FNavigationRelevantData& ElementData)
+{
+	UObject* ElementOb = ElementData.GetOwner();
+	if (ElementOb == nullptr)
+	{
+		return;
+	}
+
+	bool bShrink = false;
+	const int32 OrgElementMemory = ElementData.GetGeometryAllocatedSize();
+
+	if (ElementData.IsPendingLazyGeometryGathering() == true && ElementData.SupportsGatheringGeometrySlices() == false)
+	{
+		QUICK_SCOPE_CYCLE_COUNTER(STAT_RecastNavMeshGenerator_LazyGeometryExport);
+		UActorComponent& ActorComp = *CastChecked<UActorComponent>(ElementOb);
+		ComponentExportDelegate.ExecuteIfBound(&ActorComp, ElementData);
+
+		bShrink = true;
+
+		// mark this element as no longer needing geometry gathering
+		ElementData.bPendingLazyGeometryGathering = false;
+	}
+
+	if (ElementData.IsPendingLazyModifiersGathering())
+	{
+		QUICK_SCOPE_CYCLE_COUNTER(STAT_RecastNavMeshGenerator_LazyModifiersExport);
+		INavRelevantInterface* NavElement = Cast<INavRelevantInterface>(ElementOb);
+		check(NavElement);
+		NavElement->GetNavigationData(ElementData);
+		ElementData.bPendingLazyModifiersGathering = false;
+		bShrink = true;
+	}
+
+	if (bShrink)
+	{
+		// validate exported data
+		// shrink arrays before counting memory
+		// it will be reallocated when adding to octree and RemoveNode will have different value returned by GetAllocatedSize()
+		ElementData.ValidateAndShrink();
+	}
+
+	const int32 ElementMemoryChange = ElementData.GetGeometryAllocatedSize() - OrgElementMemory;
+	const_cast<FNavigationOctree*>(this)->NodesMemory += ElementMemoryChange;
+	INC_MEMORY_STAT_BY(STAT_Navigation_CollisionTreeMemory, ElementMemoryChange);
+}
+
+void FNavigationOctree::DemandLazyDataGathering(const FNavigationOctreeElement& Element)
+{
+	FNavigationRelevantData& MutableData = const_cast<FNavigationRelevantData&>(*Element.Data);
+	DemandLazyDataGathering(MutableData);
+}
+
+void FNavigationOctree::AddNode(UObject* ElementOb, INavRelevantInterface* NavElement, const FBox& Bounds, FNavigationOctreeElement& Element)
+{
+	// we assume NavElement is ElementOb already cast
+	Element.Bounds = Bounds;	
+
+	if (NavElement)
+	{
+		const ENavDataGatheringMode GatheringMode = NavElement->GetGeometryGatheringMode();
+		bool bDoInstantGathering = ((GatheringMode == ENavDataGatheringMode::Default && DefaultGeometryGatheringMode == ENavDataGatheringMode::Instant) 
+			|| GatheringMode == ENavDataGatheringMode::Instant);
+
+		if (bGatherGeometry)
+		{
+			UActorComponent* ActorComp = Cast<UActorComponent>(ElementOb);
+			if (ActorComp)
+			{
+				if (bDoInstantGathering)
+				{
+					ComponentExportDelegate.ExecuteIfBound(ActorComp, *Element.Data);
+				}
+				else
+				{
+					Element.Data->bPendingLazyGeometryGathering = true;
+					Element.Data->bSupportsGatheringGeometrySlices = NavElement && NavElement->SupportsGatheringGeometrySlices();
+				}
+			}
+		}
+
+		SCOPE_CYCLE_COUNTER(STAT_Navigation_GatheringNavigationModifiersSync);
+		if (bDoInstantGathering)
+		{
+			NavElement->GetNavigationData(*Element.Data);
+		}
+		else
+		{
+			Element.Data->bPendingLazyModifiersGathering = true;
+		}
+	}
+
+	// validate exported data
+	// shrink arrays before counting memory
+	// it will be reallocated when adding to octree and RemoveNode will have different value returned by GetAllocatedSize()
+	Element.ValidateAndShrink();
+
+	const int32 ElementMemory = Element.GetAllocatedSize();
+	NodesMemory += ElementMemory;
+	INC_MEMORY_STAT_BY(STAT_Navigation_CollisionTreeMemory, ElementMemory);
+
+	AddElement(Element);
+}
+
+void FNavigationOctree::AppendToNode(const FOctreeElementId& Id, INavRelevantInterface* NavElement, const FBox& Bounds, FNavigationOctreeElement& Element)
+{
+	FNavigationOctreeElement OrgData = GetElementById(Id);
+
+	Element = OrgData;
+	Element.Bounds = Bounds + OrgData.Bounds.GetBox();
+
+	if (NavElement)
+	{
+		SCOPE_CYCLE_COUNTER(STAT_Navigation_GatheringNavigationModifiersSync);
+		NavElement->GetNavigationData(*Element.Data);
+	}
+
+	// validate exported data
+	// shrink arrays before counting memory
+	// it will be reallocated when adding to octree and RemoveNode will have different value returned by GetAllocatedSize()
+	Element.ValidateAndShrink();
+
+	const int32 OrgElementMemory = OrgData.GetAllocatedSize();
+	const int32 NewElementMemory = Element.GetAllocatedSize();
+	const int32 MemoryDelta = NewElementMemory - OrgElementMemory;
+
+	NodesMemory += MemoryDelta;
+	INC_MEMORY_STAT_BY(STAT_Navigation_CollisionTreeMemory, MemoryDelta);
+
+	RemoveElement(Id);
+	AddElement(Element);
+}
+
+void FNavigationOctree::UpdateNode(const FOctreeElementId& Id, const FBox& NewBounds)
+{
+	FNavigationOctreeElement ElementCopy = GetElementById(Id);
+	RemoveElement(Id);
+	ElementCopy.Bounds = NewBounds;
+	AddElement(ElementCopy);
+}
+
+void FNavigationOctree::RemoveNode(const FOctreeElementId& Id)
+{
+	FNavigationOctreeElement& Element = GetElementById(Id);
+	const int32 ElementMemory = Element.GetAllocatedSize();
+	NodesMemory -= ElementMemory;
+	DEC_MEMORY_STAT_BY(STAT_Navigation_CollisionTreeMemory, ElementMemory);
+
+	RemoveElement(Id);
+}
+
+const FNavigationRelevantData* FNavigationOctree::GetDataForID(const FOctreeElementId& Id) const
+{
+	if (Id.IsValidId() == false)
+	{
+		return nullptr;
+	}
+
+	const FNavigationOctreeElement& OctreeElement = GetElementById(Id);
+
+	return &*OctreeElement.Data;
+}
+
+//----------------------------------------------------------------------//
+// FNavigationOctreeSemantics
+//----------------------------------------------------------------------//
+#if NAVSYS_DEBUG
+FORCENOINLINE
+#endif // NAVSYS_DEBUG
+void FNavigationOctreeSemantics::SetElementId(const FNavigationOctreeElement& Element, FOctreeElementId Id)
+{
+	UWorld* World = NULL;
+	UObject* ElementOwner = Element.GetOwner();
+
+	if (AActor* Actor = Cast<AActor>(ElementOwner))
+	{
+		World = Actor->GetWorld();
+	}
+	else if (UActorComponent* AC = Cast<UActorComponent>(ElementOwner))
+	{
+		World = AC->GetWorld();
+	}
+	else if (ULevel* Level = Cast<ULevel>(ElementOwner))
+	{
+		World = Level->OwningWorld;
+	}
+
+	UNavigationSystemV1* NavSys = FNavigationSystem::GetCurrent<UNavigationSystemV1>(World);
+	if (NavSys)
+	{
+		NavSys->SetObjectsNavOctreeId(ElementOwner, Id);
+	}
+}
