Copying //UE4/Dev-Framework to //UE4/Dev-Main (Source: //UE4/Dev-Framework @ 3967517)

diff --git a/Engine/Source/Runtime/NavigationSystem/Private/NavigationSystem.cpp b/Engine/Source/Runtime/NavigationSystem/Private/NavigationSystem.cpp
new file mode 100644
--- /dev/null
+++ b/Engine/Source/Runtime/NavigationSystem/Private/NavigationSystem.cpp
@@ -0,0 +1,4267 @@
+// Copyright 1998-2018 Epic Games, Inc. All Rights Reserved.
+
+#include "NavigationSystem.h"
+#include "Misc/ScopeLock.h"
+#include "Stats/StatsMisc.h"
+#include "Modules/ModuleManager.h"
+#include "AI/Navigation/NavAgentInterface.h"
+#include "Engine/World.h"
+#include "GameFramework/Controller.h"
+#include "AI/Navigation/NavRelevantInterface.h"
+#include "UObject/UObjectIterator.h"
+#include "EngineUtils.h"
+#include "Logging/MessageLog.h"
+#include "NavAreas/NavArea.h"
+#include "NavigationOctree.h"
+#include "VisualLogger/VisualLogger.h"
+#include "NavMesh/NavMeshBoundsVolume.h"
+#include "NavigationInvokerComponent.h"
+#include "AI/Navigation/NavigationDataChunk.h"
+#include "Engine/Engine.h"
+#include "UObject/Package.h"
+#include "Components/PrimitiveComponent.h"
+
+#if WITH_RECAST
+#include "NavMesh/RecastNavMesh.h"
+#include "NavMesh/RecastHelpers.h"
+#include "NavMesh/RecastNavMeshGenerator.h"
+#endif // WITH_RECAST
+#if WITH_EDITOR
+#include "EditorModeManager.h"
+#include "EditorModes.h"
+#include "Editor/GeometryMode/Public/GeometryEdMode.h"
+#include "Editor/GeometryMode/Public/EditorGeometry.h"
+#endif
+
+#if WITH_HOT_RELOAD
+#include "Misc/HotReloadInterface.h"
+#endif
+
+#include "NavAreas/NavArea_Null.h"
+#include "NavAreas/NavArea_Default.h"
+#include "NavAreas/NavAreaMeta_SwitchByAgent.h"
+#include "NavLinkCustomInterface.h"
+#include "NavigationPath.h"
+#include "AbstractNavData.h"
+#include "CrowdManagerBase.h"
+
+
+static const uint32 INITIAL_ASYNC_QUERIES_SIZE = 32;
+static const uint32 REGISTRATION_QUEUE_SIZE = 16;	// and we'll not reallocate
+
+#define LOCTEXT_NAMESPACE "Navigation"
+
+DEFINE_LOG_CATEGORY_STATIC(LogNavOctree, Warning, All);
+
+DECLARE_CYCLE_STAT(TEXT("Rasterize triangles"), STAT_Navigation_RasterizeTriangles,STATGROUP_Navigation);
+DECLARE_CYCLE_STAT(TEXT("Nav Tick: area register"), STAT_Navigation_TickNavAreaRegister, STATGROUP_Navigation);
+DECLARE_CYCLE_STAT(TEXT("Nav Tick: mark dirty"), STAT_Navigation_TickMarkDirty, STATGROUP_Navigation);
+DECLARE_CYCLE_STAT(TEXT("Nav Tick: async build"), STAT_Navigation_TickAsyncBuild, STATGROUP_Navigation);
+DECLARE_CYCLE_STAT(TEXT("Nav Tick: async pathfinding"), STAT_Navigation_TickAsyncPathfinding, STATGROUP_Navigation);
+DECLARE_CYCLE_STAT(TEXT("Debug NavOctree Time"), STAT_DebugNavOctree, STATGROUP_Navigation);
+//----------------------------------------------------------------------//
+// Stats
+//----------------------------------------------------------------------//
+
+DEFINE_STAT(STAT_Navigation_QueriesTimeSync);
+DEFINE_STAT(STAT_Navigation_RequestingAsyncPathfinding);
+DEFINE_STAT(STAT_Navigation_PathfindingSync);
+DEFINE_STAT(STAT_Navigation_PathfindingAsync);
+DEFINE_STAT(STAT_Navigation_AddGeneratedTiles);
+DEFINE_STAT(STAT_Navigation_TileNavAreaSorting);
+DEFINE_STAT(STAT_Navigation_TileGeometryExportToObjAsync);
+DEFINE_STAT(STAT_Navigation_TileVoxelFilteringAsync);
+DEFINE_STAT(STAT_Navigation_TileBuildAsync);
+DEFINE_STAT(STAT_Navigation_TileBuildPreparationSync);
+DEFINE_STAT(STAT_Navigation_BSPExportSync);
+DEFINE_STAT(STAT_Navigation_GatheringNavigationModifiersSync);
+DEFINE_STAT(STAT_Navigation_ActorsGeometryExportSync);
+DEFINE_STAT(STAT_Navigation_ProcessingActorsForNavMeshBuilding);
+DEFINE_STAT(STAT_Navigation_AdjustingNavLinks);
+DEFINE_STAT(STAT_Navigation_AddingActorsToNavOctree);
+DEFINE_STAT(STAT_Navigation_RecastTick);
+DEFINE_STAT(STAT_Navigation_RecastPathfinding);
+DEFINE_STAT(STAT_Navigation_RecastBuildCompressedLayers);
+DEFINE_STAT(STAT_Navigation_RecastBuildNavigation);
+DEFINE_STAT(STAT_Navigation_UpdateNavOctree);
+DEFINE_STAT(STAT_Navigation_CollisionTreeMemory);
+DEFINE_STAT(STAT_Navigation_NavDataMemory);
+DEFINE_STAT(STAT_Navigation_TileCacheMemory);
+DEFINE_STAT(STAT_Navigation_OutOfNodesPath);
+DEFINE_STAT(STAT_Navigation_PartialPath);
+DEFINE_STAT(STAT_Navigation_CumulativeBuildTime);
+DEFINE_STAT(STAT_Navigation_BuildTime);
+DEFINE_STAT(STAT_Navigation_OffsetFromCorners);
+DEFINE_STAT(STAT_Navigation_PathVisibilityOptimisation);
+DEFINE_STAT(STAT_Navigation_ObservedPathsCount);
+DEFINE_STAT(STAT_Navigation_RecastMemory);
+
+//----------------------------------------------------------------------//
+// consts
+//----------------------------------------------------------------------//
+namespace FNavigationSystem
+{
+	FORCEINLINE bool IsValidExtent(const FVector& Extent)
+	{
+		return Extent != INVALID_NAVEXTENT;
+	}
+
+	FCustomLinkOwnerInfo::FCustomLinkOwnerInfo(INavLinkCustomInterface* Link)
+	{
+		LinkInterface = Link;
+		LinkOwner = Link->GetLinkOwner();
+	}
+
+	bool ShouldLoadNavigationOnClient(ANavigationData& NavData)
+	{
+		const UNavigationSystemV1* NavSysCDO = (*GEngine->NavigationSystemClass != nullptr) 
+			? (GEngine->NavigationSystemClass->GetDefaultObject<const UNavigationSystemV1>()) 
+			: (const UNavigationSystemV1*)nullptr;
+		return NavSysCDO && NavSysCDO->ShouldLoadNavigationOnClient(&NavData);
+	}
+}
+
+namespace NavigationDebugDrawing
+{
+	const float PathLineThickness = 3.f;
+	const FVector PathOffset(0,0,15);
+	const FVector PathNodeBoxExtent(16.f);
+}
+
+//----------------------------------------------------------------------//
+// FNavigationInvoker
+//----------------------------------------------------------------------//
+FNavigationInvoker::FNavigationInvoker()
+: Actor(nullptr)
+, GenerationRadius(0)
+, RemovalRadius(0)
+{
+
+}
+
+FNavigationInvoker::FNavigationInvoker(AActor& InActor, float InGenerationRadius, float InRemovalRadius)
+: Actor(&InActor)
+, GenerationRadius(InGenerationRadius)
+, RemovalRadius(InRemovalRadius)
+{
+
+}
+
+//----------------------------------------------------------------------//
+// UNavigationSystemV1                                                                
+//----------------------------------------------------------------------//
+bool UNavigationSystemV1::bNavigationAutoUpdateEnabled = true;
+TMap<INavLinkCustomInterface*, FWeakObjectPtr> UNavigationSystemV1::PendingCustomLinkRegistration;
+FCriticalSection UNavigationSystemV1::CustomLinkRegistrationSection;
+TSubclassOf<UNavArea> UNavigationSystemV1::DefaultWalkableArea = NULL;
+TSubclassOf<UNavArea> UNavigationSystemV1::DefaultObstacleArea = NULL;
+#if !(UE_BUILD_SHIPPING || UE_BUILD_TEST)
+FNavigationSystemExec UNavigationSystemV1::ExecHandler;
+#endif // !UE_BUILD_SHIPPING
+
+/** called after navigation influencing event takes place*/
+UNavigationSystemV1::FOnNavigationDirty UNavigationSystemV1::NavigationDirtyEvent;
+
+bool UNavigationSystemV1::bUpdateNavOctreeOnComponentChange = true;
+bool UNavigationSystemV1::bStaticRuntimeNavigation = false;
+//----------------------------------------------------------------------//
+// life cycle stuff                                                                
+//----------------------------------------------------------------------//
+
+UNavigationSystemV1::UNavigationSystemV1(const FObjectInitializer& ObjectInitializer)
+	: Super(ObjectInitializer)
+	, bTickWhilePaused(false)
+	, bWholeWorldNavigable(false)
+	, bSkipAgentHeightCheckWhenPickingNavData(false)
+	, DirtyAreasUpdateFreq(60)
+	, OperationMode(FNavigationSystemRunMode::InvalidMode)
+	, NavOctree(NULL)
+	, NavBuildingLockFlags(0)
+	, InitialNavBuildingLockFlags(0)
+	, bNavOctreeLock(false)
+	, bInitialSetupHasBeenPerformed(false)
+	, bInitialLevelsAdded(false)
+	, bWorldInitDone(false)
+	, CurrentlyDrawnNavDataIndex(0)
+	, DirtyAreasUpdateTime(0)
+{
+#if WITH_EDITOR
+	NavUpdateLockFlags = 0;
+#endif
+	struct FDelegatesInitializer
+	{
+		FDelegatesInitializer()
+		{
+			UNavigationSystemBase::UpdateActorDataDelegate().BindStatic(&UNavigationSystemV1::UpdateActorInNavOctree);
+			UNavigationSystemBase::UpdateComponentDataDelegate().BindStatic(&UNavigationSystemV1::UpdateComponentInNavOctree);
+			UNavigationSystemBase::UpdateComponentDataAfterMoveDelegate().BindLambda([](USceneComponent& Comp) { UNavigationSystemV1::UpdateNavOctreeAfterMove(&Comp); });
+			UNavigationSystemBase::OnActorBoundsChangedDelegate().BindLambda([](AActor& Actor) { UNavigationSystemV1::UpdateNavOctreeBounds(&Actor); });
+			UNavigationSystemBase::OnPostEditActorMoveDelegate().BindLambda([](AActor& Actor) {
+				// update actor and all its components in navigation system after finishing move
+				// USceneComponent::UpdateNavigationData works only in game world
+				UNavigationSystemV1::UpdateNavOctreeBounds(&Actor);
+
+				TArray<AActor*> ParentedActors;
+				Actor.GetAttachedActors(ParentedActors);
+				for (int32 Idx = 0; Idx < ParentedActors.Num(); Idx++)
+				{
+					UNavigationSystemV1::UpdateNavOctreeBounds(ParentedActors[Idx]);
+				}
+
+				// not doing manual update of all attached actors since UpdateActorAndComponentsInNavOctree should take care of it
+				UNavigationSystemV1::UpdateActorAndComponentsInNavOctree(Actor);
+			});
+			UNavigationSystemBase::OnComponentTransformChangedDelegate().BindLambda([](USceneComponent& Comp) {
+				if (UNavigationSystemV1::ShouldUpdateNavOctreeOnComponentChange())
+				{
+					UWorld* World = Comp.GetWorld();
+					UNavigationSystemV1* NavSys = FNavigationSystem::GetCurrent<UNavigationSystemV1>(World);
+					if (NavSys != nullptr
+						&& (NavSys->ShouldAllowClientSideNavigation() || !World->IsNetMode(ENetMode::NM_Client)))
+					{
+						// use propagated component's transform update in editor OR server game with additional navsys check
+						UNavigationSystemV1::UpdateNavOctreeAfterMove(&Comp);
+					}
+				}
+			});
+			UNavigationSystemBase::OnActorRegisteredDelegate().BindLambda([](AActor& Actor) { UNavigationSystemV1::OnActorRegistered(&Actor); });
+			UNavigationSystemBase::OnActorUnregisteredDelegate().BindLambda([](AActor& Actor) { UNavigationSystemV1::OnActorUnregistered(&Actor); });
+			UNavigationSystemBase::OnComponentRegisteredDelegate().BindLambda([](UActorComponent& Comp) { UNavigationSystemV1::OnComponentRegistered(&Comp); });
+			UNavigationSystemBase::OnComponentUnregisteredDelegate().BindLambda([](UActorComponent& Comp) { UNavigationSystemV1::OnComponentUnregistered(&Comp); });
+			UNavigationSystemBase::RemoveActorDataDelegate().BindLambda([](AActor& Actor) { UNavigationSystemV1::ClearNavOctreeAll(&Actor); });
+			UNavigationSystemBase::HasComponentDataDelegate().BindLambda([](UActorComponent& Comp) {
+				UNavigationSystemV1* NavSys = FNavigationSystem::GetCurrent<UNavigationSystemV1>(Comp.GetWorld());
+				return (NavSys && (NavSys->GetObjectsNavOctreeId(&Comp) || NavSys->HasPendingObjectNavOctreeId(&Comp)));
+			});
+			UNavigationSystemBase::GetDefaultWalkableAreaDelegate().BindLambda([]() { return UNavigationSystemV1::GetDefaultWalkableArea(); });
+			UNavigationSystemBase::GetDefaultSupportedAgentDelegate().BindStatic(&UNavigationSystemV1::GetDefaultSupportedAgent);
+			UNavigationSystemBase::UpdateActorAndComponentDataDelegate().BindStatic(&UNavigationSystemV1::UpdateActorAndComponentsInNavOctree);
+			UNavigationSystemBase::OnComponentBoundsChangedDelegate().BindLambda([](UActorComponent& Comp, const FBox& NewBounds, const FBox& DirtyArea) {
+				UNavigationSystemV1* NavSys = FNavigationSystem::GetCurrent<UNavigationSystemV1>(Comp.GetWorld());
+				if (NavSys)
+				{
+					NavSys->UpdateNavOctreeElementBounds(&Comp, NewBounds, DirtyArea);
+				}
+			});
+			//UNavigationSystemBase::GetNavDataForPropsDelegate();
+			UNavigationSystemBase::GetNavDataForActorDelegate().BindStatic(&UNavigationSystemV1::GetNavDataForActor);
+
+#if WITH_RECAST
+			UNavigationSystemBase::GetDefaultNavDataClassDelegate().BindLambda([]() { return ARecastNavMesh::StaticClass(); });
+#endif // WITH_RECAST
+			UNavigationSystemBase::VerifyNavigationRenderingComponentsDelegate().BindLambda([](UWorld& World, const bool bShow) {
+				UNavigationSystemV1* NavSys = FNavigationSystem::GetCurrent<UNavigationSystemV1>(&World);
+				if (NavSys)
+				{
+					NavSys->VerifyNavigationRenderingComponents(bShow);
+				}
+			});
+			UNavigationSystemBase::BuildDelegate().BindLambda([](UWorld& World) {
+				UNavigationSystemV1* NavSys = FNavigationSystem::GetCurrent<UNavigationSystemV1>(&World);
+				if (NavSys)
+				{
+					NavSys->Build();
+				}
+			});
+#if WITH_EDITOR
+			UNavigationSystemBase::OnPIEStartDelegate().BindLambda([](UWorld& World) {
+				UNavigationSystemV1* NavSys = FNavigationSystem::GetCurrent<UNavigationSystemV1>(&World);
+				if (NavSys)
+				{
+					NavSys->OnPIEStart();
+				}
+			});
+			UNavigationSystemBase::OnPIEEndDelegate().BindLambda([](UWorld& World) {
+				UNavigationSystemV1* NavSys = FNavigationSystem::GetCurrent<UNavigationSystemV1>(&World);
+				if (NavSys)
+				{
+					NavSys->OnPIEEnd();
+				}
+			});
+			UNavigationSystemBase::UpdateLevelCollisionDelegate().BindLambda([](ULevel& Level) {
+				UNavigationSystemV1* NavSys = FNavigationSystem::GetCurrent<UNavigationSystemV1>(&Level);
+				if (NavSys)
+				{
+					NavSys->UpdateLevelCollision(&Level);
+				}
+			});
+			UNavigationSystemBase::SetNavigationAutoUpdateEnableDelegate().BindStatic(&UNavigationSystemV1::SetNavigationAutoUpdateEnabled);
+				/*.BindLambda([](const bool bNewEnable, UNavigationSystemBase* InNavigationSystem) {
+
+			})*/
+			UNavigationSystemBase::AddNavigationUpdateLockDelegate().BindLambda([](UWorld& World, uint8 Flags) {
+				UNavigationSystemV1* NavSys = FNavigationSystem::GetCurrent<UNavigationSystemV1>(&World);
+				if (NavSys)
+				{
+					NavSys->AddNavigationUpdateLock(Flags);
+				}
+			});
+			UNavigationSystemBase::RemoveNavigationUpdateLockDelegate().BindLambda([](UWorld& World, uint8 Flags) {
+				UNavigationSystemV1* NavSys = FNavigationSystem::GetCurrent<UNavigationSystemV1>(&World);
+				if (NavSys)
+				{
+					NavSys->RemoveNavigationUpdateLock(Flags);
+				}
+			});
+#endif // WITH_EDITOR
+		}
+	};
+	static FDelegatesInitializer DelegatesInitializer;
+	
+	// @hack, trying to load AIModule's CrowdManager
+	UClass* Class = StaticLoadClass(UCrowdManagerBase::StaticClass(), nullptr, TEXT("/Script/AIModule.CrowdManager"));
+	CrowdManagerClass = Class ? Class : UCrowdManagerBase::StaticClass();
+	
+	// active tiles
+	NextInvokersUpdateTime = 0.f;
+	ActiveTilesUpdateInterval = 1.f;
+	bGenerateNavigationOnlyAroundNavigationInvokers = false;
+	DataGatheringMode = ENavDataGatheringModeConfig::Instant;
+	bCanAccumulateDirtyAreas = true;
+#if !UE_BUILD_SHIPPING
+	bDirtyAreasReportedWhileAccumulationLocked = false;
+#endif // !UE_BUILD_SHIPPING
+
+	if (HasAnyFlags(RF_ClassDefaultObject) == false)
+	{
+		// reserve some arbitrary size
+		AsyncPathFindingQueries.Reserve( INITIAL_ASYNC_QUERIES_SIZE );
+		NavDataRegistrationQueue.Reserve( REGISTRATION_QUEUE_SIZE );
+	
+		FWorldDelegates::LevelAddedToWorld.AddUObject(this, &UNavigationSystemV1::OnLevelAddedToWorld);
+		FWorldDelegates::LevelRemovedFromWorld.AddUObject(this, &UNavigationSystemV1::OnLevelRemovedFromWorld);
+#if !UE_BUILD_SHIPPING
+		FCoreDelegates::OnGetOnScreenMessages.AddUObject(this, &UNavigationSystemV1::GetOnScreenMessages);
+#endif // !UE_BUILD_SHIPPING
+	}
+	else if (GetClass() == UNavigationSystemV1::StaticClass())
+	{
+		DefaultWalkableArea = UNavArea_Default::StaticClass();
+		DefaultObstacleArea = UNavArea_Null::StaticClass();
+		
+		const FTransform RecastToUnrealTransfrom(Recast2UnrealMatrix());
+		SetCoordTransformFrom(ENavigationCoordSystem::Recast, RecastToUnrealTransfrom);
+	}
+
+#if WITH_EDITOR
+	if (GIsEditor && HasAnyFlags(RF_ClassDefaultObject) == false)
+	{
+		FEditorDelegates::EditorModeEnter.AddUObject(this, &UNavigationSystemV1::OnEditorModeChanged, true);
+		FEditorDelegates::EditorModeExit.AddUObject(this, &UNavigationSystemV1::OnEditorModeChanged, false);
+	}
+#endif // WITH_EDITOR
+}
+
+UNavigationSystemV1::~UNavigationSystemV1()
+{
+	CleanUp(FNavigationSystem::ECleanupMode::CleanupUnsafe);
+
+#if WITH_EDITOR
+	if (GIsEditor)
+	{
+		FEditorDelegates::EditorModeEnter.RemoveAll(this);
+		FEditorDelegates::EditorModeExit.RemoveAll(this);
+	}
+#endif // WITH_EDITOR
+
+#if !UE_BUILD_SHIPPING
+	FCoreDelegates::OnGetOnScreenMessages.RemoveAll(this);
+#endif // !UE_BUILD_SHIPPING
+}
+
+void UNavigationSystemV1::ConfigureAsStatic()
+{
+	bStaticRuntimeNavigation = true;
+}
+
+void UNavigationSystemV1::SetUpdateNavOctreeOnComponentChange(bool bNewUpdateOnComponentChange)
+{
+	bUpdateNavOctreeOnComponentChange = bNewUpdateOnComponentChange;
+}
+
+void UNavigationSystemV1::DoInitialSetup()
+{
+	if (bInitialSetupHasBeenPerformed)
+	{
+		return;
+	}
+	
+	UpdateAbstractNavData();
+	CreateCrowdManager();
+
+	bInitialSetupHasBeenPerformed = true;
+}
+
+void UNavigationSystemV1::UpdateAbstractNavData()
+{
+	if (AbstractNavData != nullptr && !AbstractNavData->IsPendingKill())
+	{
+		return;
+	}
+
+	// spawn abstract nav data separately
+	// it's responsible for direct paths and shouldn't be picked for any agent type as default one
+	UWorld* NavWorld = GetWorld();
+	for (TActorIterator<AAbstractNavData> It(NavWorld); It; ++It)
+	{
+		AAbstractNavData* Nav = (*It);
+		if (Nav && !Nav->IsPendingKill())
+		{
+			AbstractNavData = Nav;
+			break;
+		}
+	}
+
+	if (AbstractNavData == NULL)
+	{
+		FNavDataConfig DummyConfig;
+		DummyConfig.NavigationDataClass = AAbstractNavData::StaticClass();
+		AbstractNavData = CreateNavigationDataInstance(DummyConfig);
+		AbstractNavData->SetFlags(RF_Transient);
+	}
+}
+
+void UNavigationSystemV1::SetSupportedAgentsNavigationClass(int32 AgentIndex, TSubclassOf<ANavigationData> NavigationDataClass)
+{
+	check(SupportedAgents.IsValidIndex(AgentIndex));
+	SupportedAgents[AgentIndex].NavigationDataClass = NavigationDataClass;
+
+	// keep preferred navigation data class in sync with actual class
+	// this will be passed to navigation data actor and will be required
+	// for comparisons done in DoesSupportAgent calls
+	//
+	// "Any" navigation data preference is valid only for instanced agents
+	SupportedAgents[AgentIndex].SetPreferredNavData(NavigationDataClass);
+
+	if (NavigationDataClass != nullptr)
+	{
+		SupportedAgents[AgentIndex].NavigationDataClassName = FSoftClassPath::GetOrCreateIDForClass(NavigationDataClass);
+	}
+	else
+	{
+		SupportedAgents[AgentIndex].NavigationDataClassName.Reset();
+	}
+
+#if WITH_EDITOR
+	if (GIsEditor)
+	{
+		if (HasAnyFlags(RF_ClassDefaultObject) == false)
+		{
+			// set it at CDO to properly show up in project settings
+			// @hack the reason for doing it this way is that engine doesn't handle default TSubclassOf properties
+			//	set to game-specific classes;
+			UNavigationSystemV1* NavigationSystemCDO = GetMutableDefault<UNavigationSystemV1>(GetClass());
+			NavigationSystemCDO->SetSupportedAgentsNavigationClass(AgentIndex, NavigationDataClass);
+		}
+	}
+#endif // WITH_EDITOR
+}
+
+void UNavigationSystemV1::PostInitProperties()
+{
+	Super::PostInitProperties();
+
+	if (HasAnyFlags(RF_ClassDefaultObject) == false)
+	{
+		// Populate our NavAreaClasses list with all known nav area classes.
+		// If more are loaded after this they will be registered as they come
+		TArray<UClass*> CurrentNavAreaClasses;
+		GetDerivedClasses(UNavArea::StaticClass(), CurrentNavAreaClasses);
+		for (UClass* NavAreaClass : CurrentNavAreaClasses)
+		{
+			RegisterNavAreaClass(NavAreaClass);
+		}
+
+		// make sure there's at least one supported navigation agent size
+		if (SupportedAgents.Num() == 0)
+		{
+			SupportedAgents.Add(FNavDataConfig(FNavigationSystem::FallbackAgentRadius, FNavigationSystem::FallbackAgentHeight));
+		}
+		else
+		{
+			for (int32 AgentIndex = 0; AgentIndex < SupportedAgents.Num(); ++AgentIndex)
+			{
+				FNavDataConfig& SupportedAgentConfig = SupportedAgents[AgentIndex];
+				// a piece of legacy maintanance 
+				if (SupportedAgentConfig.NavigationDataClass != nullptr && SupportedAgentConfig.NavigationDataClassName.IsValid() == false)
+				{
+					// fill NavigationDataClassName
+					SupportedAgentConfig.NavigationDataClassName = FSoftClassPath(SupportedAgentConfig.NavigationDataClass);
+				}
+				else
+				{
+					TSubclassOf<ANavigationData> NavigationDataClass = SupportedAgentConfig.NavigationDataClassName.IsValid()
+						? LoadClass<ANavigationData>(NULL, *SupportedAgentConfig.NavigationDataClassName.ToString(), NULL, LOAD_None, NULL)
+						: nullptr;
+
+					SetSupportedAgentsNavigationClass(AgentIndex, NavigationDataClass);
+				}
+			}
+		}
+	
+		if (bInitialBuildingLocked)
+		{
+			InitialNavBuildingLockFlags |= ENavigationBuildLock::InitialLock;
+		}
+
+		uint8 UseLockFlags = InitialNavBuildingLockFlags;
+
+		AddNavigationBuildLock(UseLockFlags);
+
+		// register for any actor move change
+#if WITH_EDITOR
+		if ( GIsEditor )
+		{
+			GEngine->OnActorMoved().AddUObject(this, &UNavigationSystemV1::OnActorMoved);
+		}
+#endif
+		FCoreUObjectDelegates::PostLoadMapWithWorld.AddUObject(this, &UNavigationSystemV1::OnPostLoadMap);
+		UNavigationSystemV1::NavigationDirtyEvent.AddUObject(this, &UNavigationSystemV1::OnNavigationDirtied);
+
+#if WITH_HOT_RELOAD
+		IHotReloadInterface& HotReloadSupport = FModuleManager::LoadModuleChecked<IHotReloadInterface>("HotReload");
+		HotReloadDelegateHandle = HotReloadSupport.OnHotReload().AddUObject(this, &UNavigationSystemV1::OnHotReload);
+#endif
+	}
+}
+
+bool UNavigationSystemV1::ConditionalPopulateNavOctree()
+{
+	// Discard all navigation updates caused by octree construction
+	TGuardValue<TArray<FNavigationDirtyArea>> DirtyGuard(DirtyAreas, TArray<FNavigationDirtyArea>());
+
+	// We are going to fully re-populate NavOctree so all pending update request are outdated
+	PendingOctreeUpdates.Empty(32);
+	
+	// Discard current octree
+	DestroyNavOctree();
+	
+	// See if any of registered navigation data need navoctree
+	bSupportRebuilding = RequiresNavOctree();
+
+	if (bSupportRebuilding)
+	{
+		NavOctree = MakeShareable(new FNavigationOctree(FVector(0,0,0), 64000));
+		NavOctree->SetDataGatheringMode(DataGatheringMode);
+		
+		const ERuntimeGenerationType RuntimeGenerationType = GetRuntimeGenerationType();
+		const bool bStoreNavGeometry = (RuntimeGenerationType == ERuntimeGenerationType::Dynamic);
+		NavOctree->SetNavigableGeometryStoringMode(bStoreNavGeometry ? FNavigationOctree::StoreNavGeometry : FNavigationOctree::SkipNavGeometry);
+		if (bStoreNavGeometry)
+		{
+#if WITH_RECAST
+			NavOctree->ComponentExportDelegate = FNavigationOctree::FNavigableGeometryComponentExportDelegate::CreateStatic(&FRecastNavMeshGenerator::ExportComponentGeometry);
+#endif // WITH_RECAST
+		}
+
+		if (!IsNavigationOctreeLocked())
+		{
+			UWorld* World = GetWorld();
+			check(World);
+
+			// now process all actors on all levels
+			for (int32 LevelIndex = 0; LevelIndex < World->GetNumLevels(); ++LevelIndex)
+			{
+				ULevel* Level = World->GetLevel(LevelIndex);
+				AddLevelCollisionToOctree(Level);
+
+				for (int32 ActorIndex = 0; ActorIndex < Level->Actors.Num(); ActorIndex++)
+				{
+					AActor* Actor = Level->Actors[ActorIndex];
+
+					const bool bLegalActor = Actor && !Actor->IsPendingKill();
+					if (bLegalActor)
+					{
+						UpdateActorAndComponentsInNavOctree(*Actor);
+					}
+				}
+			}
+		}
+	}
+	
+	// Add all found elements to octree, this will not add new dirty areas to navigation
+	if (PendingOctreeUpdates.Num())
+	{
+		for (TSet<FNavigationDirtyElement>::TIterator It(PendingOctreeUpdates); It; ++It)
+		{
+			AddElementToNavOctree(*It);
+		}
+		PendingOctreeUpdates.Empty(32);
+	}
+
+	return bSupportRebuilding;
+}
+
+#if WITH_EDITOR
+void UNavigationSystemV1::PostEditChangeChainProperty(FPropertyChangedChainEvent& PropertyChangedEvent)
+{
+	static const FName NAME_NavigationDataClass = GET_MEMBER_NAME_CHECKED(FNavDataConfig, NavigationDataClass);
+	static const FName NAME_SupportedAgents = GET_MEMBER_NAME_CHECKED(UNavigationSystemV1, SupportedAgents);
+
+	Super::PostEditChangeChainProperty(PropertyChangedEvent);
+
+	if (PropertyChangedEvent.Property)
+	{
+		FName PropName = PropertyChangedEvent.Property->GetFName();
+		if (PropName == NAME_NavigationDataClass)
+		{
+			int32 SupportedAgentIndex = PropertyChangedEvent.GetArrayIndex(NAME_SupportedAgents.ToString());
+			if (SupportedAgents.IsValidIndex(SupportedAgentIndex))
+			{
+				// reflect the change to SupportedAgent's 
+				TSubclassOf<ANavigationData> NavClass = SupportedAgents[SupportedAgentIndex].NavigationDataClass.Get();
+				SetSupportedAgentsNavigationClass(SupportedAgentIndex, NavClass);
+				SaveConfig();
+			}
+		}
+	}
+}
+
+void UNavigationSystemV1::PostEditChangeProperty(FPropertyChangedEvent& PropertyChangedEvent)
+{
+	static const FName NAME_EnableActiveTiles = GET_MEMBER_NAME_CHECKED(UNavigationSystemV1, bGenerateNavigationOnlyAroundNavigationInvokers);
+
+	Super::PostEditChangeProperty(PropertyChangedEvent);
+
+	if (PropertyChangedEvent.Property)
+	{
+		FName PropName = PropertyChangedEvent.Property->GetFName();
+		if (PropName == NAME_EnableActiveTiles)
+		{
+			if (NavOctree.IsValid())
+			{
+				NavOctree->SetDataGatheringMode(DataGatheringMode);
+			}
+
+			for (auto NavData : NavDataSet)
+			{
+				if (NavData)
+				{
+					NavData->RestrictBuildingToActiveTiles(bGenerateNavigationOnlyAroundNavigationInvokers);
+				}
+			}
+		}
+	}
+}
+#endif // WITH_EDITOR
+
+void UNavigationSystemV1::OnInitializeActors()
+{
+	
+}
+
+void UNavigationSystemV1::OnWorldInitDone(FNavigationSystemRunMode Mode)
+{
+	static const bool bSkipRebuildInEditor = true;
+	OperationMode = Mode;
+	DoInitialSetup();
+	
+	UWorld* World = GetWorld();
+
+	if (IsThereAnywhereToBuildNavigation() == false
+		// Simulation mode is a special case - better not do it in this case
+		&& OperationMode != FNavigationSystemRunMode::SimulationMode)
+	{
+		// remove all navigation data instances
+		for (TActorIterator<ANavigationData> It(World); It; ++It)
+		{
+			ANavigationData* Nav = (*It);
+			if (Nav != NULL && Nav->IsPendingKill() == false && Nav != GetAbstractNavData())
+			{
+				UnregisterNavData(Nav);
+				Nav->CleanUpAndMarkPendingKill();
+				bNavDataRemovedDueToMissingNavBounds = true;
+			}
+		}
+
+		if (OperationMode == FNavigationSystemRunMode::EditorMode)
+		{
+			RemoveNavigationBuildLock(InitialNavBuildingLockFlags, bSkipRebuildInEditor);
+		}
+	}
+	else
+	{
+		// Discard all bounds updates that was submitted during world initialization, 
+		// to avoid navigation rebuild right after map is loaded
+		PendingNavBoundsUpdates.Empty();
+		
+		// gather navigable bounds
+		GatherNavigationBounds();
+
+		// gather all navigation data instances and register all not-yet-registered
+		// (since it's quite possible navigation system was not ready by the time
+		// those instances were serialized-in or spawned)
+		RegisterNavigationDataInstances();
+
+		if (bAutoCreateNavigationData == true)
+		{
+			SpawnMissingNavigationData();
+			// in case anything spawned has registered
+			ProcessRegistrationCandidates();
+		}
+		else
+		{
+			const bool bIsBuildLocked = IsNavigationBuildingLocked();
+			if (GetDefaultNavDataInstance(FNavigationSystem::DontCreate) != NULL)
+			{
+				// trigger navmesh update
+				for (TActorIterator<ANavigationData> It(World); It; ++It)
+				{
+					ANavigationData* NavData = (*It);
+					if (NavData != NULL)
+					{
+						ERegistrationResult Result = RegisterNavData(NavData);
+
+						if (Result == RegistrationSuccessful)
+						{
+							if (!bIsBuildLocked && bNavigationAutoUpdateEnabled)
+							{
+								NavData->RebuildAll();
+							}
+						}
+						else if (Result != RegistrationFailed_DataPendingKill
+							&& Result != RegistrationFailed_AgentNotValid
+							)
+						{
+							NavData->CleanUpAndMarkPendingKill();
+						}
+					}
+				}
+			}
+		}
+
+		if (OperationMode == FNavigationSystemRunMode::EditorMode)
+		{
+			// don't lock navigation building in editor
+			RemoveNavigationBuildLock(InitialNavBuildingLockFlags, bSkipRebuildInEditor);
+		}
+
+		// See if any of registered navigation data needs NavOctree
+		ConditionalPopulateNavOctree();
+
+		// All navigation actors are registered
+		// Add NavMesh parts from all sub-levels that were streamed in prior NavMesh registration
+		const auto& Levels = World->GetLevels();
+		for (ULevel* Level : Levels)
+		{
+			if (!Level->IsPersistentLevel() && Level->bIsVisible)
+			{
+				for (ANavigationData* NavData : NavDataSet)
+				{
+					if (NavData)
+					{
+						NavData->OnStreamingLevelAdded(Level, World);
+					}
+				}
+			}
+		}
+	}
+
+	if (Mode == FNavigationSystemRunMode::EditorMode)
+	{
+#if	WITH_EDITOR
+		// make sure this static get applied to this instance
+		bNavigationAutoUpdateEnabled = !bNavigationAutoUpdateEnabled; 
+		SetNavigationAutoUpdateEnabled(!bNavigationAutoUpdateEnabled, this);
+#endif		
+		
+		// update navigation invokers
+		if (bGenerateNavigationOnlyAroundNavigationInvokers)
+		{
+			for (TObjectIterator<UNavigationInvokerComponent> It; It; ++It)
+			{
+				if (World == It->GetWorld())
+				{
+					It->RegisterWithNavigationSystem(*this);
+				}
+			}
+		}
+
+		// update navdata after loading world
+		if (bNavigationAutoUpdateEnabled)
+		{
+			const bool bIsLoadTime = true;
+			RebuildAll(bIsLoadTime);
+		}
+	}
+
+	if (!bCanAccumulateDirtyAreas)
+	{
+		DirtyAreas.Empty();
+	}
+
+	bWorldInitDone = true;
+	OnNavigationInitDone.Broadcast();
+}
+
+void UNavigationSystemV1::RegisterNavigationDataInstances()
+{
+	UWorld* World = GetWorld();
+
+	bool bProcessRegistration = false;
+	for (TActorIterator<ANavigationData> It(World); It; ++It)
+	{
+		ANavigationData* Nav = (*It);
+		if (Nav != NULL && Nav->IsPendingKill() == false && Nav->IsRegistered() == false)
+		{
+			RequestRegistration(Nav, false);
+			bProcessRegistration = true;
+		}
+	}
+	if (bProcessRegistration == true)
+	{
+		ProcessRegistrationCandidates();
+	}
+}
+
+void UNavigationSystemV1::CreateCrowdManager()
+{
+	if (CrowdManagerClass)
+	{
+		SetCrowdManager(NewObject<UCrowdManagerBase>(this, CrowdManagerClass));
+	}
+}
+
+void UNavigationSystemV1::SetCrowdManager(UCrowdManagerBase* NewCrowdManager)
+{
+	if (NewCrowdManager == CrowdManager.Get())
+	{
+		return;
+	}
+
+	if (CrowdManager.IsValid())
+	{
+		CrowdManager->RemoveFromRoot();
+	}
+	CrowdManager = NewCrowdManager;
+	if (NewCrowdManager != NULL)
+	{
+		CrowdManager->AddToRoot();
+	}
+}
+
+void UNavigationSystemV1::Tick(float DeltaSeconds)
+{
+	SET_DWORD_STAT(STAT_Navigation_ObservedPathsCount, 0);
+
+	UWorld* World = GetWorld();
+
+	if (World == nullptr || (bTickWhilePaused == false && World->IsPaused()))
+	{
+		return;
+	}
+
+	const bool bIsGame = World->IsGameWorld();
+	
+	if (PendingCustomLinkRegistration.Num())
+	{
+		ProcessCustomLinkPendingRegistration();
+	}
+
+	if (PendingNavBoundsUpdates.Num() > 0)
+	{
+		PerformNavigationBoundsUpdate(PendingNavBoundsUpdates);
+		PendingNavBoundsUpdates.Reset();
+	}
+
+	if (PendingOctreeUpdates.Num() > 0)
+	{
+		SCOPE_CYCLE_COUNTER(STAT_Navigation_AddingActorsToNavOctree);
+
+		SCOPE_CYCLE_COUNTER(STAT_Navigation_BuildTime)
+		STAT(double ThisTime = 0);
+		{
+			SCOPE_SECONDS_COUNTER(ThisTime);
+			for (TSet<FNavigationDirtyElement>::TIterator It(PendingOctreeUpdates); It; ++It)
+			{
+				AddElementToNavOctree(*It);
+			}
+			PendingOctreeUpdates.Empty(32);
+		}
+		INC_FLOAT_STAT_BY(STAT_Navigation_CumulativeBuildTime,(float)ThisTime*1000);
+	}
+	
+	if (bGenerateNavigationOnlyAroundNavigationInvokers)
+	{
+		UpdateInvokers();
+	}
+
+	{
+		SCOPE_CYCLE_COUNTER(STAT_Navigation_TickMarkDirty);
+
+		DirtyAreasUpdateTime += DeltaSeconds;
+		const float DirtyAreasUpdateDeltaTime = 1.0f / DirtyAreasUpdateFreq;
+		const bool bCanRebuildNow = (DirtyAreasUpdateTime >= DirtyAreasUpdateDeltaTime) || !bIsGame;
+		const bool bIsLocked = IsNavigationBuildingLocked();
+
+		if (DirtyAreas.Num() > 0 && bCanRebuildNow && !bIsLocked)
+		{
+			for (int32 NavDataIndex = 0; NavDataIndex < NavDataSet.Num(); ++NavDataIndex)
+			{
+				ANavigationData* NavData = NavDataSet[NavDataIndex];
+				if (NavData)
+				{
+					NavData->RebuildDirtyAreas(DirtyAreas);
+				}
+			}
+
+			DirtyAreasUpdateTime = 0;
+			DirtyAreas.Reset();
+		}
+	}
+
+	// Tick navigation mesh async builders
+	if (!bAsyncBuildPaused)
+	{
+		SCOPE_CYCLE_COUNTER(STAT_Navigation_TickAsyncBuild);
+		for (ANavigationData* NavData : NavDataSet)
+		{
+			if (NavData)
+			{
+				NavData->TickAsyncBuild(DeltaSeconds);
+			}
+		}
+	}
+
+	if (AsyncPathFindingQueries.Num() > 0)
+	{
+		SCOPE_CYCLE_COUNTER(STAT_Navigation_TickAsyncPathfinding);
+		TriggerAsyncQueries(AsyncPathFindingQueries);
+		AsyncPathFindingQueries.Reset();
+	}
+
+	if (CrowdManager.IsValid())
+	{
+		CrowdManager->Tick(DeltaSeconds);
+	}
+}
+
+void UNavigationSystemV1::AddReferencedObjects(UObject* InThis, FReferenceCollector& Collector)
+{
+	UNavigationSystemV1* This = CastChecked<UNavigationSystemV1>(InThis);
+	UCrowdManagerBase* CrowdManager = This->GetCrowdManager();
+	Collector.AddReferencedObject(CrowdManager, InThis);
+
+	// don't reference NavAreaClasses in editor (unless PIE is active)
+	if (This->OperationMode != FNavigationSystemRunMode::EditorMode)
+	{
+		Collector.AddReferencedObjects(This->NavAreaClasses, InThis);
+	}
+}
+
+#if WITH_EDITOR
+void UNavigationSystemV1::SetNavigationAutoUpdateEnabled(bool bNewEnable, UNavigationSystemBase* InNavigationSystemBase)
+{
+	if (bNewEnable != bNavigationAutoUpdateEnabled)
+	{
+		bNavigationAutoUpdateEnabled = bNewEnable;
+
+		UNavigationSystemV1* NavSystem = Cast<UNavigationSystemV1>(InNavigationSystemBase);
+		if (NavSystem)
+		{
+			NavSystem->bCanAccumulateDirtyAreas = bNavigationAutoUpdateEnabled 
+				|| (NavSystem->OperationMode != FNavigationSystemRunMode::EditorMode && NavSystem->OperationMode != FNavigationSystemRunMode::InvalidMode);
+
+			if (bNavigationAutoUpdateEnabled)
+			{
+				const bool bSkipRebuildsInEditor = false;
+				NavSystem->RemoveNavigationBuildLock(ENavigationBuildLock::NoUpdateInEditor, bSkipRebuildsInEditor);
+			}
+			else
+			{
+#if !UE_BUILD_SHIPPING
+				NavSystem->bDirtyAreasReportedWhileAccumulationLocked = false;
+#endif // !UE_BUILD_SHIPPING
+				NavSystem->AddNavigationBuildLock(ENavigationBuildLock::NoUpdateInEditor);
+			}
+		}
+	}
+}
+#endif // WITH_EDITOR
+
+//----------------------------------------------------------------------//
+// Public querying interface                                                                
+//----------------------------------------------------------------------//
+FPathFindingResult UNavigationSystemV1::FindPathSync(const FNavAgentProperties& AgentProperties, FPathFindingQuery Query, EPathFindingMode::Type Mode)
+{
+	SCOPE_CYCLE_COUNTER(STAT_Navigation_PathfindingSync);
+
+	if (Query.NavData.IsValid() == false)
+	{
+		Query.NavData = GetNavDataForProps(AgentProperties);
+	}
+
+	FPathFindingResult Result(ENavigationQueryResult::Error);
+	if (Query.NavData.IsValid())
+	{
+		if (Mode == EPathFindingMode::Hierarchical)
+		{
+			Result = Query.NavData->FindHierarchicalPath(AgentProperties, Query);
+		}
+		else
+		{
+			Result = Query.NavData->FindPath(AgentProperties, Query);
+		}
+	}
+
+	return Result;
+}
+
+FPathFindingResult UNavigationSystemV1::FindPathSync(FPathFindingQuery Query, EPathFindingMode::Type Mode)
+{
+	SCOPE_CYCLE_COUNTER(STAT_Navigation_PathfindingSync);
+
+	if (Query.NavData.IsValid() == false)
+	{
+		Query.NavData = GetDefaultNavDataInstance(FNavigationSystem::DontCreate);
+	}
+	
+	FPathFindingResult Result(ENavigationQueryResult::Error);
+	if (Query.NavData.IsValid())
+	{
+		if (Mode == EPathFindingMode::Regular)
+		{
+			Result = Query.NavData->FindPath(Query.NavAgentProperties, Query);
+		}
+		else // EPathFindingMode::Hierarchical
+		{
+			Result = Query.NavData->FindHierarchicalPath(Query.NavAgentProperties, Query);
+		}
+	}
+
+	return Result;
+}
+
+bool UNavigationSystemV1::TestPathSync(FPathFindingQuery Query, EPathFindingMode::Type Mode, int32* NumVisitedNodes) const
+{
+	SCOPE_CYCLE_COUNTER(STAT_Navigation_PathfindingSync);
+
+	if (Query.NavData.IsValid() == false)
+	{
+		Query.NavData = GetDefaultNavDataInstance();
+	}
+	
+	bool bExists = false;
+	if (Query.NavData.IsValid())
+	{
+		if (Mode == EPathFindingMode::Hierarchical)
+		{
+			bExists = Query.NavData->TestHierarchicalPath(Query.NavAgentProperties, Query, NumVisitedNodes);
+		}
+		else
+		{
+			bExists = Query.NavData->TestPath(Query.NavAgentProperties, Query, NumVisitedNodes);
+		}
+	}
+
+	return bExists;
+}
+
+void UNavigationSystemV1::AddAsyncQuery(const FAsyncPathFindingQuery& Query)
+{
+	check(IsInGameThread());
+	AsyncPathFindingQueries.Add(Query);
+}
+
+uint32 UNavigationSystemV1::FindPathAsync(const FNavAgentProperties& AgentProperties, FPathFindingQuery Query, const FNavPathQueryDelegate& ResultDelegate, EPathFindingMode::Type Mode)
+{
+	SCOPE_CYCLE_COUNTER(STAT_Navigation_RequestingAsyncPathfinding);
+
+	if (Query.NavData.IsValid() == false)
+	{
+		Query.NavData = GetNavDataForProps(AgentProperties);
+	}
+
+	if (Query.NavData.IsValid())
+	{
+		FAsyncPathFindingQuery AsyncQuery(Query, ResultDelegate, Mode);
+
+		if (AsyncQuery.QueryID != INVALID_NAVQUERYID)
+		{
+			AddAsyncQuery(AsyncQuery);
+		}
+
+		return AsyncQuery.QueryID;
+	}
+
+	return INVALID_NAVQUERYID;
+}
+
+void UNavigationSystemV1::AbortAsyncFindPathRequest(uint32 AsynPathQueryID)
+{
+	check(IsInGameThread());
+	FAsyncPathFindingQuery* Query = AsyncPathFindingQueries.GetData();
+	for (int32 Index = 0; Index < AsyncPathFindingQueries.Num(); ++Index, ++Query)
+	{
+		if (Query->QueryID == AsynPathQueryID)
+		{
+			AsyncPathFindingQueries.RemoveAtSwap(Index);
+			break;
+		}
+	}
+}
+
+FAutoConsoleTaskPriority CPrio_TriggerAsyncQueries(
+	TEXT("TaskGraph.TaskPriorities.NavTriggerAsyncQueries"),
+	TEXT("Task and thread priority for UNavigationSystemV1::PerformAsyncQueries."),
+	ENamedThreads::BackgroundThreadPriority, // if we have background priority task threads, then use them...
+	ENamedThreads::NormalTaskPriority, // .. at normal task priority
+	ENamedThreads::NormalTaskPriority // if we don't have background threads, then use normal priority threads at normal task priority instead
+	);
+
+
+void UNavigationSystemV1::TriggerAsyncQueries(TArray<FAsyncPathFindingQuery>& PathFindingQueries)
+{
+	DECLARE_CYCLE_STAT(TEXT("FSimpleDelegateGraphTask.NavigationSystem batched async queries"),
+		STAT_FSimpleDelegateGraphTask_NavigationSystemBatchedAsyncQueries,
+		STATGROUP_TaskGraphTasks);
+
+	FSimpleDelegateGraphTask::CreateAndDispatchWhenReady(
+		FSimpleDelegateGraphTask::FDelegate::CreateUObject(this, &UNavigationSystemV1::PerformAsyncQueries, PathFindingQueries),
+		GET_STATID(STAT_FSimpleDelegateGraphTask_NavigationSystemBatchedAsyncQueries), nullptr, CPrio_TriggerAsyncQueries.Get());
+}
+
+static void AsyncQueryDone(FAsyncPathFindingQuery Query)
+{
+	Query.OnDoneDelegate.ExecuteIfBound(Query.QueryID, Query.Result.Result, Query.Result.Path);
+}
+
+void UNavigationSystemV1::PerformAsyncQueries(TArray<FAsyncPathFindingQuery> PathFindingQueries)
+{
+	SCOPE_CYCLE_COUNTER(STAT_Navigation_PathfindingAsync);
+
+	if (PathFindingQueries.Num() == 0)
+	{
+		return;
+	}
+	
+	for (FAsyncPathFindingQuery& Query : PathFindingQueries)
+	{
+		// @todo this is not necessarily the safest way to use UObjects outside of main thread. 
+		//	think about something else.
+		const ANavigationData* NavData = Query.NavData.IsValid() ? Query.NavData.Get() : GetDefaultNavDataInstance(FNavigationSystem::DontCreate);
+
+		// perform query
+		if (NavData)
+		{
+			if (Query.Mode == EPathFindingMode::Hierarchical)
+			{
+				Query.Result = NavData->FindHierarchicalPath(Query.NavAgentProperties, Query);
+			}
+			else
+			{
+				Query.Result = NavData->FindPath(Query.NavAgentProperties, Query);
+			}
+		}
+		else
+		{
+			Query.Result = ENavigationQueryResult::Error;
+		}
+
+		// @todo make it return more informative results (bResult == false)
+		// trigger calling delegate on main thread - otherwise it may depend too much on stuff being thread safe
+		DECLARE_CYCLE_STAT(TEXT("FSimpleDelegateGraphTask.Async nav query finished"),
+			STAT_FSimpleDelegateGraphTask_AsyncNavQueryFinished,
+			STATGROUP_TaskGraphTasks);
+
+		FSimpleDelegateGraphTask::CreateAndDispatchWhenReady(
+			FSimpleDelegateGraphTask::FDelegate::CreateStatic(AsyncQueryDone, Query),
+			GET_STATID(STAT_FSimpleDelegateGraphTask_AsyncNavQueryFinished), NULL, ENamedThreads::GameThread);
+	}
+}
+
+bool UNavigationSystemV1::GetRandomPoint(FNavLocation& ResultLocation, ANavigationData* NavData, FSharedConstNavQueryFilter QueryFilter)
+{
+	SCOPE_CYCLE_COUNTER(STAT_Navigation_QueriesTimeSync);
+
+	if (NavData == NULL)
+	{
+		NavData = MainNavData;
+	}
+
+	if (NavData != NULL)
+	{
+		ResultLocation = NavData->GetRandomPoint(QueryFilter);
+		return true;
+	}
+
+	return false;
+}
+
+bool UNavigationSystemV1::GetRandomReachablePointInRadius(const FVector& Origin, float Radius, FNavLocation& ResultLocation, ANavigationData* NavData, FSharedConstNavQueryFilter QueryFilter) const
+{
+	SCOPE_CYCLE_COUNTER(STAT_Navigation_QueriesTimeSync);
+
+	if (NavData == nullptr)
+	{
+		NavData = MainNavData;
+	}
+
+	return NavData != nullptr && NavData->GetRandomReachablePointInRadius(Origin, Radius, ResultLocation, QueryFilter);
+}
+
+bool UNavigationSystemV1::GetRandomPointInNavigableRadius(const FVector& Origin, float Radius, FNavLocation& ResultLocation, ANavigationData* NavData, FSharedConstNavQueryFilter QueryFilter) const
+{
+	SCOPE_CYCLE_COUNTER(STAT_Navigation_QueriesTimeSync);
+
+	if (NavData == nullptr)
+	{
+		NavData = MainNavData;
+	}
+
+	return NavData != nullptr && NavData->GetRandomPointInNavigableRadius(Origin, Radius, ResultLocation, QueryFilter);
+}
+
+ENavigationQueryResult::Type UNavigationSystemV1::GetPathCost(const FVector& PathStart, const FVector& PathEnd, float& OutPathCost, const ANavigationData* NavData, FSharedConstNavQueryFilter QueryFilter) const
+{
+	SCOPE_CYCLE_COUNTER(STAT_Navigation_QueriesTimeSync);
+
+	if (NavData == NULL)
+	{
+		NavData = GetDefaultNavDataInstance();
+	}
+
+	return NavData != NULL ? NavData->CalcPathCost(PathStart, PathEnd, OutPathCost, QueryFilter) : ENavigationQueryResult::Error;
+}
+
+ENavigationQueryResult::Type UNavigationSystemV1::GetPathLength(const FVector& PathStart, const FVector& PathEnd, float& OutPathLength, const ANavigationData* NavData, FSharedConstNavQueryFilter QueryFilter) const
+{
+	SCOPE_CYCLE_COUNTER(STAT_Navigation_QueriesTimeSync);
+
+	if (NavData == NULL)
+	{
+		NavData = GetDefaultNavDataInstance();
+	}
+
+	return NavData != NULL ? NavData->CalcPathLength(PathStart, PathEnd, OutPathLength, QueryFilter) : ENavigationQueryResult::Error;
+}
+
+ENavigationQueryResult::Type UNavigationSystemV1::GetPathLengthAndCost(const FVector& PathStart, const FVector& PathEnd, float& OutPathLength, float& OutPathCost, const ANavigationData* NavData, FSharedConstNavQueryFilter QueryFilter) const
+{
+	SCOPE_CYCLE_COUNTER(STAT_Navigation_QueriesTimeSync);
+
+	if (NavData == NULL)
+	{
+		NavData = GetDefaultNavDataInstance();
+	}
+
+	return NavData != NULL ? NavData->CalcPathLengthAndCost(PathStart, PathEnd, OutPathLength, OutPathCost, QueryFilter) : ENavigationQueryResult::Error;
+}
+
+bool UNavigationSystemV1::ProjectPointToNavigation(const FVector& Point, FNavLocation& OutLocation, const FVector& Extent, const ANavigationData* NavData, FSharedConstNavQueryFilter QueryFilter) const
+{
+	SCOPE_CYCLE_COUNTER(STAT_Navigation_QueriesTimeSync);
+
+	if (NavData == NULL)
+	{
+		NavData = GetDefaultNavDataInstance();
+	}
+
+	return NavData != NULL && NavData->ProjectPoint(Point, OutLocation
+		, FNavigationSystem::IsValidExtent(Extent) ? Extent : NavData->GetConfig().DefaultQueryExtent
+		, QueryFilter);
+}
+
+UNavigationPath* UNavigationSystemV1::FindPathToActorSynchronously(UObject* WorldContextObject, const FVector& PathStart, AActor* GoalActor, float TetherDistance, AActor* PathfindingContext, TSubclassOf<UNavigationQueryFilter> FilterClass)
+{
+	if (GoalActor == NULL)
+	{
+		return NULL; 
+	}
+
+	INavAgentInterface* NavAgent = Cast<INavAgentInterface>(GoalActor);
+	UNavigationPath* GeneratedPath = FindPathToLocationSynchronously(WorldContextObject, PathStart, NavAgent ? NavAgent->GetNavAgentLocation() : GoalActor->GetActorLocation(), PathfindingContext, FilterClass);
+	if (GeneratedPath != NULL && GeneratedPath->GetPath().IsValid() == true)
+	{
+		GeneratedPath->GetPath()->SetGoalActorObservation(*GoalActor, TetherDistance);
+	}
+
+	return GeneratedPath;
+}
+
+UNavigationPath* UNavigationSystemV1::FindPathToLocationSynchronously(UObject* WorldContextObject, const FVector& PathStart, const FVector& PathEnd, AActor* PathfindingContext, TSubclassOf<UNavigationQueryFilter> FilterClass)
+{
+	UWorld* World = NULL;
+
+	if (WorldContextObject != NULL)
+	{
+		World = GEngine->GetWorldFromContextObject(WorldContextObject, EGetWorldErrorMode::LogAndReturnNull);
+	}
+	if (World == NULL && PathfindingContext != NULL)
+	{
+		World = GEngine->GetWorldFromContextObject(PathfindingContext, EGetWorldErrorMode::LogAndReturnNull);
+	}
+
+	UNavigationPath* ResultPath = NULL;
+
+	UNavigationSystemV1* NavSys = FNavigationSystem::GetCurrent<UNavigationSystemV1>(World);
+
+	if (NavSys != nullptr && NavSys->GetDefaultNavDataInstance() != nullptr)
+	{
+		ResultPath = NewObject<UNavigationPath>(NavSys);
+		bool bValidPathContext = false;
+		const ANavigationData* NavigationData = NULL;
+
+		if (PathfindingContext != NULL)
+		{
+			INavAgentInterface* NavAgent = Cast<INavAgentInterface>(PathfindingContext);
+			
+			if (NavAgent != NULL)
+			{
+				const FNavAgentProperties& AgentProps = NavAgent->GetNavAgentPropertiesRef();
+				NavigationData = NavSys->GetNavDataForProps(AgentProps);
+				bValidPathContext = true;
+			}
+			else if (Cast<ANavigationData>(PathfindingContext))
+			{
+				NavigationData = (ANavigationData*)PathfindingContext;
+				bValidPathContext = true;
+			}
+		}
+		if (bValidPathContext == false)
+		{
+			// just use default
+			NavigationData = NavSys->GetDefaultNavDataInstance();
+		}
+
+		check(NavigationData);
+
+		const FPathFindingQuery Query(PathfindingContext, *NavigationData, PathStart, PathEnd, UNavigationQueryFilter::GetQueryFilter(*NavigationData, PathfindingContext, FilterClass));
+		const FPathFindingResult Result = NavSys->FindPathSync(Query, EPathFindingMode::Regular);
+		if (Result.IsSuccessful())
+		{
+			ResultPath->SetPath(Result.Path);
+		}
+	}
+
+	return ResultPath;
+}
+
+bool UNavigationSystemV1::NavigationRaycast(UObject* WorldContextObject, const FVector& RayStart, const FVector& RayEnd, FVector& HitLocation, TSubclassOf<UNavigationQueryFilter> FilterClass, AController* Querier)
+{
+	UWorld* World = NULL;
+
+	if (WorldContextObject != NULL)
+	{
+		World = GEngine->GetWorldFromContextObject(WorldContextObject, EGetWorldErrorMode::LogAndReturnNull);
+	}
+	if (World == NULL && Querier != NULL)
+	{
+		World = GEngine->GetWorldFromContextObject(Querier, EGetWorldErrorMode::LogAndReturnNull);
+	}
+
+	// blocked, i.e. not traversable, by default
+	bool bRaycastBlocked = true;
+	HitLocation = RayStart;
+
+	const UNavigationSystemV1* NavSys = FNavigationSystem::GetCurrent<UNavigationSystemV1>(World);
+
+	if (NavSys)
+	{
+		// figure out which navigation data to use
+		const ANavigationData* NavData = NULL;
+		INavAgentInterface* MyNavAgent = Cast<INavAgentInterface>(Querier);
+		if (MyNavAgent)
+		{
+			const FNavAgentProperties& AgentProps = MyNavAgent->GetNavAgentPropertiesRef();
+			NavData = NavSys->GetNavDataForProps(AgentProps);
+		}
+		if (NavData == NULL)
+		{
+			NavData = NavSys->GetDefaultNavDataInstance();
+		}
+
+		if (NavData != NULL)
+		{
+			bRaycastBlocked = NavData->Raycast(RayStart, RayEnd, HitLocation, UNavigationQueryFilter::GetQueryFilter(*NavData, Querier, FilterClass));
+		}
+	}
+
+	return bRaycastBlocked;
+}
+
+void UNavigationSystemV1::GetNavAgentPropertiesArray(TArray<FNavAgentProperties>& OutNavAgentProperties) const
+{
+	AgentToNavDataMap.GetKeys(OutNavAgentProperties);
+}
+
+ANavigationData* UNavigationSystemV1::GetNavDataForProps(const FNavAgentProperties& AgentProperties)
+{
+	const UNavigationSystemV1* ConstThis = AsConst(this);
+	return const_cast<ANavigationData*>(ConstThis->GetNavDataForProps(AgentProperties));
+}
+
+// @todo could optimize this by having "SupportedAgentIndex" in FNavAgentProperties
+const ANavigationData* UNavigationSystemV1::GetNavDataForProps(const FNavAgentProperties& AgentProperties) const
+{
+	if (SupportedAgents.Num() <= 1)
+	{
+		return MainNavData;
+	}
+	
+	const TWeakObjectPtr<ANavigationData>* NavDataForAgent = AgentToNavDataMap.Find(AgentProperties);
+	const ANavigationData* NavDataInstance = NavDataForAgent ? NavDataForAgent->Get() : nullptr;
+
+	if (NavDataInstance == nullptr)
+	{
+		TArray<FNavAgentProperties> AgentPropertiesList;
+		AgentToNavDataMap.GenerateKeyArray(AgentPropertiesList);
+		
+		FNavAgentProperties BestFitNavAgent;
+		float BestExcessHeight = -FLT_MAX;
+		float BestExcessRadius = -FLT_MAX;
+		float ExcessRadius = -FLT_MAX;
+		float ExcessHeight = -FLT_MAX;
+		const float AgentHeight = bSkipAgentHeightCheckWhenPickingNavData ? 0.f : AgentProperties.AgentHeight;
+
+		for (TArray<FNavAgentProperties>::TConstIterator It(AgentPropertiesList); It; ++It)
+		{
+			const FNavAgentProperties& NavIt = *It;
+			const bool bNavClassMatch = NavIt.IsNavDataMatching(AgentProperties);
+			if (!bNavClassMatch)
+			{
+				continue;
+			}
+
+			ExcessRadius = NavIt.AgentRadius - AgentProperties.AgentRadius;
+			ExcessHeight = bSkipAgentHeightCheckWhenPickingNavData ? 0.f : (NavIt.AgentHeight - AgentHeight);
+
+			const bool bExcessRadiusIsBetter = ((ExcessRadius == 0) && (BestExcessRadius != 0)) 
+				|| ((ExcessRadius > 0) && (BestExcessRadius < 0))
+				|| ((ExcessRadius > 0) && (BestExcessRadius > 0) && (ExcessRadius < BestExcessRadius))
+				|| ((ExcessRadius < 0) && (BestExcessRadius < 0) && (ExcessRadius > BestExcessRadius));
+			const bool bExcessHeightIsBetter = ((ExcessHeight == 0) && (BestExcessHeight != 0))
+				|| ((ExcessHeight > 0) && (BestExcessHeight < 0))
+				|| ((ExcessHeight > 0) && (BestExcessHeight > 0) && (ExcessHeight < BestExcessHeight))
+				|| ((ExcessHeight < 0) && (BestExcessHeight < 0) && (ExcessHeight > BestExcessHeight));
+			const bool bBestIsValid = (BestExcessRadius >= 0) && (BestExcessHeight >= 0);
+			const bool bRadiusEquals = (ExcessRadius == BestExcessRadius);
+			const bool bHeightEquals = (ExcessHeight == BestExcessHeight);
+
+			bool bValuesAreBest = ((bExcessRadiusIsBetter || bRadiusEquals) && (bExcessHeightIsBetter || bHeightEquals));
+			if (!bValuesAreBest && !bBestIsValid)
+			{
+				bValuesAreBest = bExcessRadiusIsBetter || (bRadiusEquals && bExcessHeightIsBetter);
+			}
+
+			if (bValuesAreBest)
+			{
+				BestFitNavAgent = NavIt;
+				BestExcessHeight = ExcessHeight;
+				BestExcessRadius = ExcessRadius;
+			}
+		}
+
+		if (BestFitNavAgent.IsValid())
+		{
+			NavDataForAgent = AgentToNavDataMap.Find(BestFitNavAgent);
+			NavDataInstance = NavDataForAgent ? NavDataForAgent->Get() : nullptr;
+		}
+	}
+
+	return NavDataInstance ? NavDataInstance : MainNavData;
+}
+
+ANavigationData* UNavigationSystemV1::GetDefaultNavDataInstance(FNavigationSystem::ECreateIfMissing CreateNewIfNoneFound)
+{
+	checkSlow(IsInGameThread() == true);
+
+	if (MainNavData == NULL || MainNavData->IsPendingKill())
+	{
+		MainNavData = NULL;
+
+		// @TODO this should be done a differently. There should be specified a "default agent"
+		for (int32 NavDataIndex = 0; NavDataIndex < NavDataSet.Num(); ++NavDataIndex)
+		{
+			ANavigationData* NavData = NavDataSet[NavDataIndex];
+			if (NavData && !NavData->IsPendingKill() && NavData->CanBeMainNavData())
+			{
+				MainNavData = NavData;
+				break;
+			}
+		}
+
+#if WITH_RECAST
+		if ( /*GIsEditor && */(MainNavData == NULL) && CreateNewIfNoneFound == FNavigationSystem::Create )
+		{
+			// Spawn a new one if we're in the editor.  In-game, either we loaded one or we don't get one.
+			MainNavData = GetWorld()->SpawnActor<ANavigationData>(ARecastNavMesh::StaticClass());
+		}
+#endif // WITH_RECAST
+		// either way make sure it's registered. Registration stores unique
+		// navmeshes, so we have nothing to lose
+		RegisterNavData(MainNavData);
+	}
+
+	return MainNavData;
+}
+
+FSharedNavQueryFilter UNavigationSystemV1::CreateDefaultQueryFilterCopy() const 
+{ 
+	return MainNavData ? MainNavData->GetDefaultQueryFilter()->GetCopy() : NULL; 
+}
+
+bool UNavigationSystemV1::IsNavigationBuilt(const AWorldSettings* Settings) const
+{
+	if (Settings == nullptr || Settings->IsNavigationSystemEnabled() == false || IsThereAnywhereToBuildNavigation() == false)
+	{
+		return true;
+	}
+
+	bool bIsBuilt = true;
+
+	for (int32 NavDataIndex = 0; NavDataIndex < NavDataSet.Num(); ++NavDataIndex)
+	{
+		ANavigationData* NavData = NavDataSet[NavDataIndex];
+		if (NavData != NULL && NavData->GetWorldSettings() == Settings)
+		{
+			FNavDataGenerator* Generator = NavData->GetGenerator();
+			if ((NavData->GetRuntimeGenerationMode() != ERuntimeGenerationType::Static
+#if WITH_EDITOR
+				|| GEditor != NULL
+#endif // WITH_EDITOR
+				) && (Generator == NULL || Generator->IsBuildInProgress(/*bCheckDirtyToo=*/true) == true))
+			{
+				bIsBuilt = false;
+				break;
+			}
+		}
+	}
+
+	return bIsBuilt;
+}
+
+bool UNavigationSystemV1::IsThereAnywhereToBuildNavigation() const
+{
+	// not check if there are any volumes or other structures requiring/supporting navigation building
+	if (bWholeWorldNavigable == true)
+	{
+		return true;
+	}
+
+	// @TODO this should be done more flexible to be able to trigger this from game-specific 
+	// code (like Navigation System's subclass maybe)
+	bool bCreateNavigation = false;
+
+	for (TActorIterator<ANavMeshBoundsVolume> It(GetWorld()); It; ++It)
+	{
+		ANavMeshBoundsVolume const* const V = (*It);
+		if (V != NULL && !V->IsPendingKill())
+		{
+			bCreateNavigation = true;
+			break;
+		}
+	}
+
+	return bCreateNavigation;
+}
+
+bool UNavigationSystemV1::IsNavigationRelevant(const AActor* TestActor) const
+{
+	const INavRelevantInterface* NavInterface = Cast<const INavRelevantInterface>(TestActor);
+	if (NavInterface && NavInterface->IsNavigationRelevant())
+	{
+		return true;
+	}
+
+	if (TestActor)
+	{
+		TInlineComponentArray<UActorComponent*> Components;
+		for (int32 Idx = 0; Idx < Components.Num(); Idx++)
+		{
+			NavInterface = Cast<const INavRelevantInterface>(Components[Idx]);
+			if (NavInterface && NavInterface->IsNavigationRelevant())
+			{
+				return true;
+			}
+		}
+	}
+
+	return false;
+}
+
+FBox UNavigationSystemV1::GetWorldBounds() const
+{
+	checkSlow(IsInGameThread() == true);
+
+	NavigableWorldBounds = FBox(ForceInit);
+
+	if (GetWorld() != nullptr)
+	{
+		if (bWholeWorldNavigable == false)
+		{
+			for (const FNavigationBounds& Bounds : RegisteredNavBounds)
+			{
+				NavigableWorldBounds += Bounds.AreaBox;
+			}
+		}
+		else
+		{
+			// @TODO - super slow! Need to ask tech guys where I can get this from
+			for (FActorIterator It(GetWorld()); It; ++It)
+			{
+				if (IsNavigationRelevant(*It))
+				{
+					NavigableWorldBounds += (*It)->GetComponentsBoundingBox();
+				}
+			}
+		}
+	}
+
+	return NavigableWorldBounds;
+}
+
+FBox UNavigationSystemV1::GetLevelBounds(ULevel* InLevel) const
+{
+	FBox NavigableLevelBounds(ForceInit);
+
+	if (InLevel)
+	{
+		AActor** Actor = InLevel->Actors.GetData();
+		const int32 ActorCount = InLevel->Actors.Num();
+		for (int32 ActorIndex = 0; ActorIndex < ActorCount; ++ActorIndex, ++Actor)
+		{
+			if (IsNavigationRelevant(*Actor))
+			{
+				NavigableLevelBounds += (*Actor)->GetComponentsBoundingBox();
+			}
+		}
+	}
+
+	return NavigableLevelBounds;
+}
+
+const TSet<FNavigationBounds>& UNavigationSystemV1::GetNavigationBounds() const
+{
+	return RegisteredNavBounds;
+}
+
+void UNavigationSystemV1::ApplyWorldOffset(const FVector& InOffset, bool bWorldShift)
+{
+	// Attempt at generation of new nav mesh after the shift
+	// dynamic navmesh, we regenerate completely
+	if (GetRuntimeGenerationType() == ERuntimeGenerationType::Dynamic)
+	{
+		//stop generators from building navmesh
+		for (ANavigationData* NavData : NavDataSet)
+		{
+			if (NavData)
+			{
+				if (NavData->GetGenerator()) NavData->GetGenerator()->CancelBuild();
+			}
+		}
+
+		ConditionalPopulateNavOctree();
+		Build();
+
+		for (ANavigationData* NavData : NavDataSet)
+		{
+			if (NavData)
+			{
+				NavData->ConditionalConstructGenerator();
+				ARecastNavMesh* RecastNavMesh = Cast<ARecastNavMesh>(NavData);
+				if (RecastNavMesh) RecastNavMesh->RequestDrawingUpdate();
+			}
+		}
+	}
+	else // static navmesh
+	{
+		//not sure what happens when we shift farther than the extents of the NavOctree are
+		for (ANavigationData* NavData : NavDataSet)
+		{
+			if (NavData)
+			{
+				NavData->ApplyWorldOffset(InOffset, bWorldShift);
+			}
+		}
+	}
+}
+
+//----------------------------------------------------------------------//
+// Bookkeeping 
+//----------------------------------------------------------------------//
+void UNavigationSystemV1::RequestRegistration(ANavigationData* NavData, bool bTriggerRegistrationProcessing)
+{
+	FScopeLock RegistrationLock(&NavDataRegistrationSection);
+
+	if (NavDataRegistrationQueue.Num() < REGISTRATION_QUEUE_SIZE)
+	{
+		NavDataRegistrationQueue.AddUnique(NavData);
+
+		if (bTriggerRegistrationProcessing == true)
+		{
+			// trigger registration candidates processing
+			DECLARE_CYCLE_STAT(TEXT("FSimpleDelegateGraphTask.Process registration candidates"),
+				STAT_FSimpleDelegateGraphTask_ProcessRegistrationCandidates,
+				STATGROUP_TaskGraphTasks);
+
+			FSimpleDelegateGraphTask::CreateAndDispatchWhenReady(
+				FSimpleDelegateGraphTask::FDelegate::CreateUObject(this, &UNavigationSystemV1::ProcessRegistrationCandidates),
+				GET_STATID(STAT_FSimpleDelegateGraphTask_ProcessRegistrationCandidates), NULL, ENamedThreads::GameThread);
+		}
+	}
+	else
+	{
+		UE_LOG(LogNavigation, Error, TEXT("Navigation System: registration queue full!"));
+	}
+}
+
+void UNavigationSystemV1::ProcessRegistrationCandidates()
+{
+	FScopeLock RegistrationLock(&NavDataRegistrationSection);
+
+	if (NavDataRegistrationQueue.Num() == 0)
+	{
+		return;
+	}
+
+	ANavigationData** NavDataPtr = NavDataRegistrationQueue.GetData();
+	const int CandidatesCount = NavDataRegistrationQueue.Num();
+
+	for (int32 CandidateIndex = 0; CandidateIndex < CandidatesCount; ++CandidateIndex, ++NavDataPtr)
+	{
+		if (*NavDataPtr != NULL)
+		{
+			ERegistrationResult Result = RegisterNavData(*NavDataPtr);
+
+			if (Result == RegistrationSuccessful)
+			{
+				continue;
+			}
+			else if (Result != RegistrationFailed_DataPendingKill)
+			{
+				(*NavDataPtr)->CleanUpAndMarkPendingKill();
+				if ((*NavDataPtr) == MainNavData)
+				{
+					MainNavData = NULL;
+				}
+			}
+		}
+	}
+
+	MainNavData = GetDefaultNavDataInstance(FNavigationSystem::DontCreate);
+	
+	// we processed all candidates so clear the queue
+	NavDataRegistrationQueue.Reset();
+}
+
+void UNavigationSystemV1::ProcessCustomLinkPendingRegistration()
+{
+	FScopeLock AccessLock(&CustomLinkRegistrationSection);
+
+	TMap<INavLinkCustomInterface*, FWeakObjectPtr> TempPending = PendingCustomLinkRegistration;
+	PendingCustomLinkRegistration.Empty();
+
+	for (TMap<INavLinkCustomInterface*, FWeakObjectPtr>::TIterator It(TempPending); It; ++It)
+	{
+		INavLinkCustomInterface* ILink = It.Key();
+		FWeakObjectPtr LinkOb = It.Value();
+		
+		if (LinkOb.IsValid() && ILink)
+		{
+			RegisterCustomLink(*ILink);
+		}
+	}
+}
+
+UNavigationSystemV1::ERegistrationResult UNavigationSystemV1::RegisterNavData(ANavigationData* NavData)
+{
+	if (NavData == NULL)
+	{
+		return RegistrationError;
+	}
+	else if (NavData->IsPendingKill() == true)
+	{
+		return RegistrationFailed_DataPendingKill;
+	}
+	// still to be seen if this is really true, but feels right
+	else if (NavData->IsRegistered() == true)
+	{
+		return RegistrationSuccessful;
+	}
+
+	FScopeLock Lock(&NavDataRegistration);
+
+	UNavigationSystemV1::ERegistrationResult Result = RegistrationError;
+
+	// find out which, if any, navigation agents are supported by this nav data
+	// if none then fail the registration
+	FNavDataConfig NavConfig = NavData->GetConfig();
+
+	// not discarding navmesh when there's only one Supported Agent
+	if (NavConfig.IsValid() == false && SupportedAgents.Num() == 1)
+	{
+		// fill in AgentProps with whatever is the instance's setup
+		NavConfig = SupportedAgents[0];
+		NavData->SetConfig(SupportedAgents[0]);
+		NavData->SetSupportsDefaultAgent(true);	
+		NavData->ProcessNavAreas(NavAreaClasses, 0);
+	}
+
+	if (NavConfig.IsValid() == true)
+	{
+		// check if this kind of agent has already its navigation implemented
+		TWeakObjectPtr<ANavigationData>* NavDataForAgent = AgentToNavDataMap.Find(NavConfig);
+		ANavigationData* NavDataInstanceForAgent = NavDataForAgent ? NavDataForAgent->Get() : nullptr;
+
+		if (NavDataInstanceForAgent == nullptr)
+		{
+			if (NavData->IsA(AAbstractNavData::StaticClass()) == false)
+			{
+				// ok, so this navigation agent doesn't have its navmesh registered yet, but do we want to support it?
+				bool bAgentSupported = false;
+
+				for (int32 AgentIndex = 0; AgentIndex < SupportedAgents.Num(); ++AgentIndex)
+				{
+					if (NavData->GetClass() == SupportedAgents[AgentIndex].NavigationDataClass && SupportedAgents[AgentIndex].IsEquivalent(NavConfig) == true)
+					{
+						// it's supported, then just in case it's not a precise match (IsEquivalent succeeds with some precision) 
+						// update NavData with supported Agent
+						bAgentSupported = true;
+
+						NavData->SetConfig(SupportedAgents[AgentIndex]);
+						AgentToNavDataMap.Add(SupportedAgents[AgentIndex], NavData);
+						NavData->SetSupportsDefaultAgent(AgentIndex == 0);
+						NavData->ProcessNavAreas(NavAreaClasses, AgentIndex);
+
+						OnNavDataRegisteredEvent.Broadcast(NavData);
+
+						NavDataSet.AddUnique(NavData);
+						NavData->OnRegistered();
+
+						break;
+					}
+				}
+				Result = bAgentSupported == true ? RegistrationSuccessful : RegistrationFailed_AgentNotValid;
+			}
+			else
+			{
+				// fake registration since it's a special navigation data type 
+				// and it would get discarded for not implementing any particular
+				// navigation agent
+				// Node that we don't add abstract navigation data to NavDataSet
+				NavData->OnRegistered();
+
+				Result = RegistrationSuccessful;
+			}
+		}
+		else if (NavDataInstanceForAgent == NavData)
+		{
+			ensure(NavDataSet.Find(NavData) != INDEX_NONE);
+			// let's treat double registration of the same nav data with the same agent as a success
+			Result = RegistrationSuccessful;
+		}
+		else
+		{
+			// otherwise specified agent type already has its navmesh implemented, fail redundant instance
+			Result = RegistrationFailed_AgentAlreadySupported;
+		}
+	}
+	else
+	{
+		Result = RegistrationFailed_AgentNotValid;
+	}
+
+	// @todo else might consider modifying this NavData to implement navigation for one of the supported agents
+	// care needs to be taken to not make it implement navigation for agent who's real implementation has 
+	// not been loaded yet.
+
+	return Result;
+}
+
+void UNavigationSystemV1::UnregisterNavData(ANavigationData* NavData)
+{
+	NavDataSet.RemoveSingle(NavData);
+
+	if (NavData == NULL)
+	{
+		return;
+	}
+
+	FScopeLock Lock(&NavDataRegistration);
+	NavData->OnUnregistered();
+}
+
+void UNavigationSystemV1::RegisterCustomLink(INavLinkCustomInterface& CustomLink)
+{
+	uint32 LinkId = CustomLink.GetLinkId();
+
+	// if there's already a link with that Id registered, assign new Id and mark dirty area
+	// this won't fix baked data in static navmesh (in game), but every other case will regenerate affected tiles 
+	if (CustomLinksMap.Contains(LinkId))
+	{
+		LinkId = INavLinkCustomInterface::GetUniqueId();
+		CustomLink.UpdateLinkId(LinkId);
+
+		UObject* CustomLinkOb = CustomLink.GetLinkOwner();
+		UActorComponent* OwnerComp = Cast<UActorComponent>(CustomLinkOb);
+		AActor* OwnerActor = OwnerComp ? OwnerComp->GetOwner() : Cast<AActor>(CustomLinkOb);
+
+		if (OwnerActor)
+		{
+			ENavLinkDirection::Type DummyDir = ENavLinkDirection::BothWays;
+			FVector RelativePtA, RelativePtB;
+			CustomLink.GetLinkData(RelativePtA, RelativePtB, DummyDir);
+
+			const FTransform OwnerActorTM = OwnerActor->GetTransform();
+			const FVector WorldPtA = OwnerActorTM.TransformPosition(RelativePtA);
+			const FVector WorldPtB = OwnerActorTM.TransformPosition(RelativePtB);
+
+			FBox LinkBounds(EForceInit::ForceInitToZero);
+			LinkBounds += WorldPtA;
+			LinkBounds += WorldPtB;
+
+			AddDirtyArea(LinkBounds, OctreeUpdate_Modifiers);
+		}
+	}
+
+	CustomLinksMap.Add(LinkId, FNavigationSystem::FCustomLinkOwnerInfo(&CustomLink));
+}
+
+void UNavigationSystemV1::UnregisterCustomLink(INavLinkCustomInterface& CustomLink)
+{
+	CustomLinksMap.Remove(CustomLink.GetLinkId());
+}
+
+INavLinkCustomInterface* UNavigationSystemV1::GetCustomLink(uint32 UniqueLinkId) const
+{
+	const FNavigationSystem::FCustomLinkOwnerInfo* LinkInfo = CustomLinksMap.Find(UniqueLinkId);
+	return (LinkInfo && LinkInfo->IsValid()) ? LinkInfo->LinkInterface : nullptr;
+}
+
+void UNavigationSystemV1::UpdateCustomLink(const INavLinkCustomInterface* CustomLink)
+{
+	for (TMap<FNavAgentProperties, TWeakObjectPtr<ANavigationData> >::TIterator It(AgentToNavDataMap); It; ++It)
+	{
+		ANavigationData* NavData = It.Value().Get();
+		if (NavData)
+		{
+			NavData->UpdateCustomLink(CustomLink);
+		}
+	}
+}
+
+void UNavigationSystemV1::RequestCustomLinkRegistering(INavLinkCustomInterface& CustomLink, UObject* OwnerOb)
+{
+	UNavigationSystemV1* NavSys = FNavigationSystem::GetCurrent<UNavigationSystemV1>(OwnerOb);
+	if (NavSys)
+	{
+		NavSys->RegisterCustomLink(CustomLink);
+	}
+	else
+	{
+		FScopeLock AccessLock(&CustomLinkRegistrationSection);
+		PendingCustomLinkRegistration.Add(&CustomLink, OwnerOb);
+	}
+}
+
+void UNavigationSystemV1::RequestCustomLinkUnregistering(INavLinkCustomInterface& CustomLink, UObject* OwnerOb)
+{
+	UNavigationSystemV1* NavSys = FNavigationSystem::GetCurrent<UNavigationSystemV1>(OwnerOb);
+	if (NavSys)
+	{
+		NavSys->UnregisterCustomLink(CustomLink);
+	}
+	else
+	{
+		FScopeLock AccessLock(&CustomLinkRegistrationSection);
+		PendingCustomLinkRegistration.Remove(&CustomLink);
+	}
+}
+
+void UNavigationSystemV1::RequestAreaUnregistering(UClass* NavAreaClass)
+{
+	for (TObjectIterator<UNavigationSystemV1> NavSysIt; NavSysIt; ++NavSysIt)
+	{
+		NavSysIt->UnregisterNavAreaClass(NavAreaClass);
+	}
+}
+
+void UNavigationSystemV1::UnregisterNavAreaClass(UClass* NavAreaClass)
+{
+	// remove from known areas
+	if (NavAreaClasses.Remove(NavAreaClass) > 0)
+	{
+		// notify navigation data
+		// notify existing nav data
+		OnNavigationAreaEvent(NavAreaClass, ENavAreaEvent::Unregistered);
+	}
+}
+
+void UNavigationSystemV1::RequestAreaRegistering(UClass* NavAreaClass)
+{
+	for (TObjectIterator<UNavigationSystemV1> NavSysIt; NavSysIt; ++NavSysIt)
+	{
+		NavSysIt->RegisterNavAreaClass(NavAreaClass);
+	}
+}
+
+void UNavigationSystemV1::RegisterNavAreaClass(UClass* AreaClass)
+{
+	// can't be null
+	if (AreaClass == NULL)
+	{
+		return;
+	}
+
+	// can't be abstract
+	if (AreaClass->HasAnyClassFlags(CLASS_Abstract))
+	{
+		return;
+	}
+
+	// special handling of blueprint based areas
+	if (AreaClass->HasAnyClassFlags(CLASS_CompiledFromBlueprint))
+	{
+		// can't be skeleton of blueprint class
+		if (AreaClass->GetName().Contains(TEXT("SKEL_")))
+		{
+			return;
+		}
+
+		// can't be class from Developers folder (won't be saved properly anyway)
+		const UPackage* Package = AreaClass->GetOutermost();
+		if (Package && Package->GetName().Contains(TEXT("/Developers/")))
+		{
+			return;
+		}
+	}
+
+	if (NavAreaClasses.Contains(AreaClass))
+	{
+		// Already added
+		return;
+	}
+
+	UNavArea* AreaClassCDO = GetMutableDefault<UNavArea>(AreaClass);
+	check(AreaClassCDO);
+
+	// initialize flags
+	AreaClassCDO->InitializeArea();
+
+	// add to know areas
+	NavAreaClasses.Add(AreaClass);
+
+	// notify existing nav data
+	OnNavigationAreaEvent(AreaClass, ENavAreaEvent::Registered);
+
+#if WITH_EDITOR
+	UNavAreaMeta_SwitchByAgent* SwitchByAgentCDO = Cast<UNavAreaMeta_SwitchByAgent>(AreaClassCDO);
+	// update area properties
+	if (SwitchByAgentCDO)
+	{
+		SwitchByAgentCDO->UpdateAgentConfig();
+	}
+#endif
+}
+
+void UNavigationSystemV1::OnNavigationAreaEvent(UClass* AreaClass, ENavAreaEvent::Type Event)
+{
+	// notify existing nav data
+	for (auto NavigationData : NavDataSet)
+	{
+		if (NavigationData != NULL && NavigationData->IsPendingKillPending() == false)
+		{
+			NavigationData->OnNavAreaEvent(AreaClass, Event);
+		}
+	}
+}
+
+int32 UNavigationSystemV1::GetSupportedAgentIndex(const ANavigationData* NavData) const
+{
+	if (SupportedAgents.Num() < 2)
+	{
+		return 0;
+	}
+
+	const FNavDataConfig& TestConfig = NavData->GetConfig();
+	for (int32 AgentIndex = 0; AgentIndex < SupportedAgents.Num(); AgentIndex++)
+	{
+		if (SupportedAgents[AgentIndex].IsEquivalent(TestConfig))
+		{
+			return AgentIndex;
+		}
+	}
+	
+	return INDEX_NONE;
+}
+
+int32 UNavigationSystemV1::GetSupportedAgentIndex(const FNavAgentProperties& NavAgent) const
+{
+	if (SupportedAgents.Num() < 2)
+	{
+		return 0;
+	}
+
+	for (int32 AgentIndex = 0; AgentIndex < SupportedAgents.Num(); AgentIndex++)
+	{
+		if (SupportedAgents[AgentIndex].IsEquivalent(NavAgent))
+		{
+			return AgentIndex;
+		}
+	}
+
+	return INDEX_NONE;
+}
+
+void UNavigationSystemV1::DescribeFilterFlags(UEnum* FlagsEnum) const
+{
+#if WITH_EDITOR
+	TArray<FString> FlagDesc;
+	FString EmptyStr;
+	FlagDesc.Init(EmptyStr, 16);
+
+	const int32 NumEnums = FMath::Min(16, FlagsEnum->NumEnums() - 1);	// skip _MAX
+	for (int32 FlagIndex = 0; FlagIndex < NumEnums; FlagIndex++)
+	{
+		FlagDesc[FlagIndex] = FlagsEnum->GetDisplayNameTextByIndex(FlagIndex).ToString();
+	}
+
+	DescribeFilterFlags(FlagDesc);
+#endif
+}
+
+void UNavigationSystemV1::DescribeFilterFlags(const TArray<FString>& FlagsDesc) const
+{
+#if WITH_EDITOR
+	const int32 MaxFlags = 16;
+	TArray<FString> UseDesc = FlagsDesc;
+
+	FString EmptyStr;
+	while (UseDesc.Num() < MaxFlags)
+	{
+		UseDesc.Add(EmptyStr);
+	}
+
+	// get special value from recast's navmesh
+#if WITH_RECAST
+	uint16 NavLinkFlag = ARecastNavMesh::GetNavLinkFlag();
+	for (int32 FlagIndex = 0; FlagIndex < MaxFlags; FlagIndex++)
+	{
+		if ((NavLinkFlag >> FlagIndex) & 1)
+		{
+			UseDesc[FlagIndex] = TEXT("Navigation link");
+			break;
+		}
+	}
+#endif
+
+	// setup properties
+	UStructProperty* StructProp1 = FindField<UStructProperty>(UNavigationQueryFilter::StaticClass(), TEXT("IncludeFlags"));
+	UStructProperty* StructProp2 = FindField<UStructProperty>(UNavigationQueryFilter::StaticClass(), TEXT("ExcludeFlags"));
+	check(StructProp1);
+	check(StructProp2);
+
+	UStruct* Structs[] = { StructProp1->Struct, StructProp2->Struct };
+	const FString CustomNameMeta = TEXT("DisplayName");
+
+	for (int32 StructIndex = 0; StructIndex < ARRAY_COUNT(Structs); StructIndex++)
+	{
+		for (int32 FlagIndex = 0; FlagIndex < MaxFlags; FlagIndex++)
+		{
+			FString PropName = FString::Printf(TEXT("bNavFlag%d"), FlagIndex);
+			UProperty* Prop = FindField<UProperty>(Structs[StructIndex], *PropName);
+			check(Prop);
+
+			if (UseDesc[FlagIndex].Len())
+			{
+				Prop->SetPropertyFlags(CPF_Edit);
+				Prop->SetMetaData(*CustomNameMeta, *UseDesc[FlagIndex]);
+			}
+			else
+			{
+				Prop->ClearPropertyFlags(CPF_Edit);
+			}
+		}
+	}
+
+#endif
+}
+
+void UNavigationSystemV1::ResetCachedFilter(TSubclassOf<UNavigationQueryFilter> FilterClass)
+{
+	for (int32 NavDataIndex = 0; NavDataIndex < NavDataSet.Num(); NavDataIndex++)
+	{
+		if (NavDataSet[NavDataIndex])
+		{
+			NavDataSet[NavDataIndex]->RemoveQueryFilter(FilterClass);
+		}
+	}
+}
+
+UNavigationSystemV1* UNavigationSystemV1::CreateNavigationSystem(UWorld* WorldOwner)
+{
+	UNavigationSystemV1* NavSys = NULL;
+
+	// create navigation system for editor and server targets, but remove it from game clients
+	if (WorldOwner && (*GEngine->NavigationSystemClass != nullptr) 
+		&& (GEngine->NavigationSystemClass->GetDefaultObject<UNavigationSystemV1>()->bAllowClientSideNavigation || WorldOwner->GetNetMode() != NM_Client))
+	{
+		AWorldSettings* WorldSettings = WorldOwner->GetWorldSettings();
+		if (WorldSettings == NULL || WorldSettings->IsNavigationSystemEnabled())
+		{
+			NavSys = NewObject<UNavigationSystemV1>(WorldOwner, GEngine->NavigationSystemClass);		
+			WorldOwner->SetNavigationSystem(NavSys);
+		}
+	}
+
+	return NavSys;
+}
+
+void UNavigationSystemV1::InitializeForWorld(UWorld& World, FNavigationSystemRunMode Mode)
+{
+	OnWorldInitDone(Mode);
+}
+
+UNavigationSystemV1* UNavigationSystemV1::GetCurrent(UWorld* World)
+{
+	return FNavigationSystem::GetCurrent<UNavigationSystemV1>(World);
+}
+
+UNavigationSystemV1* UNavigationSystemV1::GetCurrent(UObject* WorldContextObject)
+{
+	return FNavigationSystem::GetCurrent<UNavigationSystemV1>(WorldContextObject);
+}
+
+ANavigationData* UNavigationSystemV1::GetNavDataWithID(const uint16 NavDataID) const
+{
+	for (int32 NavDataIndex = 0; NavDataIndex < NavDataSet.Num(); ++NavDataIndex)
+	{
+		const ANavigationData* NavData = NavDataSet[NavDataIndex];
+		if (NavData != NULL && NavData->GetNavDataUniqueID() == NavDataID)
+		{
+			return const_cast<ANavigationData*>(NavData);
+		}
+	}
+
+	return NULL;
+}
+
+void UNavigationSystemV1::AddDirtyArea(const FBox& NewArea, int32 Flags)
+{
+	if (Flags > 0 && bCanAccumulateDirtyAreas)
+	{
+		DirtyAreas.Add(FNavigationDirtyArea(NewArea, Flags));
+	}
+#if !UE_BUILD_SHIPPING
+	bDirtyAreasReportedWhileAccumulationLocked = bDirtyAreasReportedWhileAccumulationLocked || (Flags > 0 && !bCanAccumulateDirtyAreas);
+#endif // !UE_BUILD_SHIPPING
+}
+
+void UNavigationSystemV1::AddDirtyAreas(const TArray<FBox>& NewAreas, int32 Flags)
+{ 
+	for (int32 NewAreaIndex = 0; NewAreaIndex < NewAreas.Num(); NewAreaIndex++)
+	{
+		AddDirtyArea(NewAreas[NewAreaIndex], Flags);
+	}
+}
+
+bool UNavigationSystemV1::HasDirtyAreasQueued() const
+{
+	return DirtyAreas.Num() > 0;
+}
+
+int32 GetDirtyFlagHelper(int32 UpdateFlags, int32 DefaultValue)
+{
+	return ((UpdateFlags & UNavigationSystemV1::OctreeUpdate_Geometry) != 0) ? ENavigationDirtyFlag::All :
+		((UpdateFlags & UNavigationSystemV1::OctreeUpdate_Modifiers) != 0) ? ENavigationDirtyFlag::DynamicModifier :		
+		DefaultValue;
+}
+
+FSetElementId UNavigationSystemV1::RegisterNavOctreeElement(UObject* ElementOwner, INavRelevantInterface* ElementInterface, int32 UpdateFlags)
+{
+	FSetElementId SetId;
+
+#if WITH_EDITOR
+	if (IsNavigationRegisterLocked())
+	{
+		return SetId;
+	}
+#endif
+
+	if (NavOctree.IsValid() == false || ElementOwner == NULL || ElementInterface == NULL)
+	{
+		return SetId;
+	}
+
+	if (IsNavigationOctreeLocked())
+	{
+		UE_LOG(LogNavOctree, Log, TEXT("IGNORE(RegisterNavOctreeElement) %s"), *GetPathNameSafe(ElementOwner));
+		return SetId;
+	}
+
+	const bool bIsRelevant = ElementInterface->IsNavigationRelevant();
+	UE_LOG(LogNavOctree, Log, TEXT("REG %s %s"), *GetNameSafe(ElementOwner), bIsRelevant ? TEXT("[relevant]") : TEXT(""));
+
+	if (bIsRelevant)
+	{
+		bool bCanAdd = false;
+
+		UObject* ParentNode = ElementInterface->GetNavigationParent();
+		if (ParentNode)
+		{
+			OctreeChildNodesMap.AddUnique(ParentNode, FWeakObjectPtr(ElementOwner));
+			bCanAdd = true;
+		}
+		else
+		{
+			const FOctreeElementId* ElementId = GetObjectsNavOctreeId(ElementOwner);
+			bCanAdd = (ElementId == NULL);
+		}
+
+		if (bCanAdd)
+		{
+			FNavigationDirtyElement UpdateInfo(ElementOwner, ElementInterface, GetDirtyFlagHelper(UpdateFlags, 0));
+
+			SetId = PendingOctreeUpdates.FindId(UpdateInfo);
+			if (SetId.IsValidId())
+			{
+				// make sure this request stays, in case it has been invalidated already
+				PendingOctreeUpdates[SetId] = UpdateInfo;
+			}
+			else
+			{
+				SetId = PendingOctreeUpdates.Add(UpdateInfo);
+			}
+		}
+	}
+
+	return SetId;
+}
+
+void UNavigationSystemV1::AddElementToNavOctree(const FNavigationDirtyElement& DirtyElement)
+{
+	// handle invalidated requests first
+	if (DirtyElement.bInvalidRequest)
+	{
+		if (DirtyElement.bHasPrevData)
+		{
+			AddDirtyArea(DirtyElement.PrevBounds, DirtyElement.PrevFlags);
+		}
+		
+		return;
+	}
+
+	UObject* ElementOwner = DirtyElement.Owner.Get();
+	if (ElementOwner == NULL || ElementOwner->IsPendingKill() || DirtyElement.NavInterface == NULL)
+	{
+		return;
+	}
+
+	FNavigationOctreeElement GeneratedData(*ElementOwner);
+	const FBox ElementBounds = DirtyElement.NavInterface->GetNavigationBounds();
+
+	UObject* NavigationParent = DirtyElement.NavInterface->GetNavigationParent();
+	if (NavigationParent)
+	{
+		// check if parent node is waiting in queue
+		const FSetElementId ParentRequestId = PendingOctreeUpdates.FindId(FNavigationDirtyElement(NavigationParent));
+		const FOctreeElementId* ParentId = GetObjectsNavOctreeId(NavigationParent);
+		if (ParentRequestId.IsValidId() && ParentId == NULL)
+		{
+			FNavigationDirtyElement& ParentNode = PendingOctreeUpdates[ParentRequestId];
+			AddElementToNavOctree(ParentNode);
+
+			// mark as invalid so it won't be processed twice
+			ParentNode.bInvalidRequest = true;
+		}
+
+		const FOctreeElementId* UseParentId = ParentId ? ParentId : GetObjectsNavOctreeId(NavigationParent);
+		if (UseParentId && NavOctree->IsValidElementId(*UseParentId))
+		{
+			UE_LOG(LogNavOctree, Log, TEXT("ADD %s to %s"), *GetNameSafe(ElementOwner), *GetNameSafe(NavigationParent));
+			NavOctree->AppendToNode(*UseParentId, DirtyElement.NavInterface, ElementBounds, GeneratedData);
+		}
+		else 
+		{
+			UE_LOG(LogNavOctree, Warning, TEXT("Can't add node [%s] - parent [%s] not found in octree!"), *GetNameSafe(ElementOwner), *GetNameSafe(NavigationParent));
+		}
+	}
+	else
+	{
+		UE_LOG(LogNavOctree, Log, TEXT("ADD %s"), *GetNameSafe(ElementOwner));
+		NavOctree->AddNode(ElementOwner, DirtyElement.NavInterface, ElementBounds, GeneratedData);
+	}
+
+	const FBox BBox = GeneratedData.Bounds.GetBox();
+	const bool bValidBBox = BBox.IsValid && !BBox.GetSize().IsNearlyZero();
+
+	if (bValidBBox && !GeneratedData.IsEmpty())
+	{
+		const int32 DirtyFlag = DirtyElement.FlagsOverride ? DirtyElement.FlagsOverride : GeneratedData.Data->GetDirtyFlag();
+		AddDirtyArea(BBox, DirtyFlag);
+	}
+}
+
+bool UNavigationSystemV1::GetNavOctreeElementData(UObject* NodeOwner, int32& DirtyFlags, FBox& DirtyBounds)
+{
+	const FOctreeElementId* ElementId = GetObjectsNavOctreeId(NodeOwner);
+	if (ElementId != NULL)
+	{
+		if (NavOctree->IsValidElementId(*ElementId))
+		{
+			// mark area occupied by given actor as dirty
+			FNavigationOctreeElement& ElementData = NavOctree->GetElementById(*ElementId);
+			DirtyFlags = ElementData.Data->GetDirtyFlag();
+			DirtyBounds = ElementData.Bounds.GetBox();
+			return true;
+		}
+	}
+
+	return false;
+}
+
+void UNavigationSystemV1::UnregisterNavOctreeElement(UObject* ElementOwner, INavRelevantInterface* ElementInterface, int32 UpdateFlags)
+{
+#if WITH_EDITOR
+	if (IsNavigationUnregisterLocked())
+	{
+		return;
+	}
+#endif
+
+	if (NavOctree.IsValid() == false || ElementOwner == NULL || ElementInterface == NULL)
+	{
+		return;
+	}
+
+	if (IsNavigationOctreeLocked())
+	{
+		UE_LOG(LogNavOctree, Log, TEXT("IGNORE(UnregisterNavOctreeElement) %s"), *GetPathNameSafe(ElementOwner));
+		return;
+	}
+
+	const FOctreeElementId* ElementId = GetObjectsNavOctreeId(ElementOwner);
+	UE_LOG(LogNavOctree, Log, TEXT("UNREG %s %s"), *GetNameSafe(ElementOwner), ElementId ? TEXT("[exists]") : TEXT(""));
+
+	if (ElementId != NULL)
+	{
+		RemoveNavOctreeElementId(*ElementId, UpdateFlags);
+		RemoveObjectsNavOctreeId(ElementOwner);
+	}
+	else
+	{
+		const bool bCanRemoveChildNode = (UpdateFlags & OctreeUpdate_ParentChain) == 0;
+		UObject* ParentNode = ElementInterface->GetNavigationParent();
+		if (ParentNode && bCanRemoveChildNode)
+		{
+			// if node has navigation parent (= doesn't exists in octree on its own)
+			// and it's not part of parent chain update
+			// remove it from map and force update on parent to rebuild octree element
+
+			OctreeChildNodesMap.RemoveSingle(ParentNode, FWeakObjectPtr(ElementOwner));
+			UpdateNavOctreeParentChain(ParentNode);
+		}
+	}
+
+	// mark pending update as invalid, it will be dirtied according to currently active settings
+	const bool bCanInvalidateQueue = (UpdateFlags & OctreeUpdate_Refresh) == 0;
+	if (bCanInvalidateQueue)
+	{
+		const FSetElementId RequestId = PendingOctreeUpdates.FindId(FNavigationDirtyElement(ElementOwner));
+		if (RequestId.IsValidId())
+		{
+			PendingOctreeUpdates[RequestId].bInvalidRequest = true;
+		}
+	}
+}
+
+void UNavigationSystemV1::RemoveNavOctreeElementId(const FOctreeElementId& ElementId, int32 UpdateFlags)
+{
+	if (NavOctree->IsValidElementId(ElementId))
+	{
+		// mark area occupied by given actor as dirty
+		FNavigationOctreeElement& ElementData = NavOctree->GetElementById(ElementId);
+		const int32 DirtyFlag = GetDirtyFlagHelper(UpdateFlags, ElementData.Data->GetDirtyFlag());
+		AddDirtyArea(ElementData.Bounds.GetBox(), DirtyFlag);
+		NavOctree->RemoveNode(ElementId);
+	}
+}
+
+const FNavigationRelevantData* UNavigationSystemV1::GetDataForObject(const UObject& Object) const
+{
+	check(NavOctree.IsValid());
+
+	const FOctreeElementId* OctreeID = GetObjectsNavOctreeId(&Object);
+
+	if (OctreeID != nullptr && OctreeID->IsValidId() == true)
+	{
+		return NavOctree->GetDataForID(*OctreeID);
+	}
+
+	return nullptr;
+}
+
+void UNavigationSystemV1::UpdateActorInNavOctree(AActor& Actor)
+{
+	if (IsNavigationSystemStatic())
+	{
+		return;
+	}
+	SCOPE_CYCLE_COUNTER(STAT_DebugNavOctree);
+
+	INavRelevantInterface* NavElement = Cast<INavRelevantInterface>(&Actor);
+	if (NavElement)
+	{
+		UNavigationSystemV1* NavSys = FNavigationSystem::GetCurrent<UNavigationSystemV1>(Actor.GetWorld());
+		if (NavSys)
+		{
+			NavSys->UpdateNavOctreeElement(&Actor, NavElement, OctreeUpdate_Default);
+		}
+	}
+}
+
+void UNavigationSystemV1::UpdateComponentInNavOctree(UActorComponent& Comp)
+{
+	SCOPE_CYCLE_COUNTER(STAT_DebugNavOctree);
+
+	if (ShouldUpdateNavOctreeOnComponentChange() == false)
+	{
+		return;
+	}
+
+	// special case for early out: use cached nav relevancy
+	if (Comp.bNavigationRelevant == true)
+	{
+		INavRelevantInterface* NavElement = Cast<INavRelevantInterface>(&Comp);
+		if (NavElement)
+		{
+			AActor* OwnerActor = Comp.GetOwner();
+			if (OwnerActor)
+			{
+				UNavigationSystemV1* NavSys = FNavigationSystem::GetCurrent<UNavigationSystemV1>(OwnerActor->GetWorld());
+				if (NavSys)
+				{
+					if (OwnerActor->IsComponentRelevantForNavigation(&Comp))
+					{
+						NavSys->UpdateNavOctreeElement(&Comp, NavElement, OctreeUpdate_Default);
+					}
+					else
+					{
+						NavSys->UnregisterNavOctreeElement(&Comp, NavElement, OctreeUpdate_Default);
+					}
+				}
+			}
+		}
+	}
+	else if (Comp.CanEverAffectNavigation()) 
+	{
+		// could have been relevant before and not it isn't. Need to check if there's an octree element ID for it
+		INavRelevantInterface* NavElement = Cast<INavRelevantInterface>(&Comp);
+		if (NavElement)
+		{
+			AActor* OwnerActor = Comp.GetOwner();
+			if (OwnerActor)
+			{
+				UNavigationSystemV1* NavSys = FNavigationSystem::GetCurrent<UNavigationSystemV1>(OwnerActor->GetWorld());
+				if (NavSys)
+				{
+					NavSys->UnregisterNavOctreeElement(&Comp, NavElement, OctreeUpdate_Default);
+				}
+			}
+		}
+	}
+}
+
+void UNavigationSystemV1::UpdateActorAndComponentsInNavOctree(AActor& Actor, bool bUpdateAttachedActors)
+{
+	UpdateActorInNavOctree(Actor);
+		
+	TInlineComponentArray<UActorComponent*> Components;
+	Actor.GetComponents(Components);
+
+	for (int32 ComponentIndex = 0; ComponentIndex < Components.Num(); ComponentIndex++)
+	{
+		check(Components[ComponentIndex]);
+		UpdateComponentInNavOctree(*Components[ComponentIndex]);
+	}
+
+	if (bUpdateAttachedActors)
+	{
+		UpdateAttachedActorsInNavOctree(Actor);
+	}
+}
+
+void UNavigationSystemV1::UpdateNavOctreeAfterMove(USceneComponent* Comp)
+{
+	AActor* OwnerActor = Comp->GetOwner();
+	if (OwnerActor && OwnerActor->GetRootComponent() == Comp)
+	{
+		UpdateActorAndComponentsInNavOctree(*OwnerActor, true);
+	}
+}
+
+void UNavigationSystemV1::UpdateAttachedActorsInNavOctree(AActor& RootActor)
+{
+	TArray<AActor*> UniqueAttachedActors;
+	UniqueAttachedActors.Add(&RootActor);
+
+	TArray<AActor*> TempAttachedActors;
+	for (int32 ActorIndex = 0; ActorIndex < UniqueAttachedActors.Num(); ++ActorIndex)
+	{
+		check(UniqueAttachedActors[ActorIndex]);
+		// find all attached actors
+		UniqueAttachedActors[ActorIndex]->GetAttachedActors(TempAttachedActors);
+		
+		for (int32 AttachmentIndex = 0; AttachmentIndex < TempAttachedActors.Num(); ++AttachmentIndex)
+		{
+			// and store the ones we don't know about yet
+			UniqueAttachedActors.AddUnique(TempAttachedActors[AttachmentIndex]);
+		}
+	}
+	
+	// skipping the first item since that's the root, and we just care about the attached actors
+	for (int32 ActorIndex = 1; ActorIndex < UniqueAttachedActors.Num(); ++ActorIndex)
+	{
+		UpdateActorAndComponentsInNavOctree(*UniqueAttachedActors[ActorIndex], /*bUpdateAttachedActors = */false);
+	}
+}
+
+void UNavigationSystemV1::UpdateNavOctreeBounds(AActor* Actor)
+{
+	TInlineComponentArray<UActorComponent*> Components;
+	Actor->GetComponents(Components);
+
+	for (int32 Idx = 0; Idx < Components.Num(); Idx++)
+	{
+		INavRelevantInterface* NavElement = Cast<INavRelevantInterface>(Components[Idx]);
+		if (NavElement)
+		{
+			NavElement->UpdateNavigationBounds();
+		}
+	}
+}
+
+void UNavigationSystemV1::ClearNavOctreeAll(AActor* Actor)
+{
+	if (Actor)
+	{
+		OnActorUnregistered(Actor);
+
+		TInlineComponentArray<UActorComponent*> Components;
+		Actor->GetComponents(Components);
+
+		for (int32 Idx = 0; Idx < Components.Num(); Idx++)
+		{
+			OnComponentUnregistered(Components[Idx]);
+		}
+	}
+}
+
+void UNavigationSystemV1::UpdateNavOctreeElement(UObject* ElementOwner, INavRelevantInterface* ElementInterface, int32 UpdateFlags)
+{
+	INC_DWORD_STAT(STAT_Navigation_UpdateNavOctree);
+
+	if (IsNavigationOctreeLocked())
+	{
+		UE_LOG(LogNavOctree, Log, TEXT("IGNORE(UpdateNavOctreeElement) %s"), *GetPathNameSafe(ElementOwner));
+		return;
+	}
+
+	// grab existing octree data
+	FBox CurrentBounds;
+	int32 CurrentFlags;
+	const bool bAlreadyExists = GetNavOctreeElementData(ElementOwner, CurrentFlags, CurrentBounds);
+
+	// don't invalidate pending requests
+	UpdateFlags |= OctreeUpdate_Refresh;
+
+	// always try to unregister, even if element owner doesn't exists in octree (parent nodes)
+	UnregisterNavOctreeElement(ElementOwner, ElementInterface, UpdateFlags);
+
+	const FSetElementId RequestId = RegisterNavOctreeElement(ElementOwner, ElementInterface, UpdateFlags);
+
+	// add original data to pending registration request
+	// so it could be dirtied properly when system receive unregister request while actor is still queued
+	if (RequestId.IsValidId())
+	{
+		FNavigationDirtyElement& UpdateInfo = PendingOctreeUpdates[RequestId];
+		UpdateInfo.PrevFlags = CurrentFlags;
+		if (UpdateInfo.PrevBounds.IsValid)
+		{
+			// Is we have something stored already we want to 
+			// sum it up, since we care about the whole bounding
+			// box of changes that potentially took place
+			UpdateInfo.PrevBounds += CurrentBounds;
+		}
+		else
+		{
+			UpdateInfo.PrevBounds = CurrentBounds;
+		}
+		UpdateInfo.bHasPrevData = bAlreadyExists;
+	}
+
+	UpdateNavOctreeParentChain(ElementOwner, /*bSkipElementOwnerUpdate=*/ true);
+}
+
+void UNavigationSystemV1::UpdateNavOctreeParentChain(UObject* ElementOwner, bool bSkipElementOwnerUpdate)
+{
+	const int32 UpdateFlags = OctreeUpdate_ParentChain | OctreeUpdate_Refresh;
+
+	TArray<FWeakObjectPtr> ChildNodes;
+	OctreeChildNodesMap.MultiFind(ElementOwner, ChildNodes);
+
+	if (ChildNodes.Num() == 0)
+	{
+		if (bSkipElementOwnerUpdate == false)
+		{
+			INavRelevantInterface* ElementInterface = Cast<INavRelevantInterface>(ElementOwner);
+			UpdateNavOctreeElement(ElementOwner, ElementInterface, UpdateFlags);
+		}
+		return;
+	}
+
+	INavRelevantInterface* ElementInterface = Cast<INavRelevantInterface>(ElementOwner);
+	TArray<INavRelevantInterface*> ChildNavInterfaces;
+	ChildNavInterfaces.AddZeroed(ChildNodes.Num());
+	
+	for (int32 Idx = 0; Idx < ChildNodes.Num(); Idx++)
+	{
+		if (ChildNodes[Idx].IsValid())
+		{
+			UObject* ChildNodeOb = ChildNodes[Idx].Get();
+			ChildNavInterfaces[Idx] = Cast<INavRelevantInterface>(ChildNodeOb);
+			UnregisterNavOctreeElement(ChildNodeOb, ChildNavInterfaces[Idx], UpdateFlags);
+		}
+	}
+
+	if (bSkipElementOwnerUpdate == false)
+	{
+		UnregisterNavOctreeElement(ElementOwner, ElementInterface, UpdateFlags);
+		RegisterNavOctreeElement(ElementOwner, ElementInterface, UpdateFlags);
+	}
+
+	for (int32 Idx = 0; Idx < ChildNodes.Num(); Idx++)
+	{
+		if (ChildNodes[Idx].IsValid())
+		{
+			RegisterNavOctreeElement(ChildNodes[Idx].Get(), ChildNavInterfaces[Idx], UpdateFlags);
+		}
+	}
+}
+
+bool UNavigationSystemV1::UpdateNavOctreeElementBounds(UActorComponent* Comp, const FBox& NewBounds, const FBox& DirtyArea)
+{
+	const FOctreeElementId* ElementId = GetObjectsNavOctreeId(Comp);
+	if (ElementId && ElementId->IsValidId())
+	{
+		NavOctree->UpdateNode(*ElementId, NewBounds);
+		
+		// Add dirty area
+		if (DirtyArea.IsValid)
+		{
+			ElementId = GetObjectsNavOctreeId(Comp);
+			if (ElementId && ElementId->IsValidId())
+			{
+				FNavigationOctreeElement& ElementData = NavOctree->GetElementById(*ElementId);
+				AddDirtyArea(DirtyArea, ElementData.Data->GetDirtyFlag());
+			}
+		}
+
+		return true;
+	}
+	
+	return false;
+}
+
+void UNavigationSystemV1::OnComponentRegistered(UActorComponent* Comp)
+{
+	if (IsNavigationSystemStatic())
+	{
+		return;
+	}
+
+	SCOPE_CYCLE_COUNTER(STAT_DebugNavOctree);
+	INavRelevantInterface* NavInterface = Cast<INavRelevantInterface>(Comp);
+	if (NavInterface)
+	{
+		AActor* OwnerActor = Comp ? Comp->GetOwner() : NULL;
+		if (OwnerActor && OwnerActor->IsComponentRelevantForNavigation(Comp))
+		{
+			UNavigationSystemV1* NavSys = FNavigationSystem::GetCurrent<UNavigationSystemV1>(OwnerActor->GetWorld());
+			if (NavSys)
+			{
+				NavSys->RegisterNavOctreeElement(Comp, NavInterface, OctreeUpdate_Default);
+			}
+		}
+	}
+}
+
+void UNavigationSystemV1::OnComponentUnregistered(UActorComponent* Comp)
+{
+	if (IsNavigationSystemStatic())
+	{
+		return;
+	}
+
+	SCOPE_CYCLE_COUNTER(STAT_DebugNavOctree);
+	INavRelevantInterface* NavInterface = Cast<INavRelevantInterface>(Comp);
+	if (NavInterface)
+	{
+		AActor* OwnerActor = Comp ? Comp->GetOwner() : NULL;
+		if (OwnerActor)
+		{
+			// skip IsComponentRelevantForNavigation check, it's only for adding new stuff
+
+			UNavigationSystemV1* NavSys = FNavigationSystem::GetCurrent<UNavigationSystemV1>(OwnerActor->GetWorld());
+			if (NavSys)
+			{
+				NavSys->UnregisterNavOctreeElement(Comp, NavInterface, OctreeUpdate_Default);
+			}
+		}
+	}
+}
+
+void UNavigationSystemV1::OnActorRegistered(AActor* Actor)
+{
+	if (IsNavigationSystemStatic())
+	{
+		return;
+	}
+
+	SCOPE_CYCLE_COUNTER(STAT_DebugNavOctree);
+	INavRelevantInterface* NavInterface = Cast<INavRelevantInterface>(Actor);
+	if (NavInterface)
+	{
+		UNavigationSystemV1* NavSys = FNavigationSystem::GetCurrent<UNavigationSystemV1>(Actor->GetWorld());
+		if (NavSys)
+		{
+			NavSys->RegisterNavOctreeElement(Actor, NavInterface, OctreeUpdate_Default);
+		}
+	}
+}
+
+void UNavigationSystemV1::OnActorUnregistered(AActor* Actor)
+{
+	if (IsNavigationSystemStatic())
+	{
+		return;
+	}
+
+	SCOPE_CYCLE_COUNTER(STAT_DebugNavOctree);
+	INavRelevantInterface* NavInterface = Cast<INavRelevantInterface>(Actor);
+	if (NavInterface)
+	{
+		UNavigationSystemV1* NavSys = FNavigationSystem::GetCurrent<UNavigationSystemV1>(Actor->GetWorld());
+		if (NavSys)
+		{
+			NavSys->UnregisterNavOctreeElement(Actor, NavInterface, OctreeUpdate_Default);
+		}
+	}
+}
+
+void UNavigationSystemV1::FindElementsInNavOctree(const FBox& QueryBox, const FNavigationOctreeFilter& Filter, TArray<FNavigationOctreeElement>& Elements)
+{
+	if (NavOctree.IsValid() == false)
+	{
+		UE_LOG(LogNavigation, Warning, TEXT("UNavigationSystemV1::FindElementsInNavOctree gets called while NavOctree is null"));
+		return;
+	}
+	
+	for (FNavigationOctree::TConstElementBoxIterator<> It(*NavOctree, QueryBox); It.HasPendingElements(); It.Advance())
+	{
+		const FNavigationOctreeElement& Element = It.GetCurrentElement();
+		if (Element.IsMatchingFilter(Filter))
+		{
+			Elements.Add(Element);
+		}		
+	}
+}
+
+void UNavigationSystemV1::ReleaseInitialBuildingLock()
+{
+	RemoveNavigationBuildLock(ENavigationBuildLock::InitialLock);
+}
+
+void UNavigationSystemV1::InitializeLevelCollisions()
+{
+	if (IsNavigationSystemStatic())
+	{
+		bInitialLevelsAdded = true;
+		return;
+	}
+
+	UWorld* World = GetWorld();
+	if (!bInitialLevelsAdded && FNavigationSystem::GetCurrent<UNavigationSystemV1>(World) == this)
+	{
+		// Process all visible levels
+		const auto& Levels = World->GetLevels();
+		for (ULevel* Level : Levels)
+		{
+			if (Level->bIsVisible)
+			{
+				AddLevelCollisionToOctree(Level);
+			}
+		}
+
+		bInitialLevelsAdded = true;
+	}
+}
+
+#if WITH_EDITOR
+void UNavigationSystemV1::UpdateLevelCollision(ULevel* InLevel)
+{
+	if (InLevel != NULL)
+	{
+		UWorld* World = GetWorld();
+		OnLevelRemovedFromWorld(InLevel, World);
+		OnLevelAddedToWorld(InLevel, World);
+	}
+}
+
+void UNavigationSystemV1::OnEditorModeChanged(FEdMode* Mode, bool IsEntering)
+{
+	if (Mode == NULL)
+	{
+		return;
+	}
+
+	if (IsEntering == false && Mode->GetID() == FBuiltinEditorModes::EM_Geometry)
+	{
+		// check if any of modified brushes belongs to an ANavMeshBoundsVolume
+		FEdModeGeometry* GeometryMode = (FEdModeGeometry*)Mode;
+		for (auto GeomObjectIt = GeometryMode->GeomObjectItor(); GeomObjectIt; GeomObjectIt++)
+		{
+			ANavMeshBoundsVolume* Volume = Cast<ANavMeshBoundsVolume>((*GeomObjectIt)->GetActualBrush());
+			if (Volume)
+			{
+				OnNavigationBoundsUpdated(Volume);
+			}
+		}
+	}
+}
+#endif
+
+void UNavigationSystemV1::OnNavigationBoundsUpdated(ANavMeshBoundsVolume* NavVolume)
+{
+	if (NavVolume == nullptr || IsNavigationSystemStatic())
+	{
+		return;
+	}
+
+	FNavigationBoundsUpdateRequest UpdateRequest;
+	UpdateRequest.NavBounds.UniqueID = NavVolume->GetUniqueID();
+	UpdateRequest.NavBounds.AreaBox = NavVolume->GetComponentsBoundingBox(true);
+	UpdateRequest.NavBounds.Level = NavVolume->GetLevel();
+	UpdateRequest.NavBounds.SupportedAgents = NavVolume->SupportedAgents;
+	
+	UpdateRequest.UpdateRequest = FNavigationBoundsUpdateRequest::Updated;
+	AddNavigationBoundsUpdateRequest(UpdateRequest);
+}
+
+void UNavigationSystemV1::OnNavigationBoundsAdded(ANavMeshBoundsVolume* NavVolume)
+{
+	if (NavVolume == nullptr || IsNavigationSystemStatic())
+	{
+		return;
+	}
+
+	FNavigationBoundsUpdateRequest UpdateRequest;
+	UpdateRequest.NavBounds.UniqueID = NavVolume->GetUniqueID();
+	UpdateRequest.NavBounds.AreaBox = NavVolume->GetComponentsBoundingBox(true);
+	UpdateRequest.NavBounds.Level = NavVolume->GetLevel();
+	UpdateRequest.NavBounds.SupportedAgents = NavVolume->SupportedAgents;
+
+	UpdateRequest.UpdateRequest = FNavigationBoundsUpdateRequest::Added;
+	AddNavigationBoundsUpdateRequest(UpdateRequest);
+}
+
+void UNavigationSystemV1::OnNavigationBoundsRemoved(ANavMeshBoundsVolume* NavVolume)
+{
+	if (NavVolume == nullptr || IsNavigationSystemStatic())
+	{
+		return;
+	}
+	
+	FNavigationBoundsUpdateRequest UpdateRequest;
+	UpdateRequest.NavBounds.UniqueID = NavVolume->GetUniqueID();
+	UpdateRequest.NavBounds.AreaBox = NavVolume->GetComponentsBoundingBox(true);
+	UpdateRequest.NavBounds.Level = NavVolume->GetLevel();
+	UpdateRequest.NavBounds.SupportedAgents = NavVolume->SupportedAgents;
+
+	UpdateRequest.UpdateRequest = FNavigationBoundsUpdateRequest::Removed;
+	AddNavigationBoundsUpdateRequest(UpdateRequest);
+}
+
+void UNavigationSystemV1::AddNavigationBoundsUpdateRequest(const FNavigationBoundsUpdateRequest& UpdateRequest)
+{
+	int32 ExistingIdx = PendingNavBoundsUpdates.IndexOfByPredicate([&](const FNavigationBoundsUpdateRequest& Element) {
+		return UpdateRequest.NavBounds.UniqueID == Element.NavBounds.UniqueID;
+	});
+
+	if (ExistingIdx != INDEX_NONE)
+	{
+		// catch the case where the bounds was removed and immediately re-added with the same bounds as before
+		// in that case, we can cancel any update at all
+		bool bCanCancelUpdate = false;
+		if (PendingNavBoundsUpdates[ExistingIdx].UpdateRequest == FNavigationBoundsUpdateRequest::Removed && UpdateRequest.UpdateRequest == FNavigationBoundsUpdateRequest::Added)
+		{
+			for (TSet<FNavigationBounds>::TConstIterator It(RegisteredNavBounds); It; ++It)
+			{
+				if (*It == UpdateRequest.NavBounds)
+				{
+					bCanCancelUpdate = true;
+					break;
+				}
+			}
+		}
+		if (bCanCancelUpdate)
+		{
+			PendingNavBoundsUpdates.RemoveAt(ExistingIdx);
+		}
+		else
+		{
+			// Overwrite any previous updates
+			PendingNavBoundsUpdates[ExistingIdx] = UpdateRequest;
+		}
+	}
+	else
+	{
+		PendingNavBoundsUpdates.Add(UpdateRequest);
+	}
+}
+
+void UNavigationSystemV1::PerformNavigationBoundsUpdate(const TArray<FNavigationBoundsUpdateRequest>& UpdateRequests)
+{
+	if (NavDataSet.Num() == 0)
+	{
+		//TODO: will hitch when user places first navigation volume in the world 
+		
+		if (NavDataRegistrationQueue.Num() > 0)
+		{
+			ProcessRegistrationCandidates();
+		}
+
+		if (NavDataSet.Num() == 0)
+		{
+			SpawnMissingNavigationData();
+			ProcessRegistrationCandidates();
+		}
+				
+		ConditionalPopulateNavOctree();
+	}
+	
+	// Create list of areas that needs to be updated
+	TArray<FBox> UpdatedAreas;
+	for (const FNavigationBoundsUpdateRequest& Request : UpdateRequests)
+	{
+		FSetElementId ExistingElementId = RegisteredNavBounds.FindId(Request.NavBounds);
+
+		switch (Request.UpdateRequest)
+		{
+		case FNavigationBoundsUpdateRequest::Removed:
+			{
+				if (ExistingElementId.IsValidId())
+				{
+					UpdatedAreas.Add(RegisteredNavBounds[ExistingElementId].AreaBox);
+					RegisteredNavBounds.Remove(ExistingElementId);
+				}
+			}
+			break;
+
+		case FNavigationBoundsUpdateRequest::Added:
+		case FNavigationBoundsUpdateRequest::Updated:
+			{
+				if (ExistingElementId.IsValidId())
+				{
+					const FBox ExistingBox = RegisteredNavBounds[ExistingElementId].AreaBox;
+					const bool bSameArea = (Request.NavBounds.AreaBox == ExistingBox);
+					if (!bSameArea)
+					{
+						UpdatedAreas.Add(ExistingBox);
+					}
+
+					// always assign new bounds data, it may have different properties (like supported agents)
+					RegisteredNavBounds[ExistingElementId] = Request.NavBounds;
+				}
+				else
+				{
+					AddNavigationBounds(Request.NavBounds);
+				}
+
+				UpdatedAreas.Add(Request.NavBounds.AreaBox);
+			}
+
+			break;
+		}
+	}
+
+	if (!IsNavigationBuildingLocked())
+	{
+		if (UpdatedAreas.Num())
+		{
+			for (ANavigationData* NavData : NavDataSet)
+			{
+				if (NavData)
+				{
+					NavData->OnNavigationBoundsChanged();	
+				}
+			}
+		}
+				
+		// Propagate to generators areas that needs to be updated
+		AddDirtyAreas(UpdatedAreas, ENavigationDirtyFlag::All | ENavigationDirtyFlag::NavigationBounds);
+	}
+}
+
+void UNavigationSystemV1::AddNavigationBounds(const FNavigationBounds& NewBounds)
+{
+	RegisteredNavBounds.Add(NewBounds);
+}
+
+void UNavigationSystemV1::GatherNavigationBounds()
+{
+	// Gather all available navigation bounds
+	RegisteredNavBounds.Empty();
+	for (TActorIterator<ANavMeshBoundsVolume> It(GetWorld()); It; ++It)
+	{
+		ANavMeshBoundsVolume* V = (*It);
+		if (V != nullptr && !V->IsPendingKill())
+		{
+			FNavigationBounds NavBounds;
+			NavBounds.UniqueID = V->GetUniqueID();
+			NavBounds.AreaBox = V->GetComponentsBoundingBox(true);
+			NavBounds.Level = V->GetLevel();
+			NavBounds.SupportedAgents = V->SupportedAgents;
+
+			AddNavigationBounds(NavBounds);
+		}
+	}
+}
+
+void UNavigationSystemV1::Build()
+{
+	UWorld* World = GetWorld();
+	if (!World)
+	{
+		UE_LOG(LogNavigation, Error, TEXT("Unable to build navigation due to missing World pointer"));
+		return;
+	}
+
+	FNavigationSystem::DiscardNavigationDataChunks(*World);
+
+	const bool bHasWork = IsThereAnywhereToBuildNavigation();
+	const bool bLockedIgnoreEditor = (NavBuildingLockFlags & ~ENavigationBuildLock::NoUpdateInEditor) != 0;
+	if (!bHasWork || bLockedIgnoreEditor)
+	{
+		return;
+	}
+
+	const double BuildStartTime = FPlatformTime::Seconds();
+
+	SpawnMissingNavigationData();
+
+	// make sure freshly created navigation instances are registered before we try to build them
+	ProcessRegistrationCandidates();
+
+	// and now iterate through all registered and just start building them
+	RebuildAll();
+
+	// Block until build is finished
+	for (ANavigationData* NavData : NavDataSet)
+	{
+		if (NavData)
+		{
+			NavData->EnsureBuildCompletion();
+		}
+	}
+
+	UE_LOG(LogNavigation, Display, TEXT("UNavigationSystemV1::Build total execution time: %.5f"), float(FPlatformTime::Seconds() - BuildStartTime));
+}
+
+void UNavigationSystemV1::SpawnMissingNavigationData()
+{
+	DoInitialSetup();
+
+	const int32 SupportedAgentsCount = SupportedAgents.Num();
+	check(SupportedAgentsCount >= 0);
+	
+	// Bit array might be a bit of an overkill here, but this function will be called very rarely
+	TBitArray<> AlreadyInstantiated(false, SupportedAgentsCount);
+	uint8 NumberFound = 0;
+	UWorld* NavWorld = GetWorld();
+
+	// 1. check whether any of required navigation data has already been instantiated
+	for (TActorIterator<ANavigationData> It(NavWorld); It && NumberFound < SupportedAgentsCount; ++It)
+	{
+		ANavigationData* Nav = (*It);
+		if (Nav != NULL && Nav->GetTypedOuter<UWorld>() == NavWorld && Nav->IsPendingKill() == false)
+		{
+			// find out which one it is
+			for (int32 AgentIndex = 0; AgentIndex < SupportedAgentsCount; ++AgentIndex)
+			{
+				if (AlreadyInstantiated[AgentIndex] == true)
+				{
+					// already present, skip
+					continue;
+				}
+
+				if (Nav->GetClass() == SupportedAgents[AgentIndex].NavigationDataClass && Nav->DoesSupportAgent(SupportedAgents[AgentIndex]) == true)
+				{
+					AlreadyInstantiated[AgentIndex] = true;
+					++NumberFound;
+					break;
+				}
+			}				
+		}
+	}
+
+	// 2. for any not already instantiated navigation data call creator functions
+	if (NumberFound < SupportedAgentsCount)
+	{
+		for (int32 AgentIndex = 0; AgentIndex < SupportedAgentsCount; ++AgentIndex)
+		{
+			const FNavDataConfig& NavConfig = SupportedAgents[AgentIndex];
+			if (AlreadyInstantiated[AgentIndex] == false && NavConfig.NavigationDataClass != nullptr)
+			{
+				bool bHandled = false;
+
+				const ANavigationData* NavDataCDO = NavConfig.NavigationDataClass->GetDefaultObject<ANavigationData>();
+				if (NavDataCDO == nullptr || !NavDataCDO->CanSpawnOnRebuild())
+				{
+					continue;
+				}
+
+				if (NavWorld->WorldType != EWorldType::Editor && NavDataCDO->GetRuntimeGenerationMode() == ERuntimeGenerationType::Static)
+				{
+					// if we're not in the editor, and specified navigation class is configured 
+					// to be static, then we don't want to create an instance					
+					UE_LOG(LogNavigation, Log, TEXT("Not spawning navigation data for %s since indivated NavigationData type is not configured for dynamic generation")
+						, *NavConfig.Name.ToString());
+					continue;
+				}
+
+				ANavigationData* Instance = CreateNavigationDataInstance(NavConfig);
+				if (Instance != NULL)
+				{
+					RequestRegistration(Instance);
+				}
+				else
+				{
+					UE_LOG(LogNavigation, Warning, TEXT("Was not able to create navigation data for SupportedAgent[%d]: %s"), AgentIndex, *NavConfig.Name.ToString());
+				}
+			}
+		}
+
+		ProcessRegistrationCandidates();
+	}
+	
+	if (MainNavData == NULL || MainNavData->IsPendingKill())
+	{
+		// update 
+		MainNavData = GetDefaultNavDataInstance(FNavigationSystem::DontCreate);
+	}
+
+	UpdateAbstractNavData();
+}
+
+ANavigationData* UNavigationSystemV1::CreateNavigationDataInstance(const FNavDataConfig& NavConfig)
+{
+	UWorld* World = GetWorld();
+	check(World);
+
+	FActorSpawnParameters SpawnInfo;
+	SpawnInfo.OverrideLevel = World->PersistentLevel;
+	ANavigationData* Instance = World->SpawnActor<ANavigationData>(*NavConfig.NavigationDataClass, SpawnInfo);
+
+	if (Instance != NULL)
+	{
+		Instance->SetConfig(NavConfig);
+		if (NavConfig.Name != NAME_None)
+		{
+			FString StrName = FString::Printf(TEXT("%s-%s"), *(Instance->GetFName().GetPlainNameString()), *(NavConfig.Name.ToString()));
+			// temporary solution to make sure we don't try to change name while there's already
+			// an object with this name
+			UObject* ExistingObject = StaticFindObject(/*Class=*/ NULL, Instance->GetOuter(), *StrName, true);
+			if (ExistingObject != NULL)
+			{
+				ANavigationData* ExistingNavigationData = Cast<ANavigationData>(ExistingObject);
+				if (ExistingNavigationData)
+				{
+					UnregisterNavData(ExistingNavigationData);
+					AgentToNavDataMap.Remove(ExistingNavigationData->GetConfig());
+				}
+
+				ExistingObject->Rename(NULL, NULL, REN_DontCreateRedirectors | REN_ForceGlobalUnique | REN_DoNotDirty | REN_NonTransactional | REN_ForceNoResetLoaders);
+			}
+
+			// Set descriptive name
+			Instance->Rename(*StrName, NULL, REN_DoNotDirty | REN_ForceNoResetLoaders);
+#if WITH_EDITOR
+			if (World->WorldType == EWorldType::Editor)
+			{
+				const bool bMarkDirty = false;
+				Instance->SetActorLabel(StrName, bMarkDirty);
+			}
+#endif // WITH_EDITOR
+		}
+	}
+
+	return Instance;
+}
+
+void UNavigationSystemV1::OnPIEStart()
+{
+	// Do not tick async build for editor world while PIE is active
+	bAsyncBuildPaused = true;
+}
+
+void UNavigationSystemV1::OnPIEEnd()
+{
+	bAsyncBuildPaused = false;
+}
+
+void UNavigationSystemV1::RemoveNavigationBuildLock(uint8 Flags, bool bSkipRebuildInEditor)
+{
+	const bool bWasLocked = IsNavigationBuildingLocked();
+
+	NavBuildingLockFlags &= ~Flags;
+
+	const bool bIsLocked = IsNavigationBuildingLocked();
+	const bool bSkipRebuild = (OperationMode == FNavigationSystemRunMode::EditorMode) && bSkipRebuildInEditor;
+	if (bWasLocked && !bIsLocked && !bSkipRebuild)
+	{
+		RebuildAll();
+	}
+}
+
+void UNavigationSystemV1::RebuildAll(bool bIsLoadTime)
+{
+	const bool bIsInGame = GetWorld()->IsGameWorld();
+	
+	GatherNavigationBounds();
+
+	// make sure that octree is up to date
+	for (TSet<FNavigationDirtyElement>::TIterator It(PendingOctreeUpdates); It; ++It)
+	{
+		AddElementToNavOctree(*It);
+	}
+	PendingOctreeUpdates.Empty(32);
+
+	// discard all pending dirty areas, we are going to rebuild navmesh anyway 
+	DirtyAreas.Reset();
+	PendingNavBoundsUpdates.Reset();
+	
+	// 
+	for (int32 NavDataIndex = 0; NavDataIndex < NavDataSet.Num(); ++NavDataIndex)
+	{
+		ANavigationData* NavData = NavDataSet[NavDataIndex];
+				
+		if (NavData && (!bIsLoadTime || NavData->NeedsRebuildOnLoad()) && (!bIsInGame || NavData->SupportsRuntimeGeneration()))
+		{
+			NavData->RebuildAll();
+		}
+	}
+}
+
+bool UNavigationSystemV1::IsNavigationBuildInProgress(bool bCheckDirtyToo)
+{
+	bool bRet = false;
+
+	if (NavDataSet.Num() == 0)
+	{
+		// @todo this is wrong! Should not need to create a navigation data instance in a "getter" like function
+		// update nav data. If none found this is the place to create one
+		GetDefaultNavDataInstance(FNavigationSystem::DontCreate);
+	}
+	
+	for (int32 NavDataIndex = 0; NavDataIndex < NavDataSet.Num(); ++NavDataIndex)
+	{
+		ANavigationData* NavData = NavDataSet[NavDataIndex];
+		if (NavData != NULL && NavData->GetGenerator() != NULL 
+			&& NavData->GetGenerator()->IsBuildInProgress(bCheckDirtyToo) == true)
+		{
+			bRet = true;
+			break;
+		}
+	}
+
+	return bRet;
+}
+
+void UNavigationSystemV1::OnNavigationGenerationFinished(ANavigationData& NavData)
+{
+	OnNavigationGenerationFinishedDelegate.Broadcast(&NavData);
+}
+
+int32 UNavigationSystemV1::GetNumRemainingBuildTasks() const
+{
+	int32 NumTasks = 0;
+	
+	for (ANavigationData* NavData : NavDataSet)
+	{
+		if (NavData && NavData->GetGenerator())
+		{
+			NumTasks+= NavData->GetGenerator()->GetNumRemaningBuildTasks();
+		}
+	}
+	
+	return NumTasks;
+}
+
+int32 UNavigationSystemV1::GetNumRunningBuildTasks() const 
+{
+	int32 NumTasks = 0;
+	
+	for (ANavigationData* NavData : NavDataSet)
+	{
+		if (NavData && NavData->GetGenerator())
+		{
+			NumTasks+= NavData->GetGenerator()->GetNumRunningBuildTasks();
+		}
+	}
+	
+	return NumTasks;
+}
+
+void UNavigationSystemV1::OnLevelAddedToWorld(ULevel* InLevel, UWorld* InWorld)
+{
+	if ((IsNavigationSystemStatic() == false) && (InWorld == GetWorld()))
+	{
+		AddLevelCollisionToOctree(InLevel);
+
+		if (!InLevel->IsPersistentLevel())
+		{
+			for (ANavigationData* NavData : NavDataSet)
+			{
+				if (NavData)
+				{
+					NavData->OnStreamingLevelAdded(InLevel, InWorld);
+				}
+			}
+		}
+	}
+}
+
+void UNavigationSystemV1::OnLevelRemovedFromWorld(ULevel* InLevel, UWorld* InWorld)
+{
+	if ((IsNavigationSystemStatic() == false) && (InWorld == GetWorld()))
+	{
+		RemoveLevelCollisionFromOctree(InLevel);
+
+		if (InLevel && !InLevel->IsPersistentLevel())
+		{
+			for (int32 DataIndex = NavDataSet.Num() - 1; DataIndex >= 0; --DataIndex)
+			{		
+				ANavigationData* NavData = NavDataSet[DataIndex];
+				if (NavData)
+				{
+					if (NavData->GetLevel() != InLevel)
+					{
+						NavData->OnStreamingLevelRemoved(InLevel, InWorld);
+					}
+					else
+					{
+						NavDataSet.RemoveAt(DataIndex, 1, /*bAllowShrinking=*/false);
+					}
+				}
+			}
+		}
+	}
+}
+
+void UNavigationSystemV1::AddLevelCollisionToOctree(ULevel* Level)
+{
+#if WITH_RECAST
+	if (Level && NavOctree.IsValid() &&
+		NavOctree->GetNavGeometryStoringMode() == FNavigationOctree::StoreNavGeometry)
+	{
+		const TArray<FVector>* LevelGeom = Level->GetStaticNavigableGeometry();
+		const FOctreeElementId* ElementId = GetObjectsNavOctreeId(Level);
+
+		if (!ElementId && LevelGeom && LevelGeom->Num() > 0)
+		{
+			FNavigationOctreeElement BSPElem(*Level);
+			FRecastNavMeshGenerator::ExportVertexSoupGeometry(*LevelGeom, *BSPElem.Data);
+
+			const auto& Bounds = BSPElem.Data->Bounds;
+			if (!Bounds.GetExtent().IsNearlyZero())
+			{
+				NavOctree->AddNode(Level, NULL, Bounds, BSPElem);
+				AddDirtyArea(Bounds, ENavigationDirtyFlag::All);
+
+				UE_LOG(LogNavOctree, Log, TEXT("ADD %s"), *GetNameSafe(Level));
+			}
+		}
+	}
+#endif// WITH_RECAST
+}
+
+void UNavigationSystemV1::RemoveLevelCollisionFromOctree(ULevel* Level)
+{
+	if (Level && NavOctree.IsValid())
+	{
+		const FOctreeElementId* ElementId = GetObjectsNavOctreeId(Level);
+		UE_LOG(LogNavOctree, Log, TEXT("UNREG %s %s"), *GetNameSafe(Level), ElementId ? TEXT("[exists]") : TEXT(""));
+
+		if (ElementId != NULL)
+		{
+			if (NavOctree->IsValidElementId(*ElementId))
+			{
+				// mark area occupied by given actor as dirty
+				FNavigationOctreeElement& ElementData = NavOctree->GetElementById(*ElementId);
+				AddDirtyArea(ElementData.Bounds.GetBox(), ENavigationDirtyFlag::All);
+			}
+
+			NavOctree->RemoveNode(*ElementId);
+			RemoveObjectsNavOctreeId(Level);
+		}
+	}
+}
+
+void UNavigationSystemV1::OnPostLoadMap(UWorld*)
+{
+	UE_LOG(LogNavigation, Log, TEXT("UNavigationSystemV1::OnPostLoadMap"));
+
+	// if map has been loaded and there are some navigation bounds volumes 
+	// then create appropriate navigation structured
+	ANavigationData* NavData = GetDefaultNavDataInstance(FNavigationSystem::DontCreate);
+
+	// Do this if there's currently no navigation
+	if (NavData == NULL && bAutoCreateNavigationData == true && IsThereAnywhereToBuildNavigation() == true)
+	{
+		NavData = GetDefaultNavDataInstance(FNavigationSystem::Create);
+	}
+}
+
+#if WITH_EDITOR
+void UNavigationSystemV1::OnActorMoved(AActor* Actor)
+{
+	if (Cast<ANavMeshBoundsVolume>(Actor) != NULL)
+	{
+		OnNavigationBoundsUpdated((ANavMeshBoundsVolume*)Actor);
+	}
+}
+#endif // WITH_EDITOR
+
+void UNavigationSystemV1::OnNavigationDirtied(const FBox& Bounds)
+{
+	AddDirtyArea(Bounds, ENavigationDirtyFlag::All);
+}
+
+#if WITH_HOT_RELOAD
+void UNavigationSystemV1::OnHotReload(bool bWasTriggeredAutomatically)
+{
+	if (RequiresNavOctree() && NavOctree.IsValid() == false)
+	{
+		ConditionalPopulateNavOctree();
+
+		if (bInitialBuildingLocked)
+		{
+			RemoveNavigationBuildLock(ENavigationBuildLock::InitialLock, /*bSkipRebuildInEditor=*/true);
+		}
+	}
+}
+#endif // WITH_HOT_RELOAD
+
+void UNavigationSystemV1::CleanUp(FNavigationSystem::ECleanupMode Mode)
+{
+	UE_LOG(LogNavigation, Log, TEXT("UNavigationSystemV1::CleanUp"));
+
+#if WITH_EDITOR
+	if (GIsEditor && GEngine)
+	{
+		GEngine->OnActorMoved().RemoveAll(this);
+	}
+#endif // WITH_EDITOR
+
+	FCoreUObjectDelegates::PostLoadMapWithWorld.RemoveAll(this);
+	UNavigationSystemV1::NavigationDirtyEvent.RemoveAll(this);
+	FWorldDelegates::LevelAddedToWorld.RemoveAll(this);
+	FWorldDelegates::LevelRemovedFromWorld.RemoveAll(this);
+
+#if WITH_HOT_RELOAD
+	if (IHotReloadInterface* HotReloadSupport = FModuleManager::GetModulePtr<IHotReloadInterface>("HotReload"))
+	{
+		HotReloadSupport->OnHotReload().Remove(HotReloadDelegateHandle);
+	}
+#endif
+
+	DestroyNavOctree();
+	
+	SetCrowdManager(NULL);
+
+	NavDataSet.Reset();
+
+	// reset unique link Id for new map
+	const UWorld* MyWorld = (Mode == FNavigationSystem::ECleanupMode::CleanupWithWorld) ? GetWorld() : NULL;
+	if (MyWorld && (MyWorld->WorldType == EWorldType::Game || MyWorld->WorldType == EWorldType::Editor))
+	{
+		INavLinkCustomInterface::NextUniqueId = 1;
+	}
+}
+
+void UNavigationSystemV1::DestroyNavOctree()
+{
+	if (NavOctree.IsValid())
+	{
+		NavOctree->Destroy();
+		NavOctree = NULL;
+	}
+
+	ObjectToOctreeId.Empty();
+}
+
+bool UNavigationSystemV1::RequiresNavOctree() const
+{
+	UWorld* World = GetWorld();
+	check(World);
+	
+	// We always require navoctree in editor worlds
+	if (!World->IsGameWorld())
+	{
+		return true;
+	}
+		
+	for (ANavigationData* NavData : NavDataSet)
+	{
+		if (NavData && NavData->SupportsRuntimeGeneration())
+		{
+			return true;
+		}
+	}
+	
+	return false;
+}
+
+ERuntimeGenerationType UNavigationSystemV1::GetRuntimeGenerationType() const
+{
+	UWorld* World = GetWorld();
+	check(World);
+	
+	// We always use ERuntimeGenerationType::Dynamic in editor worlds
+	if (!World->IsGameWorld())
+	{
+		return ERuntimeGenerationType::Dynamic;
+	}
+	
+	ERuntimeGenerationType RuntimeGenerationType = ERuntimeGenerationType::Static;
+
+	for (ANavigationData* NavData : NavDataSet)
+	{
+		if (NavData && NavData->GetRuntimeGenerationMode() > RuntimeGenerationType)
+		{
+			RuntimeGenerationType = NavData->GetRuntimeGenerationMode();
+		}
+	}
+	
+	return RuntimeGenerationType;
+}
+
+//----------------------------------------------------------------------//
+// Blueprint functions
+//----------------------------------------------------------------------//
+UNavigationSystemV1* UNavigationSystemV1::GetNavigationSystem(UObject* WorldContextObject)
+{
+	return GetCurrent(WorldContextObject);
+}
+
+bool UNavigationSystemV1::K2_ProjectPointToNavigation(UObject* WorldContextObject, const FVector& Point, FVector& ProjectedLocation, ANavigationData* NavData, TSubclassOf<UNavigationQueryFilter> FilterClass, const FVector QueryExtent)
+{
+	UWorld* World = GEngine->GetWorldFromContextObject(WorldContextObject, EGetWorldErrorMode::LogAndReturnNull);
+	UNavigationSystemV1* NavSys = FNavigationSystem::GetCurrent<UNavigationSystemV1>(World);
+
+	ProjectedLocation = Point;
+	bool bResult = false;
+
+	if (NavSys)
+	{
+		FNavLocation OutNavLocation;
+		ANavigationData* UseNavData = NavData ? NavData : NavSys->GetDefaultNavDataInstance(FNavigationSystem::DontCreate);
+		if (UseNavData)
+		{
+			bResult = NavSys->ProjectPointToNavigation(Point, OutNavLocation, QueryExtent, NavData
+				, UNavigationQueryFilter::GetQueryFilter(*UseNavData, WorldContextObject, FilterClass));
+			ProjectedLocation = OutNavLocation.Location;
+		}
+	}
+
+	return bResult;
+}
+
+bool UNavigationSystemV1::K2_GetRandomReachablePointInRadius(UObject* WorldContextObject, const FVector& Origin, FVector& RandomLocation, float Radius, ANavigationData* NavData, TSubclassOf<UNavigationQueryFilter> FilterClass)
+{
+	FNavLocation RandomPoint(Origin);
+	bool bResult = false;
+
+	UWorld* World = GEngine->GetWorldFromContextObject(WorldContextObject, EGetWorldErrorMode::LogAndReturnNull);
+	UNavigationSystemV1* NavSys = FNavigationSystem::GetCurrent<UNavigationSystemV1>(World);
+	if (NavSys)
+	{
+		ANavigationData* UseNavData = NavData ? NavData : NavSys->GetDefaultNavDataInstance(FNavigationSystem::DontCreate);
+		if (UseNavData)
+		{
+			bResult = NavSys->GetRandomReachablePointInRadius(Origin, Radius, RandomPoint, UseNavData, UNavigationQueryFilter::GetQueryFilter(*UseNavData, WorldContextObject, FilterClass));
+			RandomLocation = RandomPoint.Location;
+		}
+	}
+
+	return bResult;
+}
+
+bool UNavigationSystemV1::K2_GetRandomPointInNavigableRadius(UObject* WorldContextObject, const FVector& Origin, FVector& RandomLocation, float Radius, ANavigationData* NavData, TSubclassOf<UNavigationQueryFilter> FilterClass)
+{
+	FNavLocation RandomPoint(Origin);
+	bool bResult = false;
+
+	UWorld* World = GEngine->GetWorldFromContextObject(WorldContextObject, EGetWorldErrorMode::LogAndReturnNull);
+	UNavigationSystemV1* NavSys = FNavigationSystem::GetCurrent<UNavigationSystemV1>(World);
+	if (NavSys)
+	{
+		ANavigationData* UseNavData = NavData ? NavData : NavSys->GetDefaultNavDataInstance(FNavigationSystem::DontCreate);
+		if (UseNavData)
+		{
+			bResult = NavSys->GetRandomPointInNavigableRadius(Origin, Radius, RandomPoint, UseNavData, UNavigationQueryFilter::GetQueryFilter(*UseNavData, WorldContextObject, FilterClass));
+			RandomLocation = RandomPoint.Location;
+		}
+	}
+
+	return bResult;
+}
+
+ENavigationQueryResult::Type UNavigationSystemV1::GetPathCost(UObject* WorldContextObject, const FVector& PathStart, const FVector& PathEnd, float& OutPathCost, ANavigationData* NavData, TSubclassOf<UNavigationQueryFilter> FilterClass)
+{
+	UWorld* World = GEngine->GetWorldFromContextObject(WorldContextObject, EGetWorldErrorMode::LogAndReturnNull);
+	UNavigationSystemV1* NavSys = FNavigationSystem::GetCurrent<UNavigationSystemV1>(World);
+	if (NavSys)
+	{
+		ANavigationData* UseNavData = NavData ? NavData : NavSys->GetDefaultNavDataInstance(FNavigationSystem::DontCreate);
+		if (UseNavData)
+		{
+			return NavSys->GetPathCost(PathStart, PathEnd, OutPathCost, UseNavData, UNavigationQueryFilter::GetQueryFilter(*UseNavData, WorldContextObject, FilterClass));
+		}
+	}
+
+	return ENavigationQueryResult::Error;
+}
+
+ENavigationQueryResult::Type UNavigationSystemV1::GetPathLength(UObject* WorldContextObject, const FVector& PathStart, const FVector& PathEnd, float& OutPathLength, ANavigationData* NavData, TSubclassOf<UNavigationQueryFilter> FilterClass)
+{
+	float PathLength = 0.f;
+
+	UWorld* World = GEngine->GetWorldFromContextObject(WorldContextObject, EGetWorldErrorMode::LogAndReturnNull);
+	UNavigationSystemV1* NavSys = FNavigationSystem::GetCurrent<UNavigationSystemV1>(World);
+	if (NavSys)
+	{
+		ANavigationData* UseNavData = NavData ? NavData : NavSys->GetDefaultNavDataInstance(FNavigationSystem::DontCreate);
+		if (UseNavData)
+		{
+			return NavSys->GetPathLength(PathStart, PathEnd, OutPathLength, UseNavData, UNavigationQueryFilter::GetQueryFilter(*UseNavData, WorldContextObject, FilterClass));
+		}
+	}
+
+	return ENavigationQueryResult::Error;
+}
+
+bool UNavigationSystemV1::IsNavigationBeingBuilt(UObject* WorldContextObject)
+{
+	UWorld* World = GEngine->GetWorldFromContextObject(WorldContextObject, EGetWorldErrorMode::LogAndReturnNull);
+	UNavigationSystemV1* NavSys = FNavigationSystem::GetCurrent<UNavigationSystemV1>(World);
+	
+	if (NavSys && !NavSys->IsNavigationBuildingPermanentlyLocked())
+	{
+		return NavSys->HasDirtyAreasQueued() || NavSys->IsNavigationBuildInProgress();
+	}
+
+	return false;
+}
+
+bool UNavigationSystemV1::IsNavigationBeingBuiltOrLocked(UObject* WorldContextObject)
+{
+	UWorld* World = GEngine->GetWorldFromContextObject(WorldContextObject, EGetWorldErrorMode::LogAndReturnNull);
+	UNavigationSystemV1* NavSys = FNavigationSystem::GetCurrent<UNavigationSystemV1>(World);
+
+	if (NavSys)
+	{
+		return NavSys->IsNavigationBuildingLocked() || NavSys->HasDirtyAreasQueued() || NavSys->IsNavigationBuildInProgress();
+	}
+
+	return false;
+}
+
+//----------------------------------------------------------------------//
+// HACKS!!!
+//----------------------------------------------------------------------//
+bool UNavigationSystemV1::ShouldGeneratorRun(const FNavDataGenerator* Generator) const
+{
+	if (Generator != NULL && (IsNavigationSystemStatic() == false))
+	{
+		for (int32 NavDataIndex = 0; NavDataIndex < NavDataSet.Num(); ++NavDataIndex)
+		{
+			ANavigationData* NavData = NavDataSet[NavDataIndex];
+			if (NavData != NULL && NavData->GetGenerator() == Generator)
+			{
+				return true;
+			}
+		}
+	}
+
+	return false;
+}
+
+bool UNavigationSystemV1::HandleCycleNavDrawnCommand( const TCHAR* Cmd, FOutputDevice& Ar )
+{
+	CycleNavigationDataDrawn();
+
+	return true;
+}
+
+bool UNavigationSystemV1::HandleCountNavMemCommand( const TCHAR* Cmd, FOutputDevice& Ar )
+{
+	for (int32 NavDataIndex = 0; NavDataIndex < NavDataSet.Num(); ++NavDataIndex)
+	{
+		ANavigationData* NavData = NavDataSet[NavDataIndex];
+		if (NavData != NULL)
+		{
+			NavData->LogMemUsed();
+		}
+	}
+	return true;
+}
+
+//----------------------------------------------------------------------//
+// Commands
+//----------------------------------------------------------------------//
+bool FNavigationSystemExec::Exec(UWorld* InWorld, const TCHAR* Cmd, FOutputDevice& Ar)
+{
+	UNavigationSystemV1*  NavSys = FNavigationSystem::GetCurrent<UNavigationSystemV1>(InWorld);
+
+	if (NavSys && NavSys->NavDataSet.Num() > 0)
+	{
+		if (FParse::Command(&Cmd, TEXT("CYCLENAVDRAWN")))
+		{
+			NavSys->HandleCycleNavDrawnCommand( Cmd, Ar );
+			// not returning true to enable all navigation systems to cycle their own data
+			return false;
+		}
+		else if (FParse::Command(&Cmd, TEXT("CountNavMem")))
+		{
+			NavSys->HandleCountNavMemCommand( Cmd, Ar );
+			return false;
+		}
+		/** Builds the navigation mesh (or rebuilds it). **/
+		else if (FParse::Command(&Cmd, TEXT("RebuildNavigation")))
+		{
+			NavSys->Build();
+		}
+	}
+
+	return false;
+}
+
+void UNavigationSystemV1::CycleNavigationDataDrawn()
+{
+	++CurrentlyDrawnNavDataIndex;
+	if (CurrentlyDrawnNavDataIndex >= NavDataSet.Num())
+	{
+		CurrentlyDrawnNavDataIndex = INDEX_NONE;
+	}
+
+	for (int32 NavDataIndex = 0; NavDataIndex < NavDataSet.Num(); ++NavDataIndex)
+	{
+		ANavigationData* NavData = NavDataSet[NavDataIndex];
+		if (NavData != NULL)
+		{
+			const bool bNewEnabledDrawing = (CurrentlyDrawnNavDataIndex == INDEX_NONE) || (NavDataIndex == CurrentlyDrawnNavDataIndex);
+			NavData->SetNavRenderingEnabled(bNewEnabledDrawing);
+		}
+	}
+}
+
+bool UNavigationSystemV1::IsNavigationDirty() const
+{
+#if !UE_BUILD_SHIPPING
+	if (bCanAccumulateDirtyAreas == false && bDirtyAreasReportedWhileAccumulationLocked)
+	{
+		return true;
+	}
+#endif // !UE_BUILD_SHIPPING
+
+	for (int32 NavDataIndex=0; NavDataIndex < NavDataSet.Num(); ++NavDataIndex)
+	{
+		if (NavDataSet[NavDataIndex] && NavDataSet[NavDataIndex]->NeedsRebuild())
+		{
+			return true;
+		}
+	}
+
+	return false;
+}
+
+bool UNavigationSystemV1::CanRebuildDirtyNavigation() const
+{
+	const bool bIsInGame = GetWorld()->IsGameWorld();
+
+	for (const ANavigationData* NavData : NavDataSet)
+	{
+		if (NavData)
+		{
+			const bool bIsDirty = NavData->NeedsRebuild();
+			const bool bCanRebuild = !bIsInGame || NavData->SupportsRuntimeGeneration();
+
+			if (bIsDirty && !bCanRebuild)
+			{
+				return false;
+			}
+		}
+	}
+
+	return true;
+}
+
+bool UNavigationSystemV1::DoesPathIntersectBox(const FNavigationPath* Path, const FBox& Box, uint32 StartingIndex, FVector* AgentExtent)
+{
+	return Path != NULL && Path->DoesIntersectBox(Box, StartingIndex, NULL, AgentExtent);
+}
+
+bool UNavigationSystemV1::DoesPathIntersectBox(const FNavigationPath* Path, const FBox& Box, const FVector& AgentLocation, uint32 StartingIndex, FVector* AgentExtent)
+{
+	return Path != NULL && Path->DoesIntersectBox(Box, AgentLocation, StartingIndex, NULL, AgentExtent);
+}
+
+void UNavigationSystemV1::SetMaxSimultaneousTileGenerationJobsCount(int32 MaxNumberOfJobs)
+{
+#if WITH_RECAST
+	for (auto NavigationData : NavDataSet)
+	{
+		ARecastNavMesh* RecastNavMesh = Cast<ARecastNavMesh>(NavigationData);
+		if (RecastNavMesh)
+		{
+			RecastNavMesh->SetMaxSimultaneousTileGenerationJobsCount(MaxNumberOfJobs);
+		}
+	}
+#endif
+}
+
+void UNavigationSystemV1::ResetMaxSimultaneousTileGenerationJobsCount()
+{
+#if WITH_RECAST
+	for (auto NavigationData : NavDataSet)
+	{
+		ARecastNavMesh* RecastNavMesh = Cast<ARecastNavMesh>(NavigationData);
+		if (RecastNavMesh)
+		{
+			const ARecastNavMesh* CDO = RecastNavMesh->GetClass()->GetDefaultObject<ARecastNavMesh>();
+			RecastNavMesh->SetMaxSimultaneousTileGenerationJobsCount(CDO->MaxSimultaneousTileGenerationJobsCount);
+		}
+	}
+#endif
+}
+
+//----------------------------------------------------------------------//
+// Active tiles
+//----------------------------------------------------------------------//
+
+void UNavigationSystemV1::RegisterNavigationInvoker(AActor& Invoker, float TileGenerationRadius, float TileRemovalRadius)
+{
+	UNavigationSystemV1* NavSys = FNavigationSystem::GetCurrent<UNavigationSystemV1>(Invoker.GetWorld());
+	if (NavSys)
+	{
+		NavSys->RegisterInvoker(Invoker, TileGenerationRadius, TileRemovalRadius);
+	}
+}
+
+void UNavigationSystemV1::UnregisterNavigationInvoker(AActor& Invoker)
+{
+	UNavigationSystemV1* NavSys = FNavigationSystem::GetCurrent<UNavigationSystemV1>(Invoker.GetWorld());
+	if (NavSys)
+	{
+		NavSys->UnregisterInvoker(Invoker);
+	}
+}
+
+void UNavigationSystemV1::SetGeometryGatheringMode(ENavDataGatheringModeConfig NewMode)
+{
+	DataGatheringMode = NewMode;
+	if (NavOctree.IsValid())
+	{
+		NavOctree->SetDataGatheringMode(DataGatheringMode);
+	}
+}
+
+void UNavigationSystemV1::RegisterInvoker(AActor& Invoker, float TileGenerationRadius, float TileRemovalRadius)
+{
+	UE_CVLOG(bGenerateNavigationOnlyAroundNavigationInvokers == false, this, LogNavigation, Warning
+		, TEXT("Trying to register %s as enforcer, but NavigationSystem is not set up for enforcer-centric generation. See GenerateNavigationOnlyAroundNavigationInvokers in NavigationSystem's properties")
+		, *Invoker.GetName());
+
+	TileGenerationRadius = FMath::Clamp(TileGenerationRadius, 0.f, BIG_NUMBER);
+	TileRemovalRadius = FMath::Clamp(TileRemovalRadius, TileGenerationRadius, BIG_NUMBER);
+
+	FNavigationInvoker& Data = Invokers.FindOrAdd(&Invoker);
+	Data.Actor = &Invoker;
+	Data.GenerationRadius = TileGenerationRadius;
+	Data.RemovalRadius = TileRemovalRadius;
+
+	UE_VLOG_CYLINDER(this, LogNavigation, Log, Invoker.GetActorLocation(), Invoker.GetActorLocation() + FVector(0, 0, 20), TileGenerationRadius, FColorList::LimeGreen
+		, TEXT("%s %.0f %.0f"), *Invoker.GetName(), TileGenerationRadius, TileRemovalRadius);
+	UE_VLOG_CYLINDER(this, LogNavigation, Log, Invoker.GetActorLocation(), Invoker.GetActorLocation() + FVector(0, 0, 20), TileRemovalRadius, FColorList::IndianRed, TEXT(""));
+}
+
+void UNavigationSystemV1::UnregisterInvoker(AActor& Invoker)
+{
+	UE_VLOG(this, LogNavigation, Log, TEXT("Removing %s from enforcers list"), *Invoker.GetName());
+	Invokers.Remove(&Invoker);
+}
+
+void UNavigationSystemV1::UpdateInvokers()
+{
+	UWorld* World = GetWorld();
+	const float CurrentTime = World->GetTimeSeconds();
+	if (CurrentTime >= NextInvokersUpdateTime)
+	{
+		TArray<FNavigationInvokerRaw> InvokerLocations;
+
+		if (Invokers.Num() > 0)
+		{
+			QUICK_SCOPE_CYCLE_COUNTER(STAT_NavSys_Clusterize);
+
+			const double StartTime = FPlatformTime::Seconds();
+
+			InvokerLocations.Reserve(Invokers.Num());
+
+			for (auto ItemIterator = Invokers.CreateIterator(); ItemIterator; ++ItemIterator)
+			{
+				AActor* Actor = ItemIterator->Value.Actor.Get();
+				if (Actor != nullptr
+#if WITH_EDITOR
+					// Would like to ignore objects in transactional buffer here, but there's no flag for it
+					//&& (GIsEditor == false || Item.Actor->HasAnyFlags(RF_Transactional | RF_PendingKill) == false)
+#endif //WITH_EDITOR
+					)
+				{
+					InvokerLocations.Add(FNavigationInvokerRaw(Actor->GetActorLocation(), ItemIterator->Value.GenerationRadius, ItemIterator->Value.RemovalRadius));
+				}
+				else
+				{
+					ItemIterator.RemoveCurrent();
+				}
+			}
+
+#if ENABLE_VISUAL_LOG
+			const double CachingFinishTime = FPlatformTime::Seconds();
+			UE_VLOG(this, LogNavigation, Log, TEXT("Caching time %fms"), (CachingFinishTime - StartTime) * 1000);
+
+			for (const auto& InvokerData : InvokerLocations)
+			{
+				UE_VLOG_CYLINDER(this, LogNavigation, Log, InvokerData.Location, InvokerData.Location + FVector(0, 0, 20), InvokerData.RadiusMax, FColorList::Blue, TEXT(""));
+				UE_VLOG_CYLINDER(this, LogNavigation, Log, InvokerData.Location, InvokerData.Location + FVector(0, 0, 20), InvokerData.RadiusMin, FColorList::CadetBlue, TEXT(""));
+			}
+#endif // ENABLE_VISUAL_LOG
+		}
+
+#if WITH_RECAST
+		const double UpdateStartTime = FPlatformTime::Seconds();
+		for (TActorIterator<ARecastNavMesh> It(GetWorld()); It; ++It)
+		{
+			It->UpdateActiveTiles(InvokerLocations);
+		}
+		const double UpdateEndTime = FPlatformTime::Seconds();
+		UE_VLOG(this, LogNavigation, Log, TEXT("Marking tiles to update %fms (%d invokers)"), (UpdateEndTime - UpdateStartTime) * 1000, InvokerLocations.Num());
+#endif
+
+		// once per second
+		NextInvokersUpdateTime = CurrentTime + ActiveTilesUpdateInterval;
+	}
+}
+
+void UNavigationSystemV1::RegisterNavigationInvoker(AActor* Invoker, float TileGenerationRadius, float TileRemovalRadius)
+{
+	if (Invoker != nullptr)
+	{
+		RegisterInvoker(*Invoker, TileGenerationRadius, TileRemovalRadius);
+	}
+}
+
+void UNavigationSystemV1::UnregisterNavigationInvoker(AActor* Invoker)
+{
+	if (Invoker != nullptr)
+	{
+		UnregisterInvoker(*Invoker);
+	}
+}
+
+//----------------------------------------------------------------------//
+// DEPRECATED
+//----------------------------------------------------------------------//
+FVector UNavigationSystemV1::ProjectPointToNavigation(UObject* WorldContextObject, const FVector& Point, ANavigationData* NavData, TSubclassOf<UNavigationQueryFilter> FilterClass, const FVector QueryExtent)
+{
+	FNavLocation ProjectedPoint(Point);
+
+	UWorld* World = GEngine->GetWorldFromContextObject(WorldContextObject, EGetWorldErrorMode::LogAndReturnNull);
+	UNavigationSystemV1* NavSys = FNavigationSystem::GetCurrent<UNavigationSystemV1>(World);
+	if (NavSys)
+	{
+		ANavigationData* UseNavData = NavData ? NavData : NavSys->GetDefaultNavDataInstance(FNavigationSystem::DontCreate);
+		if (UseNavData)
+		{
+			NavSys->ProjectPointToNavigation(Point, ProjectedPoint, QueryExtent.IsNearlyZero() ? INVALID_NAVEXTENT : QueryExtent, UseNavData,
+				UNavigationQueryFilter::GetQueryFilter(*UseNavData, WorldContextObject, FilterClass));
+		}
+	}
+
+	return ProjectedPoint.Location;
+}
+
+FVector UNavigationSystemV1::GetRandomReachablePointInRadius(UObject* WorldContextObject, const FVector& Origin, float Radius, ANavigationData* NavData, TSubclassOf<UNavigationQueryFilter> FilterClass)
+{
+	FNavLocation RandomPoint;
+
+	UWorld* World = GEngine->GetWorldFromContextObject(WorldContextObject, EGetWorldErrorMode::LogAndReturnNull);
+	UNavigationSystemV1* NavSys = FNavigationSystem::GetCurrent<UNavigationSystemV1>(World);
+	if (NavSys)
+	{
+		ANavigationData* UseNavData = NavData ? NavData : NavSys->GetDefaultNavDataInstance(FNavigationSystem::DontCreate);
+		if (UseNavData)
+		{
+			NavSys->GetRandomReachablePointInRadius(Origin, Radius, RandomPoint, UseNavData, UNavigationQueryFilter::GetQueryFilter(*UseNavData, WorldContextObject, FilterClass));
+		}
+	}
+
+	return RandomPoint.Location;
+}
+
+FVector UNavigationSystemV1::GetRandomPointInNavigableRadius(UObject* WorldContextObject, const FVector& Origin, float Radius, ANavigationData* NavData, TSubclassOf<UNavigationQueryFilter> FilterClass)
+{
+	FNavLocation RandomPoint;
+
+	UWorld* World = GEngine->GetWorldFromContextObject(WorldContextObject, EGetWorldErrorMode::LogAndReturnNull);
+	UNavigationSystemV1* NavSys = FNavigationSystem::GetCurrent<UNavigationSystemV1>(World);
+	if (NavSys)
+	{
+		ANavigationData* UseNavData = NavData ? NavData : NavSys->GetDefaultNavDataInstance(FNavigationSystem::DontCreate);
+		if (UseNavData)
+		{
+			NavSys->GetRandomPointInNavigableRadius(Origin, Radius, RandomPoint, UseNavData, UNavigationQueryFilter::GetQueryFilter(*UseNavData, WorldContextObject, FilterClass));
+		}
+	}
+
+	return RandomPoint.Location;
+}
+
+void UNavigationSystemV1::SimpleMoveToActor(AController* Controller, const AActor* Goal)
+{
+	UE_LOG(LogNavigation, Error, TEXT("SimpleMoveToActor is deprecated. Use UAIBlueprintHelperLibrary::SimpleMoveToActor instead"));
+}
+
+void UNavigationSystemV1::SimpleMoveToLocation(AController* Controller, const FVector& Goal)
+{
+	UE_LOG(LogNavigation, Error, TEXT("SimpleMoveToLocation is deprecated. Use UAIBlueprintHelperLibrary::SimpleMoveToLocation instead"));
+}
+
+//----------------------------------------------------------------------//
+// NEW STUFF!
+//----------------------------------------------------------------------//
+void UNavigationSystemV1::VerifyNavigationRenderingComponents(const bool bShow)
+{
+	// make sure nav mesh has a rendering component
+	ANavigationData* const NavData = GetDefaultNavDataInstance(FNavigationSystem::DontCreate);
+
+	if (NavData && NavData->RenderingComp == nullptr)
+	{
+		NavData->RenderingComp = NavData->ConstructRenderingComponent();
+		if (NavData->RenderingComp)
+		{
+			NavData->RenderingComp->SetVisibility(bShow);
+			NavData->RenderingComp->RegisterComponent();
+		}
+	}
+
+	if (NavData == nullptr)
+	{
+		UE_LOG(LogNavigation, Warning, TEXT("No NavData found when calling UNavigationSystemV1::VerifyNavigationRenderingComponents()"));
+	}
+}
+
+#if !UE_BUILD_SHIPPING
+void UNavigationSystemV1::GetOnScreenMessages(TMultiMap<FCoreDelegates::EOnScreenMessageSeverity, FText>& OutMessages)
+{
+	// check navmesh
+#if WITH_EDITOR
+	const bool bIsNavigationAutoUpdateEnabled = UNavigationSystemV1::GetIsNavigationAutoUpdateEnabled();
+#else
+	const bool bIsNavigationAutoUpdateEnabled = true;
+#endif
+	if (IsNavigationDirty() 
+		&& (!bIsNavigationAutoUpdateEnabled || !SupportsNavigationGeneration() || !CanRebuildDirtyNavigation()))
+	{
+		OutMessages.Add(FCoreDelegates::EOnScreenMessageSeverity::Error
+			, LOCTEXT("NAVMESHERROR", "NAVMESH NEEDS TO BE REBUILT"));
+
+		/*SmallTextItem.SetColor(FLinearColor::White);
+		SmallTextItem.Text = ;
+		Canvas->DrawItem(SmallTextItem, FVector2D(MessageX, MessageY));
+		MessageY += FontSizeY;*/
+	}
+}
+#endif // !UE_BUILD_SHIPPING
+// mz@todo
+//FVisualLogger::NavigationDataDumpDelegate
+//{
+//	UWorld* World = nullptr;
+//	FVisualLogEntry* CurrentEntry = nullptr;
+//	if (CheckVisualLogInputInternal(Object, Category, Verbosity, &World, &CurrentEntry) == false || CurrentEntry == nullptr)
+//	{
+//		return;
+//	}
+//
+//	const ANavigationData* MainNavData = World ? UNavigationSystemV1::GetNavigationSystem(World)->GetDefaultNavDataInstance(FNavigationSystem::ECreateIfEmpty::DontCreate) : nullptr;
+//	const FNavDataGenerator* Generator = MainNavData ? MainNavData->GetGenerator() : nullptr;
+//	if (Generator)
+//	{
+//		Generator->GrabDebugSnapshot(CurrentEntry, FMath::IsNearlyZero(Box.GetVolume()) ? MainNavData->GetBounds().ExpandBy(FVector(20, 20, 20)) : Box, Category, Verbosity);
+//	}
+//}
+
+INavigationDataInterface* UNavigationSystemV1::GetNavDataForActor(const AActor& Actor)
+{
+	UNavigationSystemV1* NavSys = FNavigationSystem::GetCurrent<UNavigationSystemV1>(Actor.GetWorld());
+	ANavigationData* NavData = nullptr;
+	const INavAgentInterface* AsNavAgent = CastChecked<INavAgentInterface>(&Actor);
+	if (AsNavAgent)
+	{
+		const FNavAgentProperties& AgentProps = AsNavAgent->GetNavAgentPropertiesRef();
+		NavData = NavSys->GetNavDataForProps(AgentProps);
+	}
+	if (NavData == nullptr)
+	{
+		NavData = NavSys->GetDefaultNavDataInstance(FNavigationSystem::DontCreate);
+	}
+
+	//  Only RecastNavMesh supported
+	return (INavigationDataInterface*)(Cast<ARecastNavMesh>(NavData));
+}
+
+int UNavigationSystemV1::GetNavigationBoundsForNavData(const ANavigationData& NavData, TArray<FBox>& OutBounds) const
+{
+	const int InitialBoundsCount = OutBounds.Num();
+	OutBounds.Reserve(InitialBoundsCount + RegisteredNavBounds.Num());
+	const int32 AgentIndex = GetSupportedAgentIndex(&NavData);
+
+	for (const FNavigationBounds& NavigationBounds : RegisteredNavBounds)
+	{
+		if (NavigationBounds.SupportedAgents.Contains(AgentIndex))
+		{
+			OutBounds.Add(NavigationBounds.AreaBox);
+		}
+	}
+
+	return OutBounds.Num() - InitialBoundsCount;
+}
+
+const FNavDataConfig& UNavigationSystemV1::GetDefaultSupportedAgent()
+{
+	static const FNavDataConfig DefaultAgent;
+	const UNavigationSystemV1* NavSysCDO = GetDefault<UNavigationSystemV1>();
+	check(NavSysCDO);
+	return NavSysCDO->SupportedAgents.Num() > 0
+		? NavSysCDO->GetDefaultSupportedAgentConfig()
+		: DefaultAgent;
+}
+
+void UNavigationSystemV1::Configure(const UNavigationSystemConfig& Config)
+{
+
+}
+
+#undef LOCTEXT_NAMESPACE
