Copying //UE4/Dev-Framework to //UE4/Dev-Main (Source: //UE4/Dev-Framework @ 3967517)

diff --git a/Engine/Source/Runtime/NavigationSystem/Public/AbstractNavData.h b/Engine/Source/Runtime/NavigationSystem/Public/AbstractNavData.h
new file mode 100644
--- /dev/null
+++ b/Engine/Source/Runtime/NavigationSystem/Public/AbstractNavData.h
@@ -0,0 +1,70 @@
+// Copyright 1998-2018 Epic Games, Inc. All Rights Reserved.
+
+#pragma once
+
+#include "CoreMinimal.h"
+#include "UObject/ObjectMacros.h"
+#include "NavFilters/NavigationQueryFilter.h"
+#include "AI/Navigation/NavigationTypes.h"
+#include "NavigationData.h"
+#include "AbstractNavData.generated.h"
+
+struct NAVIGATIONSYSTEM_API FAbstractNavigationPath : public FNavigationPath
+{
+	typedef FNavigationPath Super;
+
+	FAbstractNavigationPath();
+
+	static const FNavPathType Type;
+};
+
+class NAVIGATIONSYSTEM_API FAbstractQueryFilter : public INavigationQueryFilterInterface
+{
+public:
+	virtual void Reset() override {}
+	virtual void SetAreaCost(uint8 AreaType, float Cost) override {}
+	virtual void SetFixedAreaEnteringCost(uint8 AreaType, float Cost) override {}
+	virtual void SetExcludedArea(uint8 AreaType) override {}
+	virtual void SetAllAreaCosts(const float* CostArray, const int32 Count) override {}
+	virtual void GetAllAreaCosts(float* CostArray, float* FixedCostArray, const int32 Count) const override {}
+	virtual void SetBacktrackingEnabled(const bool bBacktracking) override {}
+	virtual bool IsBacktrackingEnabled() const override { return false; }
+	virtual bool IsEqual(const INavigationQueryFilterInterface* Other) const override { return true; }
+	virtual void SetIncludeFlags(uint16 Flags) override {}
+	virtual uint16 GetIncludeFlags() const override { return 0; }
+	virtual void SetExcludeFlags(uint16 Flags) override {}
+	virtual uint16 GetExcludeFlags() const override { return 0; }
+	virtual FVector GetAdjustedEndLocation(const FVector& EndLocation) const override { return EndLocation; }
+	virtual INavigationQueryFilterInterface* CreateCopy() const override;
+};
+
+UCLASS()
+class NAVIGATIONSYSTEM_API AAbstractNavData : public ANavigationData
+{
+	GENERATED_BODY()
+
+public:
+	AAbstractNavData(const FObjectInitializer& ObjectInitializer = FObjectInitializer::Get());
+
+	virtual void PostLoad() override;
+
+	// Begin ANavigationData overrides
+	virtual void BatchRaycast(TArray<FNavigationRaycastWork>& Workload, FSharedConstNavQueryFilter QueryFilter, const UObject* Querier = NULL) const override {};
+	virtual FNavLocation GetRandomPoint(FSharedConstNavQueryFilter Filter = NULL, const UObject* Querier = NULL) const override { return FNavLocation();  }
+	virtual bool GetRandomReachablePointInRadius(const FVector& Origin, float Radius, FNavLocation& OutResult, FSharedConstNavQueryFilter Filter = NULL, const UObject* Querier = NULL) const override { return false; }
+	virtual bool GetRandomPointInNavigableRadius(const FVector& Origin, float Radius, FNavLocation& OutResult, FSharedConstNavQueryFilter Filter = NULL, const UObject* Querier = NULL) const override { return false; }
+	virtual bool ProjectPoint(const FVector& Point, FNavLocation& OutLocation, const FVector& Extent, FSharedConstNavQueryFilter Filter = NULL, const UObject* Querier = NULL) const override { return false; }
+	virtual void BatchProjectPoints(TArray<FNavigationProjectionWork>& Workload, const FVector& Extent, FSharedConstNavQueryFilter Filter = NULL, const UObject* Querier = NULL) const override {};
+	virtual void BatchProjectPoints(TArray<FNavigationProjectionWork>& Workload, FSharedConstNavQueryFilter Filter = NULL, const UObject* Querier = NULL) const override {}
+	virtual ENavigationQueryResult::Type CalcPathCost(const FVector& PathStart, const FVector& PathEnd, float& OutPathCost, FSharedConstNavQueryFilter QueryFilter = NULL, const UObject* Querier = NULL) const override { return ENavigationQueryResult::Invalid; }
+	virtual ENavigationQueryResult::Type CalcPathLength(const FVector& PathStart, const FVector& PathEnd, float& OutPathLength, FSharedConstNavQueryFilter QueryFilter = NULL, const UObject* Querier = NULL) const override { return ENavigationQueryResult::Invalid; }
+	virtual ENavigationQueryResult::Type CalcPathLengthAndCost(const FVector& PathStart, const FVector& PathEnd, float& OutPathLength, float& OutPathCost, FSharedConstNavQueryFilter QueryFilter = NULL, const UObject* Querier = NULL) const override { return ENavigationQueryResult::Invalid;	}
+	virtual bool DoesNodeContainLocation(NavNodeRef NodeRef, const FVector& WorldSpaceLocation) const override { return true; }
+	virtual void OnNavAreaAdded(const UClass* NavAreaClass, int32 AgentIndex) override {}
+	virtual void OnNavAreaRemoved(const UClass* NavAreaClass) override {};
+	// End ANavigationData overrides
+
+	static FPathFindingResult FindPathAbstract(const FNavAgentProperties& AgentProperties, const FPathFindingQuery& Query);
+	static bool TestPathAbstract(const FNavAgentProperties& AgentProperties, const FPathFindingQuery& Query, int32* NumVisitedNodes);
+	static bool RaycastAbstract(const ANavigationData* NavDataInstance, const FVector& RayStart, const FVector& RayEnd, FVector& HitLocation, FSharedConstNavQueryFilter QueryFilter, const UObject* Querier);
+};
diff --git a/Engine/Source/Runtime/NavigationSystem/Public/CrowdManagerBase.h b/Engine/Source/Runtime/NavigationSystem/Public/CrowdManagerBase.h
new file mode 100644
--- /dev/null
+++ b/Engine/Source/Runtime/NavigationSystem/Public/CrowdManagerBase.h
@@ -0,0 +1,17 @@
+// Copyright 1998-2018 Epic Games, Inc. All Rights Reserved.
+
+#pragma once
+
+#include "CrowdManagerBase.generated.h"
+
+/** Base class for Crowd Managers. If you want to create a custom crowd manager
+ *	implement a class extending this one and set UNavigationSystem::CrowdManagerClass
+ *	to point at your class */
+UCLASS(Abstract, Transient)
+class NAVIGATIONSYSTEM_API UCrowdManagerBase : public UObject
+{
+	GENERATED_BODY()
+public:
+	virtual void Tick(float DeltaTime) PURE_VIRTUAL(UCrowdManagerBase::Tick, );
+};
+
diff --git a/Engine/Source/Runtime/NavigationSystem/Public/NavAreas/NavArea.h b/Engine/Source/Runtime/NavigationSystem/Public/NavAreas/NavArea.h
new file mode 100644
--- /dev/null
+++ b/Engine/Source/Runtime/NavigationSystem/Public/NavAreas/NavArea.h
@@ -0,0 +1,112 @@
+// Copyright 1998-2018 Epic Games, Inc. All Rights Reserved.
+
+#pragma once
+
+#include "CoreMinimal.h"
+#include "UObject/ObjectMacros.h"
+#include "UObject/Object.h"
+#include "Templates/SubclassOf.h"
+#include "AI/Navigation/NavAreaBase.h"
+#include "NavArea.generated.h"
+
+/** Class containing definition of a navigation area */
+UCLASS(DefaultToInstanced, abstract, Config=Engine, Blueprintable)
+class NAVIGATIONSYSTEM_API UNavArea : public UNavAreaBase
+{
+	GENERATED_BODY()
+
+public: 
+	/** travel cost multiplier for path distance */
+	UPROPERTY(EditAnywhere, Category=NavArea, config, meta=(ClampMin = "0.0"))
+	float DefaultCost;
+
+protected:
+	/** entering cost */
+	UPROPERTY(EditAnywhere, Category=NavArea, config, meta=(ClampMin = "0.0"))
+	float FixedAreaEnteringCost;
+
+public:
+	/** area color in navigation view */
+	UPROPERTY(EditAnywhere, Category=NavArea, config)
+	FColor DrawColor;
+
+	/** restrict area only to specified agents */
+	UPROPERTY(EditAnywhere, Category=NavArea, config)
+	FNavAgentSelector SupportedAgents;
+
+	// DEPRECATED AGENT CONFIG
+#if CPP
+	union
+	{
+		struct
+		{
+#endif
+	UPROPERTY(config)
+	uint32 bSupportsAgent0 : 1;
+	UPROPERTY(config)
+	uint32 bSupportsAgent1 : 1;
+	UPROPERTY(config)
+	uint32 bSupportsAgent2 : 1;
+	UPROPERTY(config)
+	uint32 bSupportsAgent3 : 1;
+	UPROPERTY(config)
+	uint32 bSupportsAgent4 : 1;
+	UPROPERTY(config)
+	uint32 bSupportsAgent5 : 1;
+	UPROPERTY(config)
+	uint32 bSupportsAgent6 : 1;
+	UPROPERTY(config)
+	uint32 bSupportsAgent7 : 1;
+	UPROPERTY(config)
+	uint32 bSupportsAgent8 : 1;
+	UPROPERTY(config)
+	uint32 bSupportsAgent9 : 1;
+	UPROPERTY(config)
+	uint32 bSupportsAgent10 : 1;
+	UPROPERTY(config)
+	uint32 bSupportsAgent11 : 1;
+	UPROPERTY(config)
+	uint32 bSupportsAgent12 : 1;
+	UPROPERTY(config)
+	uint32 bSupportsAgent13 : 1;
+	UPROPERTY(config)
+	uint32 bSupportsAgent14 : 1;
+	UPROPERTY(config)
+	uint32 bSupportsAgent15 : 1;
+#if CPP
+		};
+		uint32 SupportedAgentsBits;
+	};
+#endif
+
+	UNavArea(const FObjectInitializer& ObjectInitializer = FObjectInitializer::Get());
+
+	virtual void FinishDestroy() override;
+	virtual void PostLoad() override;
+	virtual void PostInitProperties() override;
+	virtual void Serialize(FArchive& Ar) override;
+
+	FORCEINLINE uint16 GetAreaFlags() const { return AreaFlags; }
+	FORCEINLINE bool HasFlags(uint16 InFlags) const { return (InFlags & AreaFlags) != 0; }
+
+	FORCEINLINE bool IsSupportingAgent(int32 AgentIndex) const { return SupportedAgents.Contains(AgentIndex); }
+
+	/** called before adding to navigation system */
+	virtual void InitializeArea() {};
+
+	/** Get the fixed area entering cost. */
+	virtual float GetFixedAreaEnteringCost() { return FixedAreaEnteringCost; }
+
+	/** Retrieved color declared for AreaDefinitionClass */
+	static FColor GetColor(UClass* AreaDefinitionClass);
+
+	/** copy properties from other area */
+	virtual void CopyFrom(TSubclassOf<UNavArea> AreaClass);
+
+protected:
+
+	/** these flags will be applied to navigation data along with AreaID */
+	uint16 AreaFlags;
+	
+	void RegisterArea();
+};
diff --git a/Engine/Source/Runtime/NavigationSystem/Public/NavAreas/NavAreaMeta.h b/Engine/Source/Runtime/NavigationSystem/Public/NavAreas/NavAreaMeta.h
new file mode 100644
--- /dev/null
+++ b/Engine/Source/Runtime/NavigationSystem/Public/NavAreas/NavAreaMeta.h
@@ -0,0 +1,29 @@
+// Copyright 1998-2018 Epic Games, Inc. All Rights Reserved.
+
+#pragma once
+
+#include "CoreMinimal.h"
+#include "UObject/ObjectMacros.h"
+#include "Templates/SubclassOf.h"
+#include "NavAreas/NavArea.h"
+#include "NavAreaMeta.generated.h"
+
+class AActor;
+
+/** A convenience class for an area that has IsMetaArea() == true.
+ *	Do not use this class when determining whether an area class is "meta". 
+ *	Call IsMetaArea instead. */
+UCLASS(Abstract)
+class NAVIGATIONSYSTEM_API UNavAreaMeta : public UNavArea
+{
+	GENERATED_BODY()
+
+public:
+	UNavAreaMeta(const FObjectInitializer& ObjectInitializer);
+
+	DEPRECATED(4.20, "UNavAreaMeta::PickAreaClass is deprecated. Use UNavArea::PickAreaClassForAgent instead")
+	static TSubclassOf<UNavArea> PickAreaClass(TSubclassOf<UNavArea> AreaClass, const AActor* Actor, const FNavAgentProperties& NavAgent);
+
+	DEPRECATED(4.20, "UNavAreaMeta::PickAreaClass is deprecated. Use UNavArea::PickAreaClassForAgent instead")
+	virtual TSubclassOf<UNavArea> PickAreaClass(const AActor* Actor, const FNavAgentProperties& NavAgent);
+};
diff --git a/Engine/Source/Runtime/NavigationSystem/Public/NavAreas/NavArea_Default.h b/Engine/Source/Runtime/NavigationSystem/Public/NavAreas/NavArea_Default.h
new file mode 100644
--- /dev/null
+++ b/Engine/Source/Runtime/NavigationSystem/Public/NavAreas/NavArea_Default.h
@@ -0,0 +1,15 @@
+// Copyright 1998-2018 Epic Games, Inc. All Rights Reserved.
+
+#pragma once
+
+#include "CoreMinimal.h"
+#include "UObject/ObjectMacros.h"
+#include "NavAreas/NavArea.h"
+#include "NavArea_Default.generated.h"
+
+/** Regular navigation area, applied to entire navigation data by default */
+UCLASS(Config=Engine)
+class NAVIGATIONSYSTEM_API UNavArea_Default : public UNavArea
+{
+	GENERATED_UCLASS_BODY()
+};
diff --git a/Engine/Source/Runtime/NavigationSystem/Public/NavAreas/NavArea_Null.h b/Engine/Source/Runtime/NavigationSystem/Public/NavAreas/NavArea_Null.h
new file mode 100644
--- /dev/null
+++ b/Engine/Source/Runtime/NavigationSystem/Public/NavAreas/NavArea_Null.h
@@ -0,0 +1,15 @@
+// Copyright 1998-2018 Epic Games, Inc. All Rights Reserved.
+
+#pragma once
+
+#include "CoreMinimal.h"
+#include "UObject/ObjectMacros.h"
+#include "NavAreas/NavArea.h"
+#include "NavArea_Null.generated.h"
+
+/** In general represents an empty area, that cannot be traversed by anyone. Ever.*/
+UCLASS(Config=Engine)
+class NAVIGATIONSYSTEM_API UNavArea_Null : public UNavArea
+{
+	GENERATED_UCLASS_BODY()
+};
diff --git a/Engine/Source/Runtime/NavigationSystem/Public/NavAreas/NavArea_Obstacle.h b/Engine/Source/Runtime/NavigationSystem/Public/NavAreas/NavArea_Obstacle.h
new file mode 100644
--- /dev/null
+++ b/Engine/Source/Runtime/NavigationSystem/Public/NavAreas/NavArea_Obstacle.h
@@ -0,0 +1,15 @@
+// Copyright 1998-2018 Epic Games, Inc. All Rights Reserved.
+
+#pragma once
+
+#include "CoreMinimal.h"
+#include "UObject/ObjectMacros.h"
+#include "NavAreas/NavArea.h"
+#include "NavArea_Obstacle.generated.h"
+
+/** In general represents a high cost area, that shouldn't be traversed by anyone unless no other path exist.*/
+UCLASS(Config = Engine)
+class NAVIGATIONSYSTEM_API UNavArea_Obstacle : public UNavArea
+{
+	GENERATED_UCLASS_BODY()
+};
diff --git a/Engine/Source/Runtime/NavigationSystem/Public/NavCollision.h b/Engine/Source/Runtime/NavigationSystem/Public/NavCollision.h
new file mode 100644
--- /dev/null
+++ b/Engine/Source/Runtime/NavigationSystem/Public/NavCollision.h
@@ -0,0 +1,117 @@
+// Copyright 1998-2018 Epic Games, Inc. All Rights Reserved.
+
+#pragma once
+
+#include "CoreMinimal.h"
+#include "UObject/ObjectMacros.h"
+#include "UObject/Object.h"
+#include "Misc/Guid.h"
+#include "Templates/SubclassOf.h"
+#include "Serialization/BulkData.h"
+#include "AI/Navigation/NavigationTypes.h"
+#include "AI/Navigation/NavCollisionBase.h"
+#include "NavCollision.generated.h"
+
+class FPrimitiveDrawInterface;
+struct FCompositeNavModifier;
+struct FNavigableGeometryExport;
+
+USTRUCT()
+struct FNavCollisionCylinder
+{
+	GENERATED_USTRUCT_BODY()
+
+	UPROPERTY(EditAnywhere, Category=Cylinder)
+	FVector Offset;
+
+	UPROPERTY(EditAnywhere, Category=Cylinder)
+	float Radius;
+
+	UPROPERTY(EditAnywhere, Category=Cylinder)
+	float Height;
+};
+
+USTRUCT()
+struct FNavCollisionBox
+{
+	GENERATED_USTRUCT_BODY()
+
+	UPROPERTY(EditAnywhere, Category=Box)
+	FVector Offset;
+
+	UPROPERTY(EditAnywhere, Category=Box)
+	FVector Extent;
+};
+
+UCLASS(config=Engine)
+class NAVIGATIONSYSTEM_API UNavCollision : public UNavCollisionBase
+{
+	GENERATED_UCLASS_BODY()
+
+	TNavStatArray<int32> ConvexShapeIndices;
+
+	/** list of nav collision cylinders */
+	UPROPERTY(EditAnywhere, Category=Navigation)
+	TArray<FNavCollisionCylinder> CylinderCollision;
+
+	/** list of nav collision boxes */
+	UPROPERTY(EditAnywhere, Category=Navigation)
+	TArray<FNavCollisionBox> BoxCollision;
+
+	/** navigation area type (empty = default obstacle) */
+	UPROPERTY(EditAnywhere, Category=Navigation)
+	TSubclassOf<class UNavArea> AreaClass;
+
+	/** If set, convex collisions will be exported offline for faster runtime navmesh building (increases memory usage) */
+	UPROPERTY(EditAnywhere, Category=Navigation, config)
+	uint32 bGatherConvexGeometry : 1;
+
+	/** if set, convex geometry will be rebuild instead of using cooked data */
+	uint32 bForceGeometryRebuild : 1;
+
+	/** Guid of associated BodySetup */
+	FGuid BodySetupGuid;
+
+	/** Cooked data for each format */
+	FFormatContainer CookedFormatData;
+
+	//~ Begin UObject Interface.
+	virtual void Serialize(FArchive& Ar) override;
+	virtual void PostLoad() override;
+	virtual void GetResourceSizeEx(FResourceSizeEx& CumulativeResourceSize) override;
+	virtual bool NeedsLoadForTargetPlatform(const class ITargetPlatform* TargetPlatform) const override;
+	//~ End UObject Interface.
+
+	FGuid GetGuid() const;
+
+	/** Tries to read data from DDC, and if that fails gathers navigation
+	 *	collision data, stores it and uploads to DDC */
+	virtual void Setup(class UBodySetup* BodySetup) override;
+
+	/** copy user settings from other nav collision data */
+	void CopyUserSettings(const UNavCollision& OtherData);
+
+	/** show cylinder and box collisions */
+	virtual void DrawSimpleGeom(FPrimitiveDrawInterface* PDI, const FTransform& Transform, const FColor Color) override;
+
+	/** Get data for dynamic obstacle */
+	virtual void GetNavigationModifier(FCompositeNavModifier& Modifier, const FTransform& LocalToWorld) override;
+
+	/** Export collision data */
+	virtual bool ExportGeometry(const FTransform& LocalToWorld, FNavigableGeometryExport& GeoExport) const override;
+
+	/** Read collisions data */
+	void GatherCollision();
+
+#if WITH_EDITOR
+	virtual void InvalidateCollision() override;
+#endif // WITH_EDITOR
+
+protected:
+	void ClearCollision();
+
+#if WITH_EDITOR
+	void InvalidatePhysicsData();
+#endif // WITH_EDITOR
+	FByteBulkData* GetCookedData(FName Format);
+};
diff --git a/Engine/Source/Runtime/NavigationSystem/Public/NavLinkComponent.h b/Engine/Source/Runtime/NavigationSystem/Public/NavLinkComponent.h
new file mode 100644
--- /dev/null
+++ b/Engine/Source/Runtime/NavigationSystem/Public/NavLinkComponent.h
@@ -0,0 +1,39 @@
+// Copyright 1998-2018 Epic Games, Inc. All Rights Reserved.
+
+#pragma once
+
+#include "CoreMinimal.h"
+#include "UObject/ObjectMacros.h"
+#include "Templates/SubclassOf.h"
+#include "Components/PrimitiveComponent.h"
+#include "AI/Navigation/NavLinkDefinition.h"
+#include "NavLinkHostInterface.h"
+#include "NavLinkComponent.generated.h"
+
+class FPrimitiveSceneProxy;
+struct FNavigationRelevantData;
+
+UCLASS(ClassGroup = (Navigation), meta = (BlueprintSpawnableComponent), hidecategories = (Activation))
+class NAVIGATIONSYSTEM_API UNavLinkComponent : public UPrimitiveComponent, public INavLinkHostInterface
+{
+	GENERATED_UCLASS_BODY()
+
+	UPROPERTY(EditAnywhere, BlueprintReadOnly, Category = Navigation)
+	TArray<FNavigationLink> Links;
+
+	virtual void GetNavigationData(FNavigationRelevantData& Data) const override;
+	virtual bool IsNavigationRelevant() const override;
+
+	virtual bool GetNavigationLinksClasses(TArray<TSubclassOf<class UNavLinkDefinition> >& OutClasses) const override { return false; }
+	virtual bool GetNavigationLinksArray(TArray<FNavigationLink>& OutLink, TArray<FNavigationSegmentLink>& OutSegments) const override;
+
+	virtual FBoxSphereBounds CalcBounds(const FTransform &LocalToWorld) const override;
+	virtual FPrimitiveSceneProxy* CreateSceneProxy() override;
+	virtual bool ShouldRecreateProxyOnUpdateTransform() const override { return true; }
+
+#if WITH_EDITOR
+	virtual void PostEditChangeProperty(FPropertyChangedEvent& PropertyChangedEvent) override;
+	virtual void PostEditUndo() override;
+	virtual void PostEditImport() override;
+#endif // WITH_EDITOR
+};
diff --git a/Engine/Source/Runtime/NavigationSystem/Public/NavLinkCustomComponent.h b/Engine/Source/Runtime/NavigationSystem/Public/NavLinkCustomComponent.h
new file mode 100644
--- /dev/null
+++ b/Engine/Source/Runtime/NavigationSystem/Public/NavLinkCustomComponent.h
@@ -0,0 +1,214 @@
+// Copyright 1998-2018 Epic Games, Inc. All Rights Reserved.
+
+#pragma once
+
+#include "CoreMinimal.h"
+#include "UObject/ObjectMacros.h"
+#include "Templates/SubclassOf.h"
+#include "Engine/EngineTypes.h"
+#include "NavAreas/NavArea.h"
+#include "AI/Navigation/NavLinkDefinition.h"
+#include "NavLinkCustomInterface.h"
+#include "NavRelevantComponent.h"
+#include "NavLinkCustomComponent.generated.h"
+
+
+struct FNavigationRelevantData;
+
+/**
+ *  Encapsulates NavLinkCustomInterface interface, can be used with Actors not relevant for navigation
+ *  
+ *  Additional functionality:
+ *  - can be toggled
+ *  - can create obstacle area for easier/forced separation of link end points
+ *  - can broadcast state changes to nearby agents
+ */
+
+UCLASS()
+class NAVIGATIONSYSTEM_API UNavLinkCustomComponent : public UNavRelevantComponent, public INavLinkCustomInterface
+{
+	GENERATED_UCLASS_BODY()
+
+	DECLARE_DELEGATE_ThreeParams(FOnMoveReachedLink, UNavLinkCustomComponent* /*ThisComp*/, UObject* /*PathComp*/, const FVector& /*DestPoint*/);
+	DECLARE_DELEGATE_TwoParams(FBroadcastFilter, UNavLinkCustomComponent* /*ThisComp*/, TArray<UObject*>& /*NotifyList*/);
+
+	// BEGIN INavLinkCustomInterface
+	virtual void GetLinkData(FVector& LeftPt, FVector& RightPt, ENavLinkDirection::Type& Direction) const override;
+	virtual TSubclassOf<UNavArea> GetLinkAreaClass() const override;
+	virtual uint32 GetLinkId() const override;
+	virtual void UpdateLinkId(uint32 NewUniqueId) override;
+	virtual bool IsLinkPathfindingAllowed(const UObject* Querier) const override;
+	virtual bool OnLinkMoveStarted(UObject* PathComp, const FVector& DestPoint) override;
+	virtual void OnLinkMoveFinished(UObject* PathComp) override;
+	// END INavLinkCustomInterface
+
+	//~ Begin UNavRelevantComponent Interface
+	virtual void GetNavigationData(FNavigationRelevantData& Data) const override;
+	virtual void CalcAndCacheBounds() const override;
+	//~ End UNavRelevantComponent Interface
+
+	//~ Begin UActorComponent Interface
+	virtual void OnRegister() override;
+	virtual void OnUnregister() override;
+	//~ End UActorComponent Interface
+
+	//~ Begin UObject Interface
+	virtual void PostLoad() override;
+#if WITH_EDITOR
+	virtual void PostEditImport() override;
+#endif
+	//~ End UObject Interface
+
+	/** set basic link data: end points and direction */
+	void SetLinkData(const FVector& RelativeStart, const FVector& RelativeEnd, ENavLinkDirection::Type Direction);
+	virtual FNavigationLink GetLinkModifier() const;
+
+	/** set area class to use when link is enabled */
+	void SetEnabledArea(TSubclassOf<UNavArea> AreaClass);
+	TSubclassOf<UNavArea> GetEnabledArea() const { return EnabledAreaClass; }
+
+	/** set area class to use when link is disabled */
+	void SetDisabledArea(TSubclassOf<UNavArea> AreaClass);
+	TSubclassOf<UNavArea> GetDisabledArea() const { return DisabledAreaClass; }
+
+	/** add box obstacle during generation of navigation data
+	  * this can be used to create empty area under doors */
+	void AddNavigationObstacle(TSubclassOf<UNavArea> AreaClass, const FVector& BoxExtent, const FVector& BoxOffset = FVector::ZeroVector);
+
+	/** removes simple obstacle */
+	void ClearNavigationObstacle();
+
+	/** set properties of trigger around link entry point(s), that will notify nearby agents about link state change */
+	void SetBroadcastData(float Radius, ECollisionChannel TraceChannel = ECC_Pawn, float Interval = 0.0f);
+
+	void SendBroadcastWhenEnabled(bool bEnabled);
+	void SendBroadcastWhenDisabled(bool bEnabled);
+
+	/** set delegate to filter  */
+	void SetBroadcastFilter(FBroadcastFilter const& InDelegate);
+
+	/** change state of smart link (used area class) */
+	void SetEnabled(bool bNewEnabled);
+	bool IsEnabled() const { return bLinkEnabled; }
+	
+	/** set delegate to notify about reaching this link during path following */
+	void SetMoveReachedLink(FOnMoveReachedLink const& InDelegate);
+
+	/** check is any agent is currently moving though this link */
+	bool HasMovingAgents() const;
+
+	/** get link start point in world space */
+	FVector GetStartPoint() const;
+
+	/** get link end point in world space */
+	FVector GetEndPoint() const;
+
+	//////////////////////////////////////////////////////////////////////////
+	// helper functions for setting delegates
+
+	template< class UserClass >	
+	FORCEINLINE void SetMoveReachedLink(UserClass* TargetOb, typename FOnMoveReachedLink::TUObjectMethodDelegate< UserClass >::FMethodPtr InFunc)
+	{
+		SetMoveReachedLink(FOnMoveReachedLink::CreateUObject(TargetOb, InFunc));
+	}
+	template< class UserClass >	
+	FORCEINLINE void SetMoveReachedLink(UserClass* TargetOb, typename FOnMoveReachedLink::TUObjectMethodDelegate_Const< UserClass >::FMethodPtr InFunc)
+	{
+		SetMoveReachedLink(FOnMoveReachedLink::CreateUObject(TargetOb, InFunc));
+	}
+
+	template< class UserClass >	
+	FORCEINLINE void SetBroadcastFilter(UserClass* TargetOb, typename FBroadcastFilter::TUObjectMethodDelegate< UserClass >::FMethodPtr InFunc)
+	{
+		SetBroadcastFilter(FBroadcastFilter::CreateUObject(TargetOb, InFunc));
+	}
+	template< class UserClass >	
+	FORCEINLINE void SetBroadcastFilter(UserClass* TargetOb, typename FBroadcastFilter::TUObjectMethodDelegate_Const< UserClass >::FMethodPtr InFunc)
+	{
+		SetBroadcastFilter(FBroadcastFilter::CreateUObject(TargetOb, InFunc));
+	}
+	
+protected:
+
+	/** link Id assigned by navigation system */
+	UPROPERTY()
+	uint32 NavLinkUserId;
+
+	/** area class to use when link is enabled */
+	UPROPERTY(EditAnywhere, Category=SmartLink)
+	TSubclassOf<UNavArea> EnabledAreaClass;
+
+	/** area class to use when link is disabled */
+	UPROPERTY(EditAnywhere, Category=SmartLink)
+	TSubclassOf<UNavArea> DisabledAreaClass;
+
+	/** start point, relative to owner */
+	UPROPERTY(EditAnywhere, Category=SmartLink)
+	FVector LinkRelativeStart;
+
+	/** end point, relative to owner */
+	UPROPERTY(EditAnywhere, Category=SmartLink)
+	FVector LinkRelativeEnd;
+
+	/** direction of link */
+	UPROPERTY(EditAnywhere, Category=SmartLink)
+	TEnumAsByte<ENavLinkDirection::Type> LinkDirection;
+
+	/** is link currently in enabled state? (area class) */
+	UPROPERTY(EditAnywhere, Category=SmartLink)
+	uint32 bLinkEnabled : 1;
+
+	/** should link notify nearby agents when it changes state to enabled */
+	UPROPERTY(EditAnywhere, Category=Broadcast)
+	uint32 bNotifyWhenEnabled : 1;
+
+	/** should link notify nearby agents when it changes state to disabled */
+	UPROPERTY(EditAnywhere, Category=Broadcast)
+	uint32 bNotifyWhenDisabled : 1;
+
+	/** if set, box obstacle area will be added to generation */
+	UPROPERTY(EditAnywhere, Category=Obstacle)
+	uint32 bCreateBoxObstacle : 1;
+
+	/** offset of simple box obstacle */
+	UPROPERTY(EditAnywhere, Category=Obstacle)
+	FVector ObstacleOffset;
+
+	/** extent of simple box obstacle */
+	UPROPERTY(EditAnywhere, Category=Obstacle)
+	FVector ObstacleExtent;
+
+	/** area class for simple box obstacle */
+	UPROPERTY(EditAnywhere, Category=Obstacle)
+	TSubclassOf<UNavArea> ObstacleAreaClass;
+
+	/** radius of state change broadcast */
+	UPROPERTY(EditAnywhere, Category=Broadcast)
+	float BroadcastRadius;
+
+	/** interval for state change broadcast (0 = single broadcast) */
+	UPROPERTY(EditAnywhere, Category=Broadcast, Meta=(UIMin="0.0", ClampMin="0.0"))
+	float BroadcastInterval;
+
+	/** trace channel for state change broadcast */
+	UPROPERTY(EditAnywhere, Category=Broadcast)
+	TEnumAsByte<ECollisionChannel> BroadcastChannel;
+
+	/** delegate to call when link is reached */
+	FBroadcastFilter OnBroadcastFilter;
+
+	/** list of agents moving though this link */
+	TArray<TWeakObjectPtr<UObject> > MovingAgents;
+
+	/** delegate to call when link is reached */
+	FOnMoveReachedLink OnMoveReachedLink;
+
+	/** Handle for efficient management of BroadcastStateChange timer */
+	FTimerHandle TimerHandle_BroadcastStateChange;
+
+	/** notify nearby agents about link changing state */
+	void BroadcastStateChange();
+	
+	/** gather agents to notify about state change */
+	void CollectNearbyAgents(TArray<UObject*>& NotifyList);
+};
diff --git a/Engine/Source/Runtime/NavigationSystem/Public/NavLinkRenderingComponent.h b/Engine/Source/Runtime/NavigationSystem/Public/NavLinkRenderingComponent.h
new file mode 100644
--- /dev/null
+++ b/Engine/Source/Runtime/NavigationSystem/Public/NavLinkRenderingComponent.h
@@ -0,0 +1,34 @@
+// Copyright 1998-2018 Epic Games, Inc. All Rights Reserved.
+
+#pragma once
+
+#include "CoreMinimal.h"
+#include "UObject/ObjectMacros.h"
+#include "Components/PrimitiveComponent.h"
+#include "NavLinkRenderingComponent.generated.h"
+
+class FPrimitiveSceneProxy;
+struct FConvexVolume;
+struct FEngineShowFlags;
+
+UCLASS(hidecategories=Object, editinlinenew)
+class NAVIGATIONSYSTEM_API UNavLinkRenderingComponent : public UPrimitiveComponent
+{
+	GENERATED_UCLASS_BODY()
+		
+	//~ Begin UPrimitiveComponent Interface
+	virtual FPrimitiveSceneProxy* CreateSceneProxy() override;
+
+	/** Should recreate proxy one very update */
+	virtual bool ShouldRecreateProxyOnUpdateTransform() const override { return true; }
+#if WITH_EDITOR
+	virtual bool ComponentIsTouchingSelectionBox(const FBox& InSelBBox, const FEngineShowFlags& ShowFlags, const bool bConsiderOnlyBSP, const bool bMustEncompassEntireComponent) const override;
+	virtual bool ComponentIsTouchingSelectionFrustum(const FConvexVolume& InFrustum, const FEngineShowFlags& ShowFlags, const bool bConsiderOnlyBSP, const bool bMustEncompassEntireComponent) const override;
+#endif
+	//~ End UPrimitiveComponent Interface
+
+	//~ Begin USceneComponent Interface
+	virtual FBoxSphereBounds CalcBounds(const FTransform &LocalToWorld) const override;
+	//~ End USceneComponent Interface
+};
+
diff --git a/Engine/Source/Runtime/NavigationSystem/Public/NavLinkRenderingProxy.h b/Engine/Source/Runtime/NavigationSystem/Public/NavLinkRenderingProxy.h
new file mode 100644
--- /dev/null
+++ b/Engine/Source/Runtime/NavigationSystem/Public/NavLinkRenderingProxy.h
@@ -0,0 +1,63 @@
+// Copyright 1998-2018 Epic Games, Inc. All Rights Reserved.
+
+#pragma once
+
+#include "CoreMinimal.h"
+#include "PrimitiveViewRelevance.h"
+#include "PrimitiveSceneProxy.h"
+#include "AI/Navigation/NavLinkDefinition.h"
+
+class FMaterialRenderProxy;
+class FMeshElementCollector;
+class FPrimitiveDrawInterface;
+class UPrimitiveComponent;
+
+class NAVIGATIONSYSTEM_API FNavLinkRenderingProxy : public FPrimitiveSceneProxy
+{
+private:
+	AActor* LinkOwnerActor;
+	class INavLinkHostInterface* LinkOwnerHost;
+
+public:
+	SIZE_T GetTypeHash() const override;
+
+	struct FNavLinkDrawing
+	{
+		FVector Left;
+		FVector Right;
+		ENavLinkDirection::Type Direction;
+		FColor Color;
+		float SnapRadius;
+		float SnapHeight;
+		uint32 SupportedAgentsBits;
+	};
+	struct FNavLinkSegmentDrawing
+	{
+		FVector LeftStart, LeftEnd;
+		FVector RightStart, RightEnd;
+		ENavLinkDirection::Type Direction;
+		FColor Color;
+		float SnapRadius;
+		float SnapHeight;
+		uint32 SupportedAgentsBits;
+	};
+
+private:
+	TArray<FNavLinkDrawing> OffMeshPointLinks;
+	TArray<FNavLinkSegmentDrawing> OffMeshSegmentLinks;
+
+public:
+	/** Initialization constructor. */
+	FNavLinkRenderingProxy(const UPrimitiveComponent* InComponent);
+	virtual void GetDynamicMeshElements(const TArray<const FSceneView*>& Views, const FSceneViewFamily& ViewFamily, uint32 VisibilityMap, FMeshElementCollector& Collector) const override;
+	virtual FPrimitiveViewRelevance GetViewRelevance(const FSceneView* View) const override;
+	virtual uint32 GetMemoryFootprint( void ) const override;
+	uint32 GetAllocatedSize( void ) const;
+	void StorePointLinks(const FTransform& LocalToWorld, const TArray<FNavigationLink>& LinksArray);
+	void StoreSegmentLinks(const FTransform& LocalToWorld, const TArray<FNavigationSegmentLink>& LinksArray);
+
+	static void GetLinkMeshes(const TArray<FNavLinkDrawing>& OffMeshPointLinks, const TArray<FNavLinkSegmentDrawing>& OffMeshSegmentLinks, TArray<float>& StepHeights, FMaterialRenderProxy* const MeshColorInstance, int32 ViewIndex, FMeshElementCollector& Collector, uint32 AgentMask);
+
+	/** made static to allow consistent navlinks drawing even if something is drawing links without FNavLinkRenderingProxy */
+	static void DrawLinks(FPrimitiveDrawInterface* PDI, TArray<FNavLinkDrawing>& OffMeshPointLinks, TArray<FNavLinkSegmentDrawing>& OffMeshSegmentLinks, TArray<float>& StepHeights, FMaterialRenderProxy* const MeshColorInstance, uint32 AgentMask);
+};
diff --git a/Engine/Source/Runtime/NavigationSystem/Public/NavLinkTrivial.h b/Engine/Source/Runtime/NavigationSystem/Public/NavLinkTrivial.h
new file mode 100644
--- /dev/null
+++ b/Engine/Source/Runtime/NavigationSystem/Public/NavLinkTrivial.h
@@ -0,0 +1,14 @@
+// Copyright 1998-2018 Epic Games, Inc. All Rights Reserved.
+
+#pragma once
+
+#include "CoreMinimal.h"
+#include "UObject/ObjectMacros.h"
+#include "AI/Navigation/NavLinkDefinition.h"
+#include "NavLinkTrivial.generated.h"
+
+UCLASS()
+class NAVIGATIONSYSTEM_API UNavLinkTrivial : public UNavLinkDefinition
+{
+	GENERATED_UCLASS_BODY()
+};
diff --git a/Engine/Source/Runtime/NavigationSystem/Public/NavMesh/NavMeshBoundsVolume.h b/Engine/Source/Runtime/NavigationSystem/Public/NavMesh/NavMeshBoundsVolume.h
new file mode 100644
--- /dev/null
+++ b/Engine/Source/Runtime/NavigationSystem/Public/NavMesh/NavMeshBoundsVolume.h
@@ -0,0 +1,36 @@
+// Copyright 1998-2018 Epic Games, Inc. All Rights Reserved.
+
+/**
+ * this volume only blocks the path builder - it has no gameplay collision
+ *
+ */
+
+#pragma once
+
+#include "CoreMinimal.h"
+#include "UObject/ObjectMacros.h"
+#include "AI/Navigation/NavigationTypes.h"
+#include "GameFramework/Volume.h"
+#include "NavMeshBoundsVolume.generated.h"
+
+
+UCLASS()
+class NAVIGATIONSYSTEM_API ANavMeshBoundsVolume : public AVolume
+{
+	GENERATED_UCLASS_BODY()
+
+	UPROPERTY(EditAnywhere, Category = Navigation)
+	FNavAgentSelector SupportedAgents;
+
+	//~ Begin AActor Interface
+	virtual void PostRegisterAllComponents() override;
+	virtual void PostUnregisterAllComponents() override;
+	//~ End AActor Interface
+#if WITH_EDITOR
+	//~ Begin UObject Interface
+	virtual void PostEditChangeProperty( struct FPropertyChangedEvent& PropertyChangedEvent) override;
+	virtual void PostEditUndo() override;
+	//~ End UObject Interface
+#endif // WITH_EDITOR
+};
+
diff --git a/Engine/Source/Runtime/NavigationSystem/Public/NavMesh/NavMeshPath.h b/Engine/Source/Runtime/NavigationSystem/Public/NavMesh/NavMeshPath.h
new file mode 100644
--- /dev/null
+++ b/Engine/Source/Runtime/NavigationSystem/Public/NavMesh/NavMeshPath.h
@@ -0,0 +1,173 @@
+// Copyright 1998-2018 Epic Games, Inc. All Rights Reserved.
+
+#pragma once
+
+#include "CoreMinimal.h"
+#include "UObject/ObjectMacros.h"
+#include "UObject/Object.h"
+#include "NavigationPath.h"
+//#include "NavigationPath.generated.h"
+
+#define RECAST_STRAIGHTPATH_OFFMESH_CONNECTION 0x04
+
+/** Helper to translate FNavPathPoint.Flags. */
+struct NAVIGATIONSYSTEM_API FNavMeshNodeFlags
+{
+	/** Extra node information (like "path start", "off-mesh connection"). */
+	uint8 PathFlags;
+	/** Area type after this node. */
+	uint8 Area;
+	/** Area flags for this node. */
+	uint16 AreaFlags;
+
+	FNavMeshNodeFlags() : PathFlags(0), Area(0), AreaFlags(0) {}
+	FNavMeshNodeFlags(uint32 Flags) : PathFlags(Flags), Area(Flags >> 8), AreaFlags(Flags >> 16) {}
+	uint32 Pack() const { return PathFlags | ((uint32)Area << 8) | ((uint32)AreaFlags << 16); }
+	bool IsNavLink() const { return (PathFlags & RECAST_STRAIGHTPATH_OFFMESH_CONNECTION) != 0; }
+
+	FNavMeshNodeFlags& AddAreaFlags(const uint16 InAreaFlags)
+	{
+		AreaFlags = (AreaFlags | InAreaFlags);
+		return *this;
+	}
+};
+
+
+struct NAVIGATIONSYSTEM_API FNavMeshPath : public FNavigationPath
+{
+	typedef FNavigationPath Super;
+
+	FNavMeshPath();
+
+	FORCEINLINE void SetWantsStringPulling(const bool bNewWantsStringPulling) { bWantsStringPulling = bNewWantsStringPulling; }
+	FORCEINLINE bool WantsStringPulling() const { return bWantsStringPulling; }
+	FORCEINLINE bool IsStringPulled() const { return bStringPulled; }
+
+	/** find string pulled path from PathCorridor */
+	void PerformStringPulling(const FVector& StartLoc, const FVector& EndLoc);
+
+	FORCEINLINE void SetWantsPathCorridor(const bool bNewWantsPathCorridor) { bWantsPathCorridor = bNewWantsPathCorridor; }
+	FORCEINLINE bool WantsPathCorridor() const { return bWantsPathCorridor; }
+
+	FORCEINLINE const TArray<FNavigationPortalEdge>& GetPathCorridorEdges() const { return bCorridorEdgesGenerated ? PathCorridorEdges : GeneratePathCorridorEdges(); }
+	FORCEINLINE void SetPathCorridorEdges(const TArray<FNavigationPortalEdge>& InPathCorridorEdges) { PathCorridorEdges = InPathCorridorEdges; bCorridorEdgesGenerated = true; }
+
+	FORCEINLINE void OnPathCorridorUpdated() { bCorridorEdgesGenerated = false; }
+
+	virtual void DebugDraw(const ANavigationData* NavData, FColor PathColor, UCanvas* Canvas, bool bPersistent, const uint32 NextPathPointIndex = 0) const override;
+
+	bool ContainsWithSameEnd(const FNavMeshPath* Other) const;
+
+	void OffsetFromCorners(float Distance);
+
+	void ApplyFlags(int32 NavDataFlags);
+
+	virtual void ResetForRepath() override;
+
+	/** get flags of path point or corridor poly (depends on bStringPulled flag) */
+	bool GetNodeFlags(int32 NodeIdx, FNavMeshNodeFlags& Flags) const;
+
+	/** get cost of path, starting from next poly in corridor */
+	virtual float GetCostFromNode(NavNodeRef PathNode) const override { return GetCostFromIndex(PathCorridor.Find(PathNode) + 1); }
+
+	/** get cost of path, starting from given point */
+	virtual float GetCostFromIndex(int32 PathPointIndex) const override
+	{
+		float TotalCost = 0.f;
+		const float* Cost = PathCorridorCost.GetData();
+		for (int32 PolyIndex = PathPointIndex; PolyIndex < PathCorridorCost.Num(); ++PolyIndex, ++Cost)
+		{
+			TotalCost += *Cost;
+		}
+
+		return TotalCost;
+	}
+
+	FORCEINLINE_DEBUGGABLE float GetTotalPathLength() const
+	{
+		return bStringPulled ? GetStringPulledLength(0) : GetPathCorridorLength(0);
+	}
+
+	FORCEINLINE int32 GetNodeRefIndex(const NavNodeRef NodeRef) const { return PathCorridor.Find(NodeRef); }
+
+	/** check if path (all polys in corridor) contains given node */
+	virtual bool ContainsNode(NavNodeRef NodeRef) const override { return PathCorridor.Contains(NodeRef); }
+
+	virtual bool ContainsCustomLink(uint32 UniqueLinkId) const override { return CustomLinkIds.Contains(UniqueLinkId); }
+	virtual bool ContainsAnyCustomLink() const override { return CustomLinkIds.Num() > 0; }
+
+	bool IsPathSegmentANavLink(const int32 PathSegmentStartIndex) const;
+
+	virtual bool DoesIntersectBox(const FBox& Box, uint32 StartingIndex = 0, int32* IntersectingSegmentIndex = NULL, FVector* AgentExtent = NULL) const override;
+	virtual bool DoesIntersectBox(const FBox& Box, const FVector& AgentLocation, uint32 StartingIndex = 0, int32* IntersectingSegmentIndex = NULL, FVector* AgentExtent = NULL) const override;
+	/** retrieves normalized direction vector to given path segment. If path is not string pulled navigation corridor is being used */
+	virtual FVector GetSegmentDirection(uint32 SegmentEndIndex) const override;
+
+private:
+	bool DoesPathIntersectBoxImplementation(const FBox& Box, const FVector& StartLocation, uint32 StartingIndex, int32* IntersectingSegmentIndex, FVector* AgentExtent) const;
+	void InternalResetNavMeshPath();
+
+public:
+
+#if ENABLE_VISUAL_LOG
+	virtual void DescribeSelfToVisLog(struct FVisualLogEntry* Snapshot) const override;
+	virtual FString GetDescription() const override;
+#endif // ENABLE_VISUAL_LOG
+
+protected:
+	/** calculates total length of string pulled path. Does not generate string pulled
+	*	path if it's not already generated (see bWantsStringPulling and bStrigPulled)
+	*	Internal use only */
+	float GetStringPulledLength(const int32 StartingPoint) const;
+
+	/** calculates estimated length of path expressed as sequence of navmesh edges.
+	*	It basically sums up distances between every subsequent nav edge pair edge middles.
+	*	Internal use only */
+	float GetPathCorridorLength(const int32 StartingEdge) const;
+
+	/** it's only const to be callable in const environment. It's not supposed to be called directly externally anyway,
+	*	just as part of retrieving corridor on demand or generating it in internal processes. It fills a mutable
+	*	array. */
+	const TArray<FNavigationPortalEdge>& GeneratePathCorridorEdges() const;
+
+public:
+
+	/** sequence of navigation mesh poly ids representing an obstacle-free navigation corridor */
+	TArray<NavNodeRef> PathCorridor;
+
+	/** for every poly in PathCorridor stores traversal cost from previous navpoly */
+	TArray<float> PathCorridorCost;
+
+	/** set of unique link Ids */
+	TArray<uint32> CustomLinkIds;
+
+private:
+	/** sequence of FVector pairs where each pair represents navmesh portal edge between two polygons navigation corridor.
+	*	Note, that it should always be accessed via GetPathCorridorEdges() since PathCorridorEdges content is generated
+	*	on first access */
+	mutable TArray<FNavigationPortalEdge> PathCorridorEdges;
+
+	/** transient variable indicating whether PathCorridorEdges contains up to date information */
+	mutable uint32 bCorridorEdgesGenerated : 1;
+
+public:
+	/** is this path generated on dynamic navmesh (i.e. one attached to moving surface) */
+	uint32 bDynamic : 1;
+
+protected:
+	/** does this path contain string pulled path? If true then NumPathVerts > 0 and
+	*	OutPathVerts contains valid data. If false there's only navigation corridor data
+	*	available.*/
+	uint32 bStringPulled : 1;
+
+	/** If set to true path instance will contain a string pulled version. Otherwise only
+	*	navigation corridor will be available. Defaults to true */
+	uint32 bWantsStringPulling : 1;
+
+	/** If set to true path instance will contain path corridor generated as a part
+	*	pathfinding call (i.e. without the need to generate it with GeneratePathCorridorEdges */
+	uint32 bWantsPathCorridor : 1;
+
+public:
+	static const FNavPathType Type;
+};
\ No newline at end of file
diff --git a/Engine/Source/Runtime/NavigationSystem/Public/NavMesh/NavMeshRenderingComponent.h b/Engine/Source/Runtime/NavigationSystem/Public/NavMesh/NavMeshRenderingComponent.h
new file mode 100644
--- /dev/null
+++ b/Engine/Source/Runtime/NavigationSystem/Public/NavMesh/NavMeshRenderingComponent.h
@@ -0,0 +1,218 @@
+// Copyright 1998-2018 Epic Games, Inc. All Rights Reserved.
+
+#pragma once
+
+#include "CoreMinimal.h"
+#include "UObject/ObjectMacros.h"
+#include "Engine/EngineTypes.h"
+#include "PrimitiveViewRelevance.h"
+#include "RenderResource.h"
+#include "MaterialShared.h"
+#include "DynamicMeshBuilder.h"
+#include "DebugRenderSceneProxy.h"
+#include "Components/PrimitiveComponent.h"
+#include "MeshBatch.h"
+#include "LocalVertexFactory.h"
+#include "GenericOctree.h"
+#include "StaticMeshResources.h"
+#include "NavMeshRenderingComponent.generated.h"
+
+class APlayerController;
+class ARecastNavMesh;
+class FMeshElementCollector;
+class FPrimitiveDrawInterface;
+class UCanvas;
+class UNavMeshRenderingComponent;
+
+enum class ENavMeshDetailFlags : uint8
+{
+	TriangleEdges,
+	PolyEdges,
+	BoundaryEdges,
+	FilledPolys,
+	TileBounds,
+	PathCollidingGeometry,
+	TileLabels,
+	PolygonLabels,
+	PolygonCost,
+	PathLabels,
+	NavLinks,
+	FailedNavLinks,
+	Clusters,
+	NavOctree,
+	NavOctreeDetails,
+};
+
+// exported to API for GameplayDebugger module
+struct NAVIGATIONSYSTEM_API FNavMeshSceneProxyData : public TSharedFromThis<FNavMeshSceneProxyData, ESPMode::ThreadSafe>
+{
+	struct FDebugMeshData
+	{
+		TArray<FDynamicMeshVertex> Vertices;
+		TArray<uint32> Indices;
+		FColor ClusterColor;
+	};
+	TArray<FDebugMeshData> MeshBuilders;
+
+	TArray<FDebugRenderSceneProxy::FDebugLine> ThickLineItems;
+	TArray<FDebugRenderSceneProxy::FDebugLine> TileEdgeLines;
+	TArray<FDebugRenderSceneProxy::FDebugLine> NavMeshEdgeLines;
+	TArray<FDebugRenderSceneProxy::FDebugLine> NavLinkLines;
+	TArray<FDebugRenderSceneProxy::FDebugLine> ClusterLinkLines;
+
+	struct FDebugText
+	{
+		FVector Location;
+		FString Text;
+
+		FDebugText() {}
+		FDebugText(const FVector& InLocation, const FString& InText) : Location(InLocation), Text(InText) {}
+	};
+	TArray<FDebugText> DebugLabels;
+	
+	TArray<FBoxCenterAndExtent>	OctreeBounds;
+
+	FBox Bounds;
+	FVector NavMeshDrawOffset;
+	uint32 bDataGathered : 1;
+	uint32 bNeedsNewData : 1;
+	int32 NavDetailFlags;
+
+	FNavMeshSceneProxyData() : NavMeshDrawOffset(0, 0, 10.f),
+		bDataGathered(false), bNeedsNewData(true), NavDetailFlags(0) {}
+
+	void Reset();
+	void Serialize(FArchive& Ar);
+	uint32 GetAllocatedSize() const;
+
+#if WITH_RECAST
+	int32 GetDetailFlags(const ARecastNavMesh* NavMesh) const;
+	void GatherData(const ARecastNavMesh* NavMesh, int32 InNavDetailFlags, const TArray<int32>& TileSet);
+#endif
+};
+
+// exported to API for GameplayDebugger module
+class NAVIGATIONSYSTEM_API FNavMeshSceneProxy final : public FDebugRenderSceneProxy
+{
+	friend class FNavMeshDebugDrawDelegateHelper;
+public:
+	SIZE_T GetTypeHash() const override;
+
+	FNavMeshSceneProxy(const UPrimitiveComponent* InComponent, FNavMeshSceneProxyData* InProxyData, bool ForceToRender = false);
+	virtual ~FNavMeshSceneProxy();
+
+	virtual void GetDynamicMeshElements(const TArray<const FSceneView*>& Views, const FSceneViewFamily& ViewFamily, uint32 VisibilityMap, FMeshElementCollector& Collector) const override;
+
+protected:
+	void DrawDebugBox(FPrimitiveDrawInterface* PDI, FVector const& Center, FVector const& Box, FColor const& Color) const;
+	virtual FPrimitiveViewRelevance GetViewRelevance(const FSceneView* View) const override;
+
+	virtual uint32 GetMemoryFootprint(void) const override { return sizeof(*this) + GetAllocatedSize(); }
+	uint32 GetAllocatedSize(void) const;
+
+private:			
+	FNavMeshSceneProxyData ProxyData;
+
+	FDynamicMeshIndexBuffer32 IndexBuffer;
+	FStaticMeshVertexBuffers VertexBuffers;
+	FLocalVertexFactory VertexFactory;
+
+	TArray<FColoredMaterialRenderProxy> MeshColors;
+	TArray<FMeshBatchElement> MeshBatchElements;
+
+	FDebugDrawDelegate DebugTextDrawingDelegate;
+	FDelegateHandle DebugTextDrawingDelegateHandle;
+	TWeakObjectPtr<UNavMeshRenderingComponent> RenderingComponent;
+	uint32 bRequestedData : 1;
+	uint32 bForceRendering : 1;
+	uint32 bSkipDistanceCheck : 1;
+	uint32 bUseThickLines : 1;
+};
+
+#if WITH_RECAST && !UE_BUILD_SHIPPING && !UE_BUILD_TEST
+class FNavMeshDebugDrawDelegateHelper : public FDebugDrawDelegateHelper
+{
+	typedef FDebugDrawDelegateHelper Super;
+
+public:
+	FNavMeshDebugDrawDelegateHelper()
+		: bForceRendering(false)
+		, bNeedsNewData(false)
+	{
+	}
+
+	virtual void InitDelegateHelper(const FDebugRenderSceneProxy* InSceneProxy) override
+	{
+		check(0);
+	}
+
+	void InitDelegateHelper(const FNavMeshSceneProxy* InSceneProxy)
+	{
+		Super::InitDelegateHelper(InSceneProxy);
+
+		DebugLabels.Reset();
+		DebugLabels.Append(InSceneProxy->ProxyData.DebugLabels);
+		bForceRendering = InSceneProxy->bForceRendering;
+		bNeedsNewData = InSceneProxy->ProxyData.bNeedsNewData;
+	}
+
+	NAVIGATIONSYSTEM_API virtual void RegisterDebugDrawDelgate() override;
+	NAVIGATIONSYSTEM_API virtual void UnregisterDebugDrawDelgate() override;
+
+protected:
+	NAVIGATIONSYSTEM_API virtual void DrawDebugLabels(UCanvas* Canvas, APlayerController*) override;
+
+private:
+	TArray<FNavMeshSceneProxyData::FDebugText> DebugLabels;
+	uint32 bForceRendering : 1;
+	uint32 bNeedsNewData : 1;
+};
+#endif
+
+UCLASS(hidecategories=Object, editinlinenew)
+class NAVIGATIONSYSTEM_API UNavMeshRenderingComponent : public UPrimitiveComponent
+{
+	GENERATED_UCLASS_BODY()
+
+public:
+	
+	//~ Begin UPrimitiveComponent Interface
+	virtual FPrimitiveSceneProxy* CreateSceneProxy() override;
+	virtual void OnRegister()  override;
+	virtual void OnUnregister()  override;
+	//~ End UPrimitiveComponent Interface
+
+	//~ Begin UActorComponent Interface
+	virtual void CreateRenderState_Concurrent() override;
+	virtual void DestroyRenderState_Concurrent() override;
+	//~ End UActorComponent Interface
+
+	//~ Begin USceneComponent Interface
+	virtual FBoxSphereBounds CalcBounds(const FTransform &LocalToWorld) const override;
+	//~ End USceneComponent Interface
+
+	void ForceUpdate() { bForceUpdate = true; }
+	bool IsForcingUpdate() const { return bForceUpdate; }
+
+	static bool IsNavigationShowFlagSet(const UWorld* World);
+
+protected:
+	void TimerFunction();
+
+protected:
+	uint32 bCollectNavigationData : 1;
+	uint32 bForceUpdate : 1;
+	FTimerHandle TimerHandle;
+
+protected:
+#if WITH_RECAST && !UE_BUILD_SHIPPING && !UE_BUILD_TEST
+	FNavMeshDebugDrawDelegateHelper NavMeshDebugDrawDelgateManager;
+#endif
+};
+
+namespace FNavMeshRenderingHelpers
+{
+	NAVIGATIONSYSTEM_API void AddVertex(FNavMeshSceneProxyData::FDebugMeshData& MeshData, const FVector& Pos, const FColor Color = FColor::White);
+
+	NAVIGATIONSYSTEM_API void AddTriangle(FNavMeshSceneProxyData::FDebugMeshData& MeshData, int32 V0, int32 V1, int32 V2);
+}
diff --git a/Engine/Source/Runtime/NavigationSystem/Public/NavMesh/NavTestRenderingComponent.h b/Engine/Source/Runtime/NavigationSystem/Public/NavMesh/NavTestRenderingComponent.h
new file mode 100644
--- /dev/null
+++ b/Engine/Source/Runtime/NavigationSystem/Public/NavMesh/NavTestRenderingComponent.h
@@ -0,0 +1,139 @@
+// Copyright 1998-2018 Epic Games, Inc. All Rights Reserved.
+
+#pragma once
+
+#include "CoreMinimal.h"
+#include "UObject/ObjectMacros.h"
+#include "PrimitiveViewRelevance.h"
+#include "Components/PrimitiveComponent.h"
+#include "DynamicMeshBuilder.h"
+#include "DebugRenderSceneProxy.h"
+
+class ANavigationTestingActor;
+class APlayerController;
+class FMeshElementCollector;
+class UCanvas;
+class UNavTestRenderingComponent;
+struct FNodeDebugData;
+
+class FNavTestSceneProxy final : public FDebugRenderSceneProxy
+{
+	friend class FNavTestDebugDrawDelegateHelper;
+
+public:
+	SIZE_T GetTypeHash() const override;
+
+	struct FNodeDebugData
+	{
+		NavNodeRef PolyRef;
+		FVector Position;
+		FString Desc;
+		FSetElementId ParentId;
+		uint32 bClosedSet : 1;
+		uint32 bBestPath : 1;
+		uint32 bModified : 1;
+		uint32 bOffMeshLink : 1;
+
+		FORCEINLINE bool operator==(const FNodeDebugData& Other) const
+		{
+			return PolyRef == Other.PolyRef;
+		}
+		FORCEINLINE friend uint32 GetTypeHash(const FNodeDebugData& Other)
+		{
+			return Other.PolyRef;
+		}
+	};
+
+	FNavTestSceneProxy(const UNavTestRenderingComponent* InComponent);
+
+	~FNavTestSceneProxy() {}
+
+	virtual void GetDynamicMeshElements(const TArray<const FSceneView*>& Views, const FSceneViewFamily& ViewFamily, uint32 VisibilityMap, FMeshElementCollector& Collector) const override;
+
+	void GatherPathPoints();
+
+	void GatherPathStep();
+
+	FORCEINLINE static bool LocationInView(const FVector& Location, const FSceneView* View)
+	{
+		return View->ViewFrustum.IntersectBox(Location, FVector::ZeroVector);
+	}
+
+	virtual FPrimitiveViewRelevance GetViewRelevance(const FSceneView* View) const override;
+
+	virtual uint32 GetMemoryFootprint(void) const override { return sizeof(*this) + GetAllocatedSize(); }
+
+	uint32 GetAllocatedSize(void) const;
+
+private:
+	FVector NavMeshDrawOffset;
+	ANavigationTestingActor* NavTestActor;
+	FDebugDrawDelegate DebugTextDrawingDelegate;
+	FDelegateHandle DebugTextDrawingDelegateHandle;
+	TArray<FVector> PathPoints;
+	TArray<FString> PathPointFlags;
+
+	TArray<FDynamicMeshVertex> OpenSetVerts;
+	TArray<uint32> OpenSetIndices;
+	TArray<FDynamicMeshVertex> ClosedSetVerts;
+	TArray<uint32> ClosedSetIndices;
+	TSet<FNodeDebugData> NodeDebug;
+	FSetElementId BestNodeId;
+
+	FVector ClosestWallLocation;
+
+	uint32 bShowBestPath : 1;
+	uint32 bShowNodePool : 1;
+	uint32 bShowDiff : 1;
+};
+
+#if WITH_RECAST && WITH_EDITOR
+class FNavTestDebugDrawDelegateHelper : public FDebugDrawDelegateHelper
+{
+	typedef FDebugDrawDelegateHelper Super;
+
+public:
+	virtual void InitDelegateHelper(const FDebugRenderSceneProxy* InSceneProxy) override
+	{
+		check(0);
+	}
+
+	void InitDelegateHelper(const FNavTestSceneProxy* InSceneProxy);
+
+	virtual void RegisterDebugDrawDelgate() override;
+	virtual void UnregisterDebugDrawDelgate() override;
+
+protected:
+	void DrawDebugLabels(UCanvas* Canvas, APlayerController*) override;
+
+private:
+	TSet<FNavTestSceneProxy::FNodeDebugData> NodeDebug;
+	ANavigationTestingActor* NavTestActor;
+	TArray<FVector> PathPoints;
+	TArray<FString> PathPointFlags;
+	FSetElementId BestNodeId;
+	uint32 bShowBestPath : 1;
+	uint32 bShowDiff : 1;
+};
+#endif //WITH_RECAST && WITH_EDITOR
+
+#include "NavTestRenderingComponent.generated.h"
+
+class FPrimitiveSceneProxy;
+struct FTransform;
+
+UCLASS(hidecategories=Object)
+class UNavTestRenderingComponent: public UPrimitiveComponent
+{
+	GENERATED_UCLASS_BODY()
+
+	virtual FPrimitiveSceneProxy* CreateSceneProxy() override;
+	virtual FBoxSphereBounds CalcBounds(const FTransform &LocalToWorld) const override;
+	virtual void CreateRenderState_Concurrent() override;
+	virtual void DestroyRenderState_Concurrent() override;
+
+private:
+#if WITH_RECAST && WITH_EDITOR
+	FNavTestDebugDrawDelegateHelper NavTestDebugDrawDelegateHelper;
+#endif
+};
diff --git a/Engine/Source/Runtime/NavigationSystem/Public/NavMesh/RecastHelpers.h b/Engine/Source/Runtime/NavigationSystem/Public/NavMesh/RecastHelpers.h
new file mode 100644
--- /dev/null
+++ b/Engine/Source/Runtime/NavigationSystem/Public/NavMesh/RecastHelpers.h
@@ -0,0 +1,24 @@
+// Copyright 1998-2018 Epic Games, Inc. All Rights Reserved.
+
+/**
+ * Coord system utilities
+ *
+ * Translates between Unreal and Recast coords.
+ * Unreal: x, y, z
+ * Recast: -x, z, -y
+ */
+
+#pragma once
+
+#include "CoreMinimal.h"
+
+extern NAVIGATIONSYSTEM_API FVector Unreal2RecastPoint(const float* UnrealPoint);
+extern NAVIGATIONSYSTEM_API FVector Unreal2RecastPoint(const FVector& UnrealPoint);
+extern NAVIGATIONSYSTEM_API FBox Unreal2RecastBox(const FBox& UnrealBox);
+extern NAVIGATIONSYSTEM_API FMatrix Unreal2RecastMatrix();
+
+extern NAVIGATIONSYSTEM_API FVector Recast2UnrealPoint(const float* RecastPoint);
+extern NAVIGATIONSYSTEM_API FVector Recast2UnrealPoint(const FVector& RecastPoint);
+extern NAVIGATIONSYSTEM_API FBox Recast2UnrealBox(const float* RecastMin, const float* RecastMax);
+extern NAVIGATIONSYSTEM_API FBox Recast2UnrealBox(const FBox& RecastBox);
+extern NAVIGATIONSYSTEM_API FMatrix Recast2UnrealMatrix();
diff --git a/Engine/Source/Runtime/NavigationSystem/Public/NavMesh/RecastNavMesh.h b/Engine/Source/Runtime/NavigationSystem/Public/NavMesh/RecastNavMesh.h
new file mode 100644
--- /dev/null
+++ b/Engine/Source/Runtime/NavigationSystem/Public/NavMesh/RecastNavMesh.h
@@ -0,0 +1,941 @@
+// Copyright 1998-2018 Epic Games, Inc. All Rights Reserved.
+
+#pragma once
+
+#include "CoreMinimal.h"
+#include "UObject/ObjectMacros.h"
+#include "Templates/SubclassOf.h"
+#include "EngineDefines.h"
+#include "AI/Navigation/NavigationTypes.h"
+#include "NavigationSystemTypes.h"
+#include "NavigationData.h"
+#include "NavMesh/NavMeshPath.h"
+#include "RecastNavMesh.generated.h"
+
+#define RECAST_MAX_SEARCH_NODES		2048
+
+#define RECAST_MIN_TILE_SIZE		300.f
+
+#define RECAST_MAX_AREAS			64
+#define RECAST_DEFAULT_AREA			(RECAST_MAX_AREAS - 1)
+#define RECAST_LOW_AREA				(RECAST_MAX_AREAS - 2)
+#define RECAST_NULL_AREA			0
+#define RECAST_UNWALKABLE_POLY_COST	FLT_MAX
+
+class FPImplRecastNavMesh;
+class FRecastQueryFilter;
+class INavLinkCustomInterface;
+class UCanvas;
+class UNavArea;
+class UNavigationSystem;
+class UPrimitiveComponent;
+class URecastNavMeshDataChunk;
+struct FRecastAreaNavModifierElement;
+class dtNavMesh;
+class dtQueryFilter;
+
+UENUM()
+namespace ERecastPartitioning
+{
+	// keep in sync with rcRegionPartitioning enum!
+
+	enum Type
+	{
+		Monotone,
+		Watershed,
+		ChunkyMonotone,
+	};
+}
+
+namespace ERecastPathFlags
+{
+	/** If set, path won't be post processed. */
+	const int32 SkipStringPulling = (1 << 0);
+
+	/** If set, path will contain navigation corridor. */
+	const int32 GenerateCorridor = (1 << 1);
+}
+
+#if WITH_RECAST
+struct FRecastDebugPathfindingNode
+{
+	NavNodeRef PolyRef;
+	NavNodeRef ParentRef;
+	float Cost;
+	float TotalCost;
+	float Length;
+
+	FVector NodePos;
+	TArray<FVector, TInlineAllocator<6> > Verts;
+	uint8 NumVerts;
+
+	uint8 bOpenSet : 1;
+	uint8 bOffMeshLink : 1;
+	uint8 bModified : 1;
+
+	FRecastDebugPathfindingNode() : PolyRef(0), ParentRef(0), NumVerts(0) {}
+	FRecastDebugPathfindingNode(NavNodeRef InPolyRef) : PolyRef(InPolyRef), ParentRef(0), NumVerts(0) {}
+
+	FORCEINLINE bool operator==(const NavNodeRef& OtherPolyRef) const { return PolyRef == OtherPolyRef; }
+	FORCEINLINE bool operator==(const FRecastDebugPathfindingNode& Other) const { return PolyRef == Other.PolyRef; }
+	FORCEINLINE friend uint32 GetTypeHash(const FRecastDebugPathfindingNode& Other) { return GetTypeHash(Other.PolyRef); }
+
+	FORCEINLINE float GetHeuristicCost() const { return TotalCost - Cost; }
+};
+
+namespace ERecastDebugPathfindingFlags
+{
+	enum Type
+	{
+		Basic = 0x0,
+		BestNode = 0x1,
+		Vertices = 0x2,
+		PathLength = 0x4
+	};
+}
+
+struct FRecastDebugPathfindingData
+{
+	TSet<FRecastDebugPathfindingNode> Nodes;
+	FSetElementId BestNode;
+	uint8 Flags;
+
+	FRecastDebugPathfindingData() : Flags(ERecastDebugPathfindingFlags::Basic) {}
+	FRecastDebugPathfindingData(ERecastDebugPathfindingFlags::Type InFlags) : Flags(InFlags) {}
+};
+
+struct FRecastDebugGeometry
+{
+	enum EOffMeshLinkEnd
+	{
+		OMLE_None = 0x0,
+		OMLE_Left = 0x1,
+		OMLE_Right = 0x2,
+		OMLE_Both = OMLE_Left | OMLE_Right
+	};
+
+	struct FOffMeshLink
+	{
+		FVector Left;
+		FVector Right;
+		uint8	AreaID;
+		uint8	Direction;
+		uint8	ValidEnds;
+		float	Radius;
+		float	Height;
+		FColor	Color;
+	};
+
+	struct FCluster
+	{
+		TArray<int32> MeshIndices;
+	};
+
+	struct FClusterLink
+	{
+		FVector FromCluster;
+		FVector ToCluster;
+	};
+
+	struct FOffMeshSegment
+	{
+		FVector LeftStart, LeftEnd;
+		FVector RightStart, RightEnd;
+		uint8	AreaID;
+		uint8	Direction;
+		uint8	ValidEnds;
+	};
+
+	TArray<FVector> MeshVerts;
+	TArray<int32> AreaIndices[RECAST_MAX_AREAS];
+	TArray<int32> BuiltMeshIndices;
+	TArray<FVector> PolyEdges;
+	TArray<FVector> NavMeshEdges;
+	TArray<FOffMeshLink> OffMeshLinks;
+	TArray<FCluster> Clusters;
+	TArray<FClusterLink> ClusterLinks;
+	TArray<FOffMeshSegment> OffMeshSegments;
+	TArray<int32> OffMeshSegmentAreas[RECAST_MAX_AREAS];
+
+	int32 bGatherPolyEdges : 1;
+	int32 bGatherNavMeshEdges : 1;
+
+	FRecastDebugGeometry() : bGatherPolyEdges(false), bGatherNavMeshEdges(false)
+	{}
+
+	uint32 NAVIGATIONSYSTEM_API GetAllocatedSize() const;
+};
+
+struct FNavPoly
+{
+	NavNodeRef Ref;
+	FVector Center;
+};
+
+namespace ERecastNamedFilter
+{
+	enum Type 
+	{
+		FilterOutNavLinks = 0,		// filters out all off-mesh connections
+		FilterOutAreas,				// filters out all navigation areas except the default one (RECAST_DEFAULT_AREA)
+		FilterOutNavLinksAndAreas,	// combines FilterOutNavLinks and FilterOutAreas
+
+		NamedFiltersCount,
+	};
+}
+#endif //WITH_RECAST
+
+
+/**
+ *	Structure to handle nav mesh tile's raw data persistence and releasing
+ */
+struct FNavMeshTileData
+{
+	// helper function so that we release NavData via dtFree not regular delete (for navigation mem stats)
+	struct FNavData
+	{
+		uint8* RawNavData;
+
+		FNavData(uint8* InNavData) : RawNavData(InNavData) {}
+		~FNavData();
+	};
+	
+	// layer index
+	int32	LayerIndex;
+	FBox	LayerBBox;
+	// size of allocated data
+	int32	DataSize;
+	// actual tile data
+	TSharedPtr<FNavData, ESPMode::ThreadSafe> NavData;
+	
+	FNavMeshTileData() : LayerIndex(0), DataSize(0) { }
+	~FNavMeshTileData();
+	
+	explicit FNavMeshTileData(uint8* RawData, int32 RawDataSize, int32 LayerIdx = 0, FBox LayerBounds = FBox(ForceInit));
+		
+	FORCEINLINE uint8* GetData()
+	{
+		check(NavData.IsValid());
+		return NavData->RawNavData;
+	}
+
+	FORCEINLINE const uint8* GetData() const
+	{
+		check(NavData.IsValid());
+		return NavData->RawNavData;
+	}
+
+	FORCEINLINE uint8* GetDataSafe()
+	{
+		return NavData.IsValid() ? NavData->RawNavData : NULL;
+	}
+
+	FORCEINLINE bool operator==(const uint8* RawData) const
+	{
+		return GetData() == RawData;
+	}
+
+	FORCEINLINE bool IsValid() const { return NavData.IsValid() && GetData() != nullptr && DataSize > 0; }
+
+	uint8* Release();
+
+	// Duplicate shared state so we will have own copy of the data
+	void MakeUnique();
+};
+
+DECLARE_MULTICAST_DELEGATE(FOnNavMeshUpdate);
+
+namespace FNavMeshConfig
+{
+	struct FRecastNamedFiltersCreator
+	{
+		FRecastNamedFiltersCreator(bool bVirtualFilters);
+	};
+}
+
+
+UCLASS(config=Engine, defaultconfig, hidecategories=(Input,Rendering,Tags,Transform,"Utilities|Transformation",Actor,Layers,Replication), notplaceable)
+class NAVIGATIONSYSTEM_API ARecastNavMesh : public ANavigationData
+{
+	GENERATED_UCLASS_BODY()
+
+	typedef uint16 FNavPolyFlags;
+
+	virtual void Serialize( FArchive& Ar ) override;
+
+	/** should we draw edges of every navmesh's triangle */
+	UPROPERTY(EditAnywhere, Category=Display)
+	uint32 bDrawTriangleEdges:1;
+
+	/** should we draw edges of every poly (i.e. not only border-edges)  */
+	UPROPERTY(EditAnywhere, Category=Display, config)
+	uint32 bDrawPolyEdges:1;
+
+	/** if disabled skips filling drawn navmesh polygons */
+	UPROPERTY(EditAnywhere, Category = Display)
+	uint32 bDrawFilledPolys:1;
+
+	/** should we draw border-edges */
+	UPROPERTY(EditAnywhere, Category=Display)
+	uint32 bDrawNavMeshEdges:1;
+
+	/** should we draw the tile boundaries */
+	UPROPERTY(EditAnywhere, Category=Display)
+	uint32 bDrawTileBounds:1;
+	
+	/** Draw input geometry passed to the navmesh generator.  Recommend disabling other geometry rendering via viewport showflags in editor. */
+	UPROPERTY(EditAnywhere, Category=Display)
+	uint32 bDrawPathCollidingGeometry:1;
+
+	UPROPERTY(EditAnywhere, Category=Display)
+	uint32 bDrawTileLabels:1;
+
+	UPROPERTY(EditAnywhere, Category=Display)
+	uint32 bDrawPolygonLabels:1;
+
+	UPROPERTY(EditAnywhere, Category=Display)
+	uint32 bDrawDefaultPolygonCost:1;
+
+	UPROPERTY(EditAnywhere, Category=Display)
+	uint32 bDrawLabelsOnPathNodes:1;
+
+	UPROPERTY(EditAnywhere, Category=Display)
+	uint32 bDrawNavLinks:1;
+
+	UPROPERTY(EditAnywhere, Category=Display)
+	uint32 bDrawFailedNavLinks:1;
+	
+	UPROPERTY(EditAnywhere, Category=Display)
+	uint32 bDrawClusters:1;
+
+	/** should we draw edges of every navmesh's triangle */
+	UPROPERTY(EditAnywhere, Category = Display)
+	uint32 bDrawOctree : 1;
+
+	UPROPERTY(EditAnywhere, Category = Display, meta=(editcondition = "bDrawOctree"))
+	uint32 bDrawOctreeDetails : 1;
+
+	UPROPERTY(config)
+	uint32 bDistinctlyDrawTilesBeingBuilt:1;
+
+	UPROPERTY(EditAnywhere, Category = Display)
+	uint32 bDrawNavMesh : 1;
+
+	/** vertical offset added to navmesh's debug representation for better readability */
+	UPROPERTY(EditAnywhere, Category=Display, config)
+	float DrawOffset;
+	
+	//----------------------------------------------------------------------//
+	// NavMesh generation parameters
+	//----------------------------------------------------------------------//
+
+	/** if true, the NavMesh will allocate fixed size pool for tiles, should be enabled to support streaming */
+	UPROPERTY(EditAnywhere, Category=Generation, config)
+	uint32 bFixedTilePoolSize:1;
+
+	/** maximum number of tiles NavMesh can hold */
+	UPROPERTY(EditAnywhere, Category=Generation, config, meta=(editcondition = "bFixedTilePoolSize"))
+	int32 TilePoolSize;
+
+	/** size of single tile, expressed in uu */
+	UPROPERTY(EditAnywhere, Category=Generation, config, meta=(ClampMin = "300.0"))
+	float TileSizeUU;
+
+	/** horizontal size of voxelization cell */
+	UPROPERTY(EditAnywhere, Category = Generation, config, meta = (ClampMin = "1.0", ClampMax = "1024.0"))
+	float CellSize;
+
+	/** vertical size of voxelization cell */
+	UPROPERTY(EditAnywhere, Category = Generation, config, meta = (ClampMin = "1.0", ClampMax = "1024.0"))
+	float CellHeight;
+
+	/** Radius of smallest agent to traverse this navmesh */
+	UPROPERTY(EditAnywhere, Category = Generation, config, meta = (ClampMin = "0.0"))
+	float AgentRadius;
+
+	UPROPERTY(EditAnywhere, Category = Generation, config, meta = (ClampMin = "0.0"))
+	float AgentHeight;
+
+	/** Size of the tallest agent that will path with this navmesh. */
+	UPROPERTY(EditAnywhere, Category=Generation, config, meta=(ClampMin = "0.0"))
+	float AgentMaxHeight;
+
+	/* The maximum slope (angle) that the agent can move on. */ 
+	UPROPERTY(EditAnywhere, Category=Generation, config, meta=(ClampMin = "0.0", ClampMax = "89.0", UIMin = "0.0", UIMax = "89.0" ))
+	float AgentMaxSlope;
+
+	UPROPERTY(EditAnywhere, Category = Generation, config, meta = (ClampMin = "0.0"))
+	float AgentMaxStepHeight;
+
+	/* The minimum dimension of area. Areas smaller than this will be discarded */
+	UPROPERTY(EditAnywhere, Category=Generation, config, meta=(ClampMin = "0.0"))
+	float MinRegionArea;
+
+	/* The size limit of regions to be merged with bigger regions (watershed partitioning only) */
+	UPROPERTY(EditAnywhere, Category=Generation, config, meta=(ClampMin = "0.0"))
+	float MergeRegionSize;
+
+	/** How much navigable shapes can get simplified - the higher the value the more freedom */
+	UPROPERTY(EditAnywhere, Category = Generation, config, meta = (ClampMin = "0.0"))
+	float MaxSimplificationError;
+
+	UPROPERTY(EditAnywhere, Category = Generation, config, meta = (ClampMin = "0", UIMin = "0"), AdvancedDisplay)
+	int32 MaxSimultaneousTileGenerationJobsCount;
+
+	/** Absolute hard limit to number of navmesh tiles. Be very, very careful while modifying it while
+	 *	having big maps with navmesh. A single, empty tile takes 176 bytes and empty tiles are
+	 *	allocated up front (subject to change, but that's where it's at now)
+	 *	@note TileNumberHardLimit is always rounded up to the closest power of 2 */
+	UPROPERTY(EditAnywhere, Category = Generation, config, meta = (ClampMin = "1", UIMin = "1"), AdvancedDisplay)
+	int32 TileNumberHardLimit;
+
+	UPROPERTY(VisibleAnywhere, Category = Generation, AdvancedDisplay)
+	int32 PolyRefTileBits;
+
+	UPROPERTY(VisibleAnywhere, Category = Generation, AdvancedDisplay)
+	int32 PolyRefNavPolyBits;
+
+	UPROPERTY(VisibleAnywhere, Category = Generation, AdvancedDisplay)
+	int32 PolyRefSaltBits;
+
+	/** navmesh draw distance in game (always visible in editor) */
+	UPROPERTY(config)
+	float DefaultDrawDistance;
+
+	/** specifes default limit to A* nodes used when performing navigation queries. 
+	 *	Can be overridden by passing custom FNavigationQueryFilter */
+	UPROPERTY(config)
+	float DefaultMaxSearchNodes;
+
+	/** specifes default limit to A* nodes used when performing hierarchical navigation queries. */
+	UPROPERTY(config)
+	float DefaultMaxHierarchicalSearchNodes;
+
+	/** partitioning method for creating navmesh polys */
+	UPROPERTY(EditAnywhere, Category=Generation, config, AdvancedDisplay)
+	TEnumAsByte<ERecastPartitioning::Type> RegionPartitioning;
+
+	/** partitioning method for creating tile layers */
+	UPROPERTY(EditAnywhere, Category=Generation, config, AdvancedDisplay)
+	TEnumAsByte<ERecastPartitioning::Type> LayerPartitioning;
+
+	/** number of chunk splits (along single axis) used for region's partitioning: ChunkyMonotone */
+	UPROPERTY(EditAnywhere, Category=Generation, config, AdvancedDisplay)
+	int32 RegionChunkSplits;
+
+	/** number of chunk splits (along single axis) used for layer's partitioning: ChunkyMonotone */
+	UPROPERTY(EditAnywhere, Category=Generation, config, AdvancedDisplay)
+	int32 LayerChunkSplits;
+
+	/** Controls whether Navigation Areas will be sorted by cost before application 
+	 *	to navmesh during navmesh generation. This is relevant then there are
+	 *	areas overlapping and we want to have area cost express area relevancy
+	 *	as well. Setting it to true will result in having area sorted by cost,
+	 *	but it will also increase navmesh generation cost a bit */
+	UPROPERTY(EditAnywhere, Category=Generation, config)
+	uint32 bSortNavigationAreasByCost:1;
+
+	/** controls whether voxel filterring will be applied (via FRecastTileGenerator::ApplyVoxelFilter). 
+	 *	Results in generated navemesh better fitting navigation bounds, but hits (a bit) generation performance */
+	UPROPERTY(EditAnywhere, Category=Generation, config, AdvancedDisplay)
+	uint32 bPerformVoxelFiltering:1;
+
+	/** mark areas with insufficient free height above instead of cutting them out (accessible only for area modifiers using replace mode) */
+	UPROPERTY(EditAnywhere, Category = Generation, config, AdvancedDisplay)
+	uint32 bMarkLowHeightAreas : 1;
+
+	/** if set, only single low height span will be allowed under valid one */
+	UPROPERTY(EditAnywhere, Category = Generation, config, AdvancedDisplay)
+	uint32 bFilterLowSpanSequences : 1;
+
+	/** if set, only low height spans with corresponding area modifier will be stored in tile cache (reduces memory, can't modify without full tile rebuild) */
+	UPROPERTY(EditAnywhere, Category = Generation, config, AdvancedDisplay)
+	uint32 bFilterLowSpanFromTileCache : 1;
+
+	UPROPERTY(EditAnywhere, Category = Generation, config, AdvancedDisplay)
+	uint32 bDoFullyAsyncNavDataGathering : 1;
+	
+	/** TODO: switch to disable new code from OffsetFromCorners if necessary - remove it later */
+	UPROPERTY(config)
+	uint32 bUseBetterOffsetsFromCorners : 1;
+
+	/** If set, tiles generated without any navmesh data will be marked to distinguish them from not generated / streamed out ones. Defaults to false. */
+	UPROPERTY(config)
+	uint32 bStoreEmptyTileLayers : 1;
+
+	/** Indicates whether default navigation filters will use virtual functions. Defaults to true. */
+	UPROPERTY(config)
+	uint32 bUseVirtualFilters : 1;
+
+	/** If set, paths can end at navlink poly (not the ground one!) */
+	UPROPERTY(config)
+	uint32 bAllowNavLinkAsPathEnd : 1;
+
+private:
+	/** Cache rasterized voxels instead of just collision vertices/indices in navigation octree */
+	UPROPERTY(config)
+	uint32 bUseVoxelCache : 1;
+
+	/** indicates how often we will sort navigation tiles to mach players position */
+	UPROPERTY(config)
+	float TileSetUpdateInterval;
+	
+	/** contains last available dtPoly's flag bit set (8th bit at the moment of writing) */
+	static FNavPolyFlags NavLinkFlag;
+
+	/** Squared draw distance */
+	static float DrawDistanceSq;
+
+public:
+
+	struct FRaycastResult
+	{
+		enum 
+		{
+			MAX_PATH_CORRIDOR_POLYS = 128
+		};
+
+		NavNodeRef CorridorPolys[MAX_PATH_CORRIDOR_POLYS];
+		float CorridorCost[MAX_PATH_CORRIDOR_POLYS];
+		int32 CorridorPolysCount;
+		float HitTime;
+		FVector HitNormal;
+		uint32 bIsRaycastEndInCorridor : 1;
+
+		FRaycastResult()
+			: CorridorPolysCount(0)
+			, HitTime(FLT_MAX)
+			, HitNormal(0.f)
+			, bIsRaycastEndInCorridor(false)
+		{
+			FMemory::Memzero(CorridorPolys);
+			FMemory::Memzero(CorridorCost);
+		}
+
+		FORCEINLINE int32 GetMaxCorridorSize() const { return MAX_PATH_CORRIDOR_POLYS; }
+		FORCEINLINE bool HasHit() const { return HitTime != FLT_MAX; }
+		FORCEINLINE NavNodeRef GetLastNodeRef() const { return CorridorPolysCount > 0 ? CorridorPolys[CorridorPolysCount - 1] : INVALID_NAVNODEREF; }
+	};
+
+	//----------------------------------------------------------------------//
+	// Recast runtime params
+	//----------------------------------------------------------------------//
+	/** Euclidean distance heuristic scale used while pathfinding */
+	UPROPERTY(EditAnywhere, Category = Query, config, meta = (ClampMin = "0.1"))
+	float HeuristicScale;
+
+	/** Value added to each search height to compensate for error between navmesh polys and walkable geometry  */
+	UPROPERTY(EditAnywhere, Category = Query, config, meta = (ClampMin = "0.0"))
+	float VerticalDeviationFromGroundCompensation;
+
+	/** broadcast for navmesh updates */
+	FOnNavMeshUpdate OnNavMeshUpdate;
+
+	FORCEINLINE static void SetDrawDistance(float NewDistance) { DrawDistanceSq = NewDistance * NewDistance; }
+	FORCEINLINE static float GetDrawDistanceSq() { return DrawDistanceSq; }
+
+	//////////////////////////////////////////////////////////////////////////
+
+	bool HasValidNavmesh() const;
+
+#if WITH_RECAST
+	//----------------------------------------------------------------------//
+	// Life cycle & serialization
+	//----------------------------------------------------------------------//
+public:
+	/** Dtor */
+	virtual ~ARecastNavMesh();
+
+	//~ Begin UObject Interface
+	virtual void PostInitProperties() override;
+	virtual void PostLoad() override;
+	virtual void PostRegisterAllComponents() override;
+
+#if WITH_EDITOR
+	virtual void PostEditChangeProperty( struct FPropertyChangedEvent& PropertyChangedEvent) override;
+#endif // WITH_EDITOR
+	//~ End UObject Interface
+
+#if WITH_EDITOR
+	/** RecastNavMesh instances are dynamically spawned and should not be coppied */
+	virtual bool ShouldExport() override { return false; }
+#endif
+	
+	virtual void CleanUp() override;
+
+	//~ Begin ANavigationData Interface
+	virtual FNavLocation GetRandomPoint(FSharedConstNavQueryFilter Filter = NULL, const UObject* Querier = NULL) const override;
+	virtual bool GetRandomReachablePointInRadius(const FVector& Origin, float Radius, FNavLocation& OutResult, FSharedConstNavQueryFilter Filter = NULL, const UObject* Querier = NULL) const override;
+	virtual bool GetRandomPointInNavigableRadius(const FVector& Origin, float Radius, FNavLocation& OutResult, FSharedConstNavQueryFilter Filter = NULL, const UObject* Querier = NULL) const override;
+
+	virtual bool ProjectPoint(const FVector& Point, FNavLocation& OutLocation, const FVector& Extent, FSharedConstNavQueryFilter Filter = NULL, const UObject* Querier = NULL) const override;
+	
+	/** Project batch of points using shared search extent and filter */
+	virtual void BatchProjectPoints(TArray<FNavigationProjectionWork>& Workload, const FVector& Extent, FSharedConstNavQueryFilter Filter = NULL, const UObject* Querier = NULL) const override;
+
+	/** Project batch of points using shared search filter. This version is not requiring user to pass in Extent, 
+	 *	and is instead relying on FNavigationProjectionWork.ProjectionLimit.
+	 *	@note function will assert if item's FNavigationProjectionWork.ProjectionLimit is invalid */
+	virtual void BatchProjectPoints(TArray<FNavigationProjectionWork>& Workload, FSharedConstNavQueryFilter Filter = NULL, const UObject* Querier = NULL) const override;
+	
+	virtual ENavigationQueryResult::Type CalcPathCost(const FVector& PathStart, const FVector& PathEnd, float& OutPathCost, FSharedConstNavQueryFilter Filter = NULL, const UObject* Querier = NULL) const override;
+	virtual ENavigationQueryResult::Type CalcPathLength(const FVector& PathStart, const FVector& PathEnd, float& OutPathLength, FSharedConstNavQueryFilter QueryFilter = NULL, const UObject* Querier = NULL) const override;
+	virtual ENavigationQueryResult::Type CalcPathLengthAndCost(const FVector& PathStart, const FVector& PathEnd, float& OutPathLength, float& OutPathCost, FSharedConstNavQueryFilter QueryFilter = NULL, const UObject* Querier = NULL) const override;
+	virtual bool DoesNodeContainLocation(NavNodeRef NodeRef, const FVector& WorldSpaceLocation) const override;
+
+	virtual UPrimitiveComponent* ConstructRenderingComponent() override;
+	/** Returns bounding box for the navmesh. */
+	virtual FBox GetBounds() const override { return GetNavMeshBounds(); }
+	/** Called on world origin changes **/
+	virtual void ApplyWorldOffset(const FVector& InOffset, bool bWorldShift) override;
+
+	virtual void OnStreamingLevelAdded(ULevel* InLevel, UWorld* InWorld) override;
+	virtual void OnStreamingLevelRemoved(ULevel* InLevel, UWorld* InWorld) override;
+	//~ End ANavigationData Interface
+
+	virtual void AttachNavMeshDataChunk(URecastNavMeshDataChunk& NavDataChunk);
+	virtual void DetachNavMeshDataChunk(URecastNavMeshDataChunk& NavDataChunk);
+
+protected:
+	/** Serialization helper. */
+	void SerializeRecastNavMesh(FArchive& Ar, FPImplRecastNavMesh*& NavMesh, int32 NavMeshVersion);
+
+	TArray<FIntPoint>& GetActiveTiles(); 
+	virtual void RestrictBuildingToActiveTiles(bool InRestrictBuildingToActiveTiles) override;
+
+public:
+	/** Whether NavMesh should adjust his tile pool size when NavBounds are changed */
+	bool IsResizable() const;
+
+	/** Returns bounding box for the whole navmesh. */
+	FBox GetNavMeshBounds() const;
+
+	/** Returns bounding box for a given navmesh tile. */
+	FBox GetNavMeshTileBounds(int32 TileIndex) const;
+
+	/** Retrieves XY coordinates of tile specified by index */
+	bool GetNavMeshTileXY(int32 TileIndex, int32& OutX, int32& OutY, int32& Layer) const;
+
+	/** Retrieves XY coordinates of tile specified by position */
+	bool GetNavMeshTileXY(const FVector& Point, int32& OutX, int32& OutY) const;
+
+	/** Retrieves all tile indices at matching XY coordinates */
+	void GetNavMeshTilesAt(int32 TileX, int32 TileY, TArray<int32>& Indices) const;
+
+	/** Retrieves number of tiles in this navmesh */
+	int32 GetNavMeshTilesCount() const;
+
+	/** Removes compressed tile data at given tile coord */
+	void RemoveTileCacheLayers(int32 TileX, int32 TileY);
+	
+	/** Stores compressed tile data for given tile coord */
+	void AddTileCacheLayers(int32 TileX, int32 TileY, const TArray<FNavMeshTileData>& InLayers);
+
+	/** Marks tile coord as rebuild and empty */
+	void MarkEmptyTileCacheLayers(int32 TileX, int32 TileY);
+	
+	/** Returns compressed tile data at given tile coord */
+	TArray<FNavMeshTileData> GetTileCacheLayers(int32 TileX, int32 TileY) const;
+	
+	void GetEdgesForPathCorridor(const TArray<NavNodeRef>* PathCorridor, TArray<struct FNavigationPortalEdge>* PathCorridorEdges) const;
+
+	void UpdateDrawing();
+
+	/** Creates a task to be executed on GameThread calling UpdateDrawing */
+	void RequestDrawingUpdate(bool bForce = false);
+
+	/** called after regenerating tiles */
+	virtual void OnNavMeshTilesUpdated(const TArray<uint32>& ChangedTiles);
+
+	/** Event from generator that navmesh build has finished */
+	virtual void OnNavMeshGenerationFinished();
+
+	virtual void EnsureBuildCompletion() override;
+
+	virtual void SetConfig(const FNavDataConfig& Src) override;
+protected:
+	virtual void FillConfig(FNavDataConfig& Dest) override;
+
+	FORCEINLINE const FNavigationQueryFilter& GetRightFilterRef(FSharedConstNavQueryFilter Filter) const 
+	{
+		return *(Filter.IsValid() ? Filter.Get() : GetDefaultQueryFilter().Get());
+	}
+
+public:
+
+	static bool IsVoxelCacheEnabled();
+
+	//----------------------------------------------------------------------//
+	// Debug                                                                
+	//----------------------------------------------------------------------//
+	/** Debug rendering support. */
+	void GetDebugGeometry(FRecastDebugGeometry& OutGeometry, int32 TileIndex = INDEX_NONE) const;
+
+	// @todo docuement
+	void DrawDebugPathCorridor(NavNodeRef const* PathPolys, int32 NumPathPolys, bool bPersistent=true) const;
+
+#if !UE_BUILD_SHIPPING
+	virtual uint32 LogMemUsed() const override;
+#endif // !UE_BUILD_SHIPPING
+
+	void UpdateNavMeshDrawing();
+
+	//----------------------------------------------------------------------//
+	// Utilities
+	//----------------------------------------------------------------------//
+	virtual void OnNavAreaChanged() override;
+	virtual void OnNavAreaAdded(const UClass* NavAreaClass, int32 AgentIndex) override;
+	virtual int32 GetNewAreaID(const UClass* AreaClass) const override;
+	virtual int32 GetMaxSupportedAreas() const override { return RECAST_MAX_AREAS; }
+
+	/** Get forbidden area flags from default query filter */
+	uint16 GetDefaultForbiddenFlags() const;
+	/** Change forbidden area flags in default query filter */
+	void SetDefaultForbiddenFlags(uint16 ForbiddenAreaFlags);
+
+	/** Area sort function */
+	virtual void SortAreasForGenerator(TArray<FRecastAreaNavModifierElement>& Areas) const;
+
+	void RecreateDefaultFilter();
+
+	int32 GetMaxSimultaneousTileGenerationJobsCount() const { return MaxSimultaneousTileGenerationJobsCount; }
+	void SetMaxSimultaneousTileGenerationJobsCount(int32 NewJobsCountLimit);
+
+	/** Returns query extent including adjustments for voxelization error compensation */
+	FVector GetModifiedQueryExtent(const FVector& QueryExtent) const
+	{
+		return FVector(QueryExtent.X, QueryExtent.Y, QueryExtent.Z + FMath::Max(0.0f, VerticalDeviationFromGroundCompensation));
+	}
+
+	//----------------------------------------------------------------------//
+	// Custom navigation links
+	//----------------------------------------------------------------------//
+
+	virtual void UpdateCustomLink(const INavLinkCustomInterface* CustomLink) override;
+
+	/** update area class and poly flags for all offmesh links with given UserId */
+	void UpdateNavigationLinkArea(int32 UserId, TSubclassOf<UNavArea> AreaClass) const;
+
+	/** update area class and poly flags for all offmesh segment links with given UserId */
+	void UpdateSegmentLinkArea(int32 UserId, TSubclassOf<UNavArea> AreaClass) const;
+
+	//----------------------------------------------------------------------//
+	// Batch processing (important with async rebuilding)
+	//----------------------------------------------------------------------//
+
+	/** Starts batch processing and locks access to navmesh from other threads */
+	virtual void BeginBatchQuery() const override;
+
+	/** Finishes batch processing and release locks */
+	virtual void FinishBatchQuery() const override;
+
+	//----------------------------------------------------------------------//
+	// Querying                                                                
+	//----------------------------------------------------------------------//
+	
+	/** dtNavMesh getter */
+	dtNavMesh* GetRecastMesh();
+
+	/** dtNavMesh getter */
+	const dtNavMesh* GetRecastMesh() const;
+
+	/** Retrieves LinkUserID associated with indicated PolyID */
+	uint32 GetLinkUserId(NavNodeRef LinkPolyID) const;
+
+	FColor GetAreaIDColor(uint8 AreaID) const;
+
+	/** Returns nearest navmesh polygon to Loc, or INVALID_NAVMESHREF if Loc is not on the navmesh. */
+	NavNodeRef FindNearestPoly(FVector const& Loc, FVector const& Extent, FSharedConstNavQueryFilter Filter = NULL, const UObject* Querier = NULL) const;
+
+	/** Finds the distance to the closest wall, limited to MaxDistance
+	 *	[out] OutClosestPointOnWall, if supplied, will be set to closest point on closest wall. Will not be set if no wall in the area (return value 0.f) */
+	float FindDistanceToWall(const FVector& StartLoc, FSharedConstNavQueryFilter Filter = nullptr, float MaxDistance = FLT_MAX, FVector* OutClosestPointOnWall = nullptr) const;
+
+	/** Retrieves center of the specified polygon. Returns false on error. */
+	bool GetPolyCenter(NavNodeRef PolyID, FVector& OutCenter) const;
+
+	/** Retrieves the vertices for the specified polygon. Returns false on error. */
+	bool GetPolyVerts(NavNodeRef PolyID, TArray<FVector>& OutVerts) const;
+
+	/** Retrieves area ID for the specified polygon. */
+	uint32 GetPolyAreaID(NavNodeRef PolyID) const;
+
+	/** Sets area ID for the specified polygon. */
+	bool SetPolyArea(NavNodeRef PolyID, TSubclassOf<UNavArea> AreaClass);
+
+	/** Sets area ID for the specified polygons */
+	void SetPolyArrayArea(const TArray<FNavPoly>& Polys, TSubclassOf<UNavArea> AreaClass);
+
+	/** Retrieves poly and area flags for specified polygon */
+	bool GetPolyFlags(NavNodeRef PolyID, uint16& PolyFlags, uint16& AreaFlags) const;
+	bool GetPolyFlags(NavNodeRef PolyID, FNavMeshNodeFlags& Flags) const;
+
+	/** Finds all polys connected with specified one */
+	bool GetPolyNeighbors(NavNodeRef PolyID, TArray<FNavigationPortalEdge>& Neighbors) const;
+
+	/** Finds all polys connected with specified one, results expressed as array of NavNodeRefs */
+	bool GetPolyNeighbors(NavNodeRef PolyID, TArray<NavNodeRef>& Neighbors) const;
+
+	/** Finds edges of specified poly */
+	bool GetPolyEdges(NavNodeRef PolyID, TArray<FNavigationPortalEdge>& Neighbors) const;
+
+	/** Finds closest point constrained to given poly */
+	bool GetClosestPointOnPoly(NavNodeRef PolyID, const FVector& TestPt, FVector& PointOnPoly) const;
+
+	/** Decode poly ID into tile index and poly index */
+	bool GetPolyTileIndex(NavNodeRef PolyID, uint32& PolyIndex, uint32& TileIndex) const;
+
+	/** Retrieves start and end point of offmesh link */
+	bool GetLinkEndPoints(NavNodeRef LinkPolyID, FVector& PointA, FVector& PointB) const;
+
+	/** Retrieves bounds of cluster. Returns false on error. */
+	bool GetClusterBounds(NavNodeRef ClusterRef, FBox& OutBounds) const;
+
+	/** Get random point in given cluster */
+	bool GetRandomPointInCluster(NavNodeRef ClusterRef, FNavLocation& OutLocation) const;
+
+	/** Get cluster ref containing given poly ref */
+	NavNodeRef GetClusterRef(NavNodeRef PolyRef) const;
+
+	/** Retrieves all polys within given pathing distance from StartLocation.
+	 *	@NOTE query is not using string-pulled path distance (for performance reasons),
+	 *		it measured distance between middles of portal edges, do you might want to 
+	 *		add an extra margin to PathingDistance */
+	bool GetPolysWithinPathingDistance(FVector const& StartLoc, const float PathingDistance, TArray<NavNodeRef>& FoundPolys,
+		FSharedConstNavQueryFilter Filter = nullptr, const UObject* Querier = nullptr, FRecastDebugPathfindingData* DebugData = nullptr) const;
+
+	/** Filters nav polys in PolyRefs with Filter */
+	bool FilterPolys(TArray<NavNodeRef>& PolyRefs, const FRecastQueryFilter* Filter, const UObject* Querier = NULL) const;
+
+	/** Get all polys from tile */
+	bool GetPolysInTile(int32 TileIndex, TArray<FNavPoly>& Polys) const;
+
+	/** Get all polys that overlap the specified box */
+	bool GetPolysInBox(const FBox& Box, TArray<FNavPoly>& Polys, FSharedConstNavQueryFilter Filter = nullptr, const UObject* Owner = nullptr) const;
+
+	/** Get all polys from tile */
+	bool GetNavLinksInTile(const int32 TileIndex, TArray<FNavPoly>& Polys, const bool bIncludeLinksFromNeighborTiles) const;
+
+	/** Projects point on navmesh, returning all hits along vertical line defined by min-max Z params */
+	bool ProjectPointMulti(const FVector& Point, TArray<FNavLocation>& OutLocations, const FVector& Extent,
+		float MinZ, float MaxZ, FSharedConstNavQueryFilter Filter = NULL, const UObject* Querier = NULL) const;
+	
+	// @todo docuement
+	static FPathFindingResult FindPath(const FNavAgentProperties& AgentProperties, const FPathFindingQuery& Query);
+	static bool TestPath(const FNavAgentProperties& AgentProperties, const FPathFindingQuery& Query, int32* NumVisitedNodes);
+	static bool TestHierarchicalPath(const FNavAgentProperties& AgentProperties, const FPathFindingQuery& Query, int32* NumVisitedNodes);
+	static bool NavMeshRaycast(const ANavigationData* Self, const FVector& RayStart, const FVector& RayEnd, FVector& HitLocation, FSharedConstNavQueryFilter QueryFilter, const UObject* Querier, FRaycastResult& Result);
+	static bool NavMeshRaycast(const ANavigationData* Self, const FVector& RayStart, const FVector& RayEnd, FVector& HitLocation, FSharedConstNavQueryFilter QueryFilter, const UObject* Querier = NULL);
+	static bool NavMeshRaycast(const ANavigationData* Self, NavNodeRef RayStartNode, const FVector& RayStart, const FVector& RayEnd, FVector& HitLocation, FSharedConstNavQueryFilter QueryFilter, const UObject* Querier = NULL);
+
+	virtual void BatchRaycast(TArray<FNavigationRaycastWork>& Workload, FSharedConstNavQueryFilter QueryFilter, const UObject* Querier = NULL) const override;
+
+	/** finds a Filter-passing navmesh location closest to specified StartLoc
+	 *	@return true if adjusting was required, false otherwise */
+	bool AdjustLocationWithFilter(const FVector& StartLoc, FVector& OutAdjustedLocation, const FNavigationQueryFilter& Filter, const UObject* Querier = NULL) const;
+	
+	/** Check if navmesh is defined (either built/streamed or recognized as empty tile by generator) in given radius.
+	  * @returns true if ALL tiles inside are ready
+	  */
+	bool HasCompleteDataInRadius(const FVector& TestLocation, float TestRadius) const;
+
+	/** @return true is specified segment is fully on navmesh (respecting the optional filter) */
+	bool IsSegmentOnNavmesh(const FVector& SegmentStart, const FVector& SegmentEnd, FSharedConstNavQueryFilter Filter = NULL, const UObject* Querier = NULL) const;
+
+	/** Check if poly is a custom link */
+	bool IsCustomLink(NavNodeRef PolyRef) const;
+
+	/** finds stringpulled path from given corridor */
+	bool FindStraightPath(const FVector& StartLoc, const FVector& EndLoc, const TArray<NavNodeRef>& PathCorridor, TArray<FNavPathPoint>& PathPoints, TArray<uint32>* CustomLinks = NULL) const;
+
+	/** Runs A* pathfinding on navmesh and collect data for every step */
+	int32 DebugPathfinding(const FPathFindingQuery& Query, TArray<FRecastDebugPathfindingData>& Steps);
+
+	static const FRecastQueryFilter* GetNamedFilter(ERecastNamedFilter::Type FilterType);
+	FORCEINLINE static FNavPolyFlags GetNavLinkFlag() { return NavLinkFlag; }
+	
+	virtual bool NeedsRebuild() const override;
+	virtual bool SupportsRuntimeGeneration() const override;
+	virtual bool SupportsStreaming() const override;
+	virtual void ConditionalConstructGenerator() override;
+	bool ShouldGatherDataOnGameThread() const { return bDoFullyAsyncNavDataGathering == false; }
+	int32 GetTileNumberHardLimit() const { return TileNumberHardLimit; }
+
+	void UpdateActiveTiles(const TArray<FNavigationInvokerRaw>& InvokerLocations);
+	void RemoveTiles(const TArray<FIntPoint>& Tiles);
+	void RebuildTile(const TArray<FIntPoint>& Tiles);
+
+protected:
+
+	void UpdatePolyRefBitsPreview();
+	
+	/** Invalidates active paths that go through changed tiles  */
+	void InvalidateAffectedPaths(const TArray<uint32>& ChangedTiles);
+
+	/** Spawns an ARecastNavMesh instance, and configures it if AgentProps != NULL */
+	static ARecastNavMesh* SpawnInstance(UNavigationSystem* NavSys, const FNavDataConfig* AgentProps = NULL);
+
+private:
+	friend class FRecastNavMeshGenerator;
+	friend class FPImplRecastNavMesh;
+	// destroys FPImplRecastNavMesh instance if it has been created 
+	void DestroyRecastPImpl();
+	// @todo docuement
+	void UpdateNavVersion();
+	void UpdateNavObject();
+
+	/** @return Navmesh data chunk that belongs to this actor */
+	URecastNavMeshDataChunk* GetNavigationDataChunk(ULevel* InLevel) const;
+
+protected:
+	// retrieves RecastNavMeshImpl
+	FPImplRecastNavMesh* GetRecastNavMeshImpl() { return RecastNavMeshImpl; }
+	const FPImplRecastNavMesh* GetRecastNavMeshImpl() const { return RecastNavMeshImpl; }
+
+private:
+	/** NavMesh versioning. */
+	uint32 NavMeshVersion;
+	
+	/** 
+	 * This is a pimpl-style arrangement used to tightly hide the Recast internals from the rest of the engine.
+	 * Using this class should *not* require the inclusion of the private RecastNavMesh.h
+	 *	@NOTE: if we switch over to C++11 this should be unique_ptr
+	 *	@TODO since it's no secret we're using recast there's no point in having separate implementation class. FPImplRecastNavMesh should be merged into ARecastNavMesh
+	 */
+	FPImplRecastNavMesh* RecastNavMeshImpl;
+	
+#if RECAST_ASYNC_REBUILDING
+	/** batch query counter */
+	mutable int32 BatchQueryCounter;
+
+#endif // RECAST_ASYNC_REBUILDING
+
+private:
+	static const FRecastQueryFilter* NamedFilters[ERecastNamedFilter::NamedFiltersCount];
+#endif // WITH_RECAST
+};
+
+#if WITH_RECAST
+FORCEINLINE
+bool ARecastNavMesh::NavMeshRaycast(const ANavigationData* Self, const FVector& RayStart, const FVector& RayEnd, FVector& HitLocation, FSharedConstNavQueryFilter QueryFilter, const UObject* Querier)
+{
+	FRaycastResult Result;
+	return NavMeshRaycast(Self, RayStart, RayEnd, HitLocation, QueryFilter, Querier, Result);
+}
+
+
+/** structure to cache owning RecastNavMesh data so that it doesn't have to be polled
+ *	directly from RecastNavMesh while asyncronously generating navmesh */
+struct FRecastNavMeshCachedData
+{
+	ARecastNavMesh::FNavPolyFlags FlagsPerArea[RECAST_MAX_AREAS];
+	ARecastNavMesh::FNavPolyFlags FlagsPerOffMeshLinkArea[RECAST_MAX_AREAS];
+	TMap<const UClass*, int32> AreaClassToIdMap;
+	const ARecastNavMesh* ActorOwner;
+	uint32 bUseSortFunction : 1;
+
+	static FRecastNavMeshCachedData Construct(const ARecastNavMesh* RecastNavMeshActor);
+	void OnAreaAdded(const UClass* AreaClass, int32 AreaID);
+};
+
+#endif // WITH_RECAST
diff --git a/Engine/Source/Runtime/NavigationSystem/Public/NavMesh/RecastNavMeshDataChunk.h b/Engine/Source/Runtime/NavigationSystem/Public/NavMesh/RecastNavMeshDataChunk.h
new file mode 100644
--- /dev/null
+++ b/Engine/Source/Runtime/NavigationSystem/Public/NavMesh/RecastNavMeshDataChunk.h
@@ -0,0 +1,82 @@
+// Copyright 1998-2018 Epic Games, Inc. All Rights Reserved.
+
+#pragma once
+
+#include "CoreMinimal.h"
+#include "UObject/ObjectMacros.h"
+#include "AI/Navigation/NavigationTypes.h"
+#include "AI/Navigation/NavigationDataChunk.h"
+#include "RecastNavMeshDataChunk.generated.h"
+
+class FPImplRecastNavMesh;
+
+struct FRecastTileData
+{
+	struct FRawData
+	{
+		FRawData(uint8* InData);
+		~FRawData();
+
+		uint8* RawData;
+	};
+
+	FRecastTileData();
+	FRecastTileData(int32 TileDataSize, uint8* TileRawData, int32 TileCacheDataSize, uint8* TileCacheRawData);
+	
+	// Location of attached tile
+	int32					X;					
+	int32					Y;					
+	int32					Layer;
+		
+	// Tile data
+	int32					TileDataSize;
+	TSharedPtr<FRawData>	TileRawData;
+
+	// Compressed tile cache layer 
+	int32					TileCacheDataSize;
+	TSharedPtr<FRawData>	TileCacheRawData;
+
+	// Whether this tile is attached to NavMesh
+	bool					bAttached;	
+};
+
+class dtNavMesh;
+class FPImplRecastNavMesh;
+
+/** 
+ * 
+ */
+UCLASS()
+class NAVIGATIONSYSTEM_API URecastNavMeshDataChunk : public UNavigationDataChunk
+{
+	GENERATED_UCLASS_BODY()
+
+	//~ Begin UObject Interface
+	virtual void Serialize(FArchive& Ar) override;
+	//~ End UObject Interface
+
+	/** Attaches tiles to specified navmesh, transferring tile ownership to navmesh */
+	TArray<uint32> AttachTiles(FPImplRecastNavMesh& NavMeshImpl);
+	
+	/** Detaches tiles from specified navmesh, taking tile ownership */
+	TArray<uint32> DetachTiles(FPImplRecastNavMesh& NavMeshImpl);
+	
+	/** Number of tiles in this chunk */
+	int32 GetNumTiles() const;
+
+	/** Releases all tiles that this chunk holds */
+	void ReleaseTiles();
+
+public:
+#if WITH_EDITOR
+	void GatherTiles(const FPImplRecastNavMesh* NavMeshImpl, const TArray<int32>& TileIndices);
+#endif// WITH_EDITOR
+
+private:
+#if WITH_RECAST
+	void SerializeRecastData(FArchive& Ar, int32 NavMeshVersion);
+#endif//WITH_RECAST
+
+private:
+	TArray<FRecastTileData> Tiles;
+};
