Copying //UE4/Dev-Framework to //UE4/Dev-Main (Source: //UE4/Dev-Framework @ 3967517)

diff --git a/Engine/Source/Runtime/NavigationSystem/Public/NavMesh/RecastNavMeshGenerator.h b/Engine/Source/Runtime/NavigationSystem/Public/NavMesh/RecastNavMeshGenerator.h
new file mode 100644
--- /dev/null
+++ b/Engine/Source/Runtime/NavigationSystem/Public/NavMesh/RecastNavMeshGenerator.h
@@ -0,0 +1,555 @@
+// Copyright 1998-2018 Epic Games, Inc. All Rights Reserved.
+#pragma once 
+
+#include "CoreMinimal.h"
+#include "Stats/Stats.h"
+#include "AI/Navigation/NavigationTypes.h"
+#include "EngineDefines.h"
+#include "AI/NavigationModifier.h"
+#include "NavigationOctree.h"
+#include "NavMesh/RecastNavMesh.h"
+#include "Async/AsyncWork.h"
+#include "UObject/GCObject.h"
+#include "AI/NavDataGenerator.h"
+
+#if WITH_RECAST
+
+#include "Recast/Recast.h"
+#include "Detour/DetourNavMesh.h"
+
+class UBodySetup;
+class ARecastNavMesh;
+class FNavigationOctree;
+class FNavMeshBuildContext;
+class FRecastNavMeshGenerator;
+struct BuildContext;
+struct FNavigationRelevantData;
+struct dtTileCacheLayer;
+struct FKAggregateGeom;
+
+#define MAX_VERTS_PER_POLY	6
+
+struct FRecastBuildConfig : public rcConfig
+{
+	/** controls whether voxel filterring will be applied (via FRecastTileGenerator::ApplyVoxelFilter) */
+	uint32 bPerformVoxelFiltering:1;
+	/** generate detailed mesh (additional tessellation to match heights of geometry) */
+	uint32 bGenerateDetailedMesh:1;
+	/** generate BV tree (space partitioning for queries) */
+	uint32 bGenerateBVTree:1;
+	/** if set, mark areas with insufficient free height instead of cutting them out  */
+	uint32 bMarkLowHeightAreas : 1;
+	/** if set, only single low height span will be allowed under valid one */
+	uint32 bFilterLowSpanSequences : 1;
+	/** if set, only low height spans with corresponding area modifier will be stored in tile cache (reduces memory, can't modify without full tile rebuild) */
+	uint32 bFilterLowSpanFromTileCache : 1;
+
+	/** region partitioning method used by tile cache */
+	int32 TileCachePartitionType;
+	/** chunk size for ChunkyMonotone partitioning */
+	int32 TileCacheChunkSize;
+
+	int32 PolyMaxHeight;
+	/** indicates what's the limit of navmesh polygons per tile. This value is calculated from other
+	 *	factors - DO NOT SET IT TO ARBITRARY VALUE */
+	int32 MaxPolysPerTile;
+
+	/** Actual agent height (in uu)*/
+	float AgentHeight;
+	/** Actual agent climb (in uu)*/
+	float AgentMaxClimb;
+	/** Actual agent radius (in uu)*/
+	float AgentRadius;
+	/** Agent index for filtering links */
+	int32 AgentIndex;
+
+	FRecastBuildConfig()
+	{
+		Reset();
+	}
+
+	void Reset()
+	{
+		FMemory::Memzero(*this);
+		bPerformVoxelFiltering = true;
+		bGenerateDetailedMesh = true;
+		bGenerateBVTree = true;
+		bMarkLowHeightAreas = false;
+		bFilterLowSpanSequences = false;
+		bFilterLowSpanFromTileCache = false;
+		PolyMaxHeight = 10;
+		MaxPolysPerTile = -1;
+		AgentIndex = 0;
+	}
+};
+
+struct FRecastVoxelCache
+{
+	struct FTileInfo
+	{
+		int16 TileX;
+		int16 TileY;
+		int32 NumSpans;
+		FTileInfo* NextTile;
+		rcSpanCache* SpanData;
+	};
+
+	int32 NumTiles;
+
+	/** tile info */
+	FTileInfo* Tiles;
+
+	FRecastVoxelCache() {}
+	FRecastVoxelCache(const uint8* Memory);
+};
+
+struct FRecastGeometryCache
+{
+	struct FHeader
+	{
+		FNavigationRelevantData::FCollisionDataHeader Validation;
+		
+		int32 NumVerts;
+		int32 NumFaces;
+		struct FWalkableSlopeOverride SlopeOverride;
+
+		static uint32 StaticMagicNumber;
+	};
+
+	FHeader Header;
+
+	/** recast coords of vertices (size: NumVerts * 3) */
+	float* Verts;
+
+	/** vert indices for triangles (size: NumFaces * 3) */
+	int32* Indices;
+
+	FRecastGeometryCache() {}
+	FRecastGeometryCache(const uint8* Memory);
+
+	static bool IsValid(const uint8* Memory, int32 MemorySize);
+};
+
+struct FRecastRawGeometryElement
+{
+	// Instance geometry
+	TArray<float>		GeomCoords;
+	TArray<int32>		GeomIndices;
+	
+	// Per instance transformations in unreal coords
+	// When empty geometry is in world space
+	TArray<FTransform>	PerInstanceTransform;
+};
+
+struct FRecastAreaNavModifierElement
+{
+	TArray<FAreaNavModifier> Areas;
+	
+	// Per instance transformations in unreal coords
+	// When empty areas are in world space
+	TArray<FTransform>	PerInstanceTransform;
+};
+
+/**
+ * Class handling generation of a single tile, caching data that can speed up subsequent tile generations
+ */
+class NAVIGATIONSYSTEM_API FRecastTileGenerator : public FNoncopyable, public FGCObject
+{
+	friend FRecastNavMeshGenerator;
+
+public:
+	FRecastTileGenerator(FRecastNavMeshGenerator& ParentGenerator, const FIntPoint& Location);
+	virtual ~FRecastTileGenerator();
+		
+	void DoWork();
+
+	FORCEINLINE int32 GetTileX() const { return TileX; }
+	FORCEINLINE int32 GetTileY() const { return TileY; }
+	/** Whether specified layer was updated */
+	FORCEINLINE bool IsLayerChanged(int32 LayerIdx) const { return DirtyLayers[LayerIdx]; }
+	FORCEINLINE const TBitArray<>& GetDirtyLayersMask() const { return DirtyLayers; }
+	/** Whether tile data was fully regenerated */
+	FORCEINLINE bool IsFullyRegenerated() const { return bRegenerateCompressedLayers; }
+	/** Whether tile task has anything to build */
+	bool HasDataToBuild() const;
+
+	const TArray<FNavMeshTileData>& GetCompressedLayers() const { return CompressedLayers; }
+protected:
+	// to be used solely by FRecastNavMeshGenerator
+	TArray<FNavMeshTileData>& GetNavigationData() { return NavigationData; }
+	
+public:
+	uint32 GetUsedMemCount() const;
+
+	// Memory amount used to construct generator 
+	uint32 UsedMemoryOnStartup;
+
+	// FGCObject begin
+	virtual void AddReferencedObjects(FReferenceCollector& Collector);
+	// FGCObject end
+		
+protected:
+	/** Does the actual tile generations. 
+	 *	@note always trigger tile generation only via TriggerAsyncBuild. This is a worker function
+	 *	@return true if new tile navigation data has been generated and is ready to be added to navmesh instance, 
+	 *	false if failed or no need to generate (still valid)
+	 */
+	bool GenerateTile();
+
+	void Setup(const FRecastNavMeshGenerator& ParentGenerator, const TArray<FBox>& DirtyAreas);
+	
+	void GatherGeometry(const FRecastNavMeshGenerator& ParentGenerator, bool bGeometryChanged);
+	void PrepareGeometrySources(const FRecastNavMeshGenerator& ParentGenerator, bool bGeometryChanged);
+	void DoAsyncGeometryGathering();
+
+	/** builds CompressedLayers array (geometry + modifiers) */
+	virtual bool GenerateCompressedLayers(FNavMeshBuildContext& BuildContext);
+
+	/** builds NavigationData array (layers + obstacles) */
+	bool GenerateNavigationData(FNavMeshBuildContext& BuildContext);
+
+	void ApplyVoxelFilter(struct rcHeightfield* SolidHF, float WalkableRadius);
+
+	/** apply areas from DynamicAreas to layer */
+	void MarkDynamicAreas(dtTileCacheLayer& Layer);
+	void MarkDynamicArea(const FAreaNavModifier& Modifier, const FTransform& LocalToWorld, dtTileCacheLayer& Layer);
+	void MarkDynamicArea(const FAreaNavModifier& Modifier, const FTransform& LocalToWorld, dtTileCacheLayer& Layer, const int32 AreaId, const int32* ReplaceAreaIdPtr);
+
+	void AppendModifier(const FCompositeNavModifier& Modifier, const FNavDataPerInstanceTransformDelegate& InTransformsDelegate);
+	/** Appends specified geometry to tile's geometry */
+	void ValidateAndAppendGeometry(TSharedRef<FNavigationRelevantData, ESPMode::ThreadSafe> ElementData);
+	void AppendGeometry(const TNavStatArray<uint8>& RawCollisionCache, const FNavDataPerInstanceTransformDelegate& InTransformsDelegate);
+	void AppendVoxels(rcSpanCache* SpanData, int32 NumSpans);
+	
+	/** prepare voxel cache from collision data */
+	void PrepareVoxelCache(const TNavStatArray<uint8>& RawCollisionCache, TNavStatArray<rcSpanCache>& SpanData);
+	bool HasVoxelCache(const TNavStatArray<uint8>& RawVoxelCache, rcSpanCache*& CachedVoxels, int32& NumCachedVoxels) const;
+	void AddVoxelCache(TNavStatArray<uint8>& RawVoxelCache, const rcSpanCache* CachedVoxels, const int32 NumCachedVoxels) const;
+
+	void DumpAsyncData();
+
+protected:
+	uint32 bSucceeded : 1;
+	uint32 bRegenerateCompressedLayers : 1;
+	uint32 bFullyEncapsulatedByInclusionBounds : 1;
+	uint32 bUpdateGeometry : 1;
+	uint32 bHasLowAreaModifiers : 1;
+	
+	int32 TileX;
+	int32 TileY;
+	uint32 Version;
+	/** Tile's bounding box, Unreal coords */
+	FBox TileBB;
+
+	/** Tile's bounding box expanded by Agent's radius *2 + CellSize, Unreal coords */
+	FBox TileBBExpandedForAgent;
+	
+	/** Layers dirty flags */
+	TBitArray<> DirtyLayers;
+	
+	/** Parameters defining navmesh tiles */
+	FRecastBuildConfig TileConfig;
+
+	/** Bounding geometry definition. */
+	TNavStatArray<FBox> InclusionBounds;
+
+	/** Additional config */
+	FRecastNavMeshCachedData AdditionalCachedData;
+
+	// generated tile data
+	TArray<FNavMeshTileData> CompressedLayers;
+	TArray<FNavMeshTileData> NavigationData;
+	
+	// tile's geometry: without voxel cache
+	TArray<FRecastRawGeometryElement> RawGeometry;
+	// areas used for creating navigation data: obstacles
+	TArray<FRecastAreaNavModifierElement> Modifiers;
+	// navigation links
+	TArray<FSimpleLinkNavModifier> OffmeshLinks;
+
+	TWeakPtr<FNavDataGenerator, ESPMode::ThreadSafe> ParentGeneratorWeakPtr;
+
+	TNavStatArray<TSharedRef<FNavigationRelevantData, ESPMode::ThreadSafe> > NavigationRelevantData;
+	TSharedPtr<FNavigationOctree, ESPMode::ThreadSafe> NavOctree; 
+	FNavDataConfig NavDataConfig;
+};
+
+struct NAVIGATIONSYSTEM_API FRecastTileGeneratorWrapper : public FNonAbandonableTask
+{
+	TSharedRef<FRecastTileGenerator> TileGenerator;
+
+	FRecastTileGeneratorWrapper(TSharedRef<FRecastTileGenerator> InTileGenerator)
+		: TileGenerator(InTileGenerator)
+	{
+	}
+	
+	void DoWork()
+	{
+		TileGenerator->DoWork();
+	}
+
+	FORCEINLINE TStatId GetStatId() const
+	{
+		RETURN_QUICK_DECLARE_CYCLE_STAT(FRecastTileGenerator, STATGROUP_ThreadPoolAsyncTasks);
+	}
+};
+
+typedef FAsyncTask<FRecastTileGeneratorWrapper> FRecastTileGeneratorTask;
+//typedef FAsyncTask<FRecastTileGenerator> FRecastTileGeneratorTask;
+
+struct FPendingTileElement
+{
+	/** tile coordinates on a grid in recast space */
+	FIntPoint	Coord;
+	/** distance to seed, used for sorting pending tiles */
+	float		SeedDistance; 
+	/** Whether we need a full rebuild for this tile grid cell */
+	bool		bRebuildGeometry;
+	/** We need to store dirty area bounds to check which cached layers needs to be regenerated
+	 *  In case geometry is changed cached layers data will be fully regenerated without using dirty areas list
+	 */
+	TArray<FBox> DirtyAreas;
+
+	FPendingTileElement()
+		: Coord(FIntPoint::NoneValue)
+		, SeedDistance(MAX_flt)
+		, bRebuildGeometry(false)
+	{
+	}
+
+	bool operator == (const FIntPoint& Location) const
+	{
+		return Coord == Location;
+	}
+
+	bool operator == (const FPendingTileElement& Other) const
+	{
+		return Coord == Other.Coord;
+	}
+
+	bool operator < (const FPendingTileElement& Other) const
+	{
+		return Other.SeedDistance < SeedDistance;
+	}
+
+	friend uint32 GetTypeHash(const FPendingTileElement& Element)
+	{
+		return GetTypeHash(Element.Coord);
+	}
+};
+
+struct FRunningTileElement
+{
+	FRunningTileElement()
+		: Coord(FIntPoint::NoneValue)
+		, bShouldDiscard(false)
+		, AsyncTask(nullptr)
+	{
+	}
+	
+	FRunningTileElement(FIntPoint InCoord)
+		: Coord(InCoord)
+		, bShouldDiscard(false)
+		, AsyncTask(nullptr)
+	{
+	}
+
+	bool operator == (const FRunningTileElement& Other) const
+	{
+		return Coord == Other.Coord;
+	}
+	
+	/** tile coordinates on a grid in recast space */
+	FIntPoint					Coord;
+	/** whether generated results should be discarded */
+	bool						bShouldDiscard; 
+	FRecastTileGeneratorTask*	AsyncTask;
+};
+
+struct FTileTimestamp
+{
+	uint32 TileIdx;
+	double Timestamp;
+	
+	bool operator == (const FTileTimestamp& Other) const
+	{
+		return TileIdx == Other.TileIdx;
+	}
+};
+
+/**
+ * Class that handles generation of the whole Recast-based navmesh.
+ */
+class NAVIGATIONSYSTEM_API FRecastNavMeshGenerator : public FNavDataGenerator
+{
+public:
+	FRecastNavMeshGenerator(ARecastNavMesh& InDestNavMesh);
+	virtual ~FRecastNavMeshGenerator();
+
+private:
+	/** Prevent copying. */
+	FRecastNavMeshGenerator(FRecastNavMeshGenerator const& NoCopy) { check(0); };
+	FRecastNavMeshGenerator& operator=(FRecastNavMeshGenerator const& NoCopy) { check(0); return *this; }
+
+public:
+	virtual bool RebuildAll() override;
+	virtual void EnsureBuildCompletion() override;
+	virtual void CancelBuild() override;
+	virtual void TickAsyncBuild(float DeltaSeconds) override;
+	virtual void OnNavigationBoundsChanged() override;
+
+	/** Asks generator to update navigation affected by DirtyAreas */
+	virtual void RebuildDirtyAreas(const TArray<FNavigationDirtyArea>& DirtyAreas) override;
+	virtual bool IsBuildInProgress(bool bCheckDirtyToo = false) const override;
+	virtual int32 GetNumRemaningBuildTasks() const override;
+	virtual int32 GetNumRunningBuildTasks() const override;
+
+	/** Checks if a given tile is being build or has just finished building */
+	bool IsTileChanged(int32 TileIdx) const;
+		
+	FORCEINLINE uint32 GetVersion() const { return Version; }
+
+	const ARecastNavMesh* GetOwner() const { return DestNavMesh; }
+
+	/** update area data */
+	void OnAreaAdded(const UClass* AreaClass, int32 AreaID);
+		
+	//--- accessors --- //
+	FORCEINLINE class UWorld* GetWorld() const { return DestNavMesh->GetWorld(); }
+
+	const FRecastBuildConfig& GetConfig() const { return Config; }
+
+	const TNavStatArray<FBox>& GetInclusionBounds() const { return InclusionBounds; }
+	
+	/** checks if any on InclusionBounds encapsulates given box.
+	 *	@return index to first item in InclusionBounds that meets expectations */
+	int32 FindInclusionBoundEncapsulatingBox(const FBox& Box) const;
+
+	/** Total navigable area box, sum of all navigation volumes bounding boxes */
+	FBox GetTotalBounds() const { return TotalNavBounds; }
+
+	const FRecastNavMeshCachedData& GetAdditionalCachedData() const { return AdditionalCachedData; }
+
+	bool HasDirtyTiles() const;
+
+	bool GatherGeometryOnGameThread() const;
+
+	FBox GrowBoundingBox(const FBox& BBox, bool bIncludeAgentHeight) const;
+	
+#if !(UE_BUILD_SHIPPING || UE_BUILD_TEST) && ENABLE_VISUAL_LOG
+	virtual void ExportNavigationData(const FString& FileName) const override;
+	virtual void GrabDebugSnapshot(struct FVisualLogEntry* Snapshot, const FBox& BoundingBox, const FName& CategoryName, ELogVerbosity::Type Verbosity) const override;
+#endif
+
+	/** 
+	 *	@param Actor is a reference to make callee responsible for assuring it's valid
+	 */
+	static void ExportComponentGeometry(UActorComponent* Component, FNavigationRelevantData& Data);
+	static void ExportVertexSoupGeometry(const TArray<FVector>& Verts, FNavigationRelevantData& Data);
+
+	static void ExportRigidBodyGeometry(UBodySetup& BodySetup, TNavStatArray<FVector>& OutVertexBuffer, TNavStatArray<int32>& OutIndexBuffer, const FTransform& LocalToWorld = FTransform::Identity);
+	static void ExportRigidBodyGeometry(UBodySetup& BodySetup, TNavStatArray<FVector>& OutTriMeshVertexBuffer, TNavStatArray<int32>& OutTriMeshIndexBuffer, TNavStatArray<FVector>& OutConvexVertexBuffer, TNavStatArray<int32>& OutConvexIndexBuffer, TNavStatArray<int32>& OutShapeBuffer, const FTransform& LocalToWorld = FTransform::Identity);
+	static void ExportAggregatedGeometry(const FKAggregateGeom& AggGeom, TNavStatArray<FVector>& OutConvexVertexBuffer, TNavStatArray<int32>& OutConvexIndexBuffer, TNavStatArray<int32>& OutShapeBuffer, const FTransform& LocalToWorld = FTransform::Identity);
+
+protected:
+	// Performs initial setup of member variables so that generator is ready to do its thing from this point on
+	void Init();
+
+	// Updates cached list of navigation bounds
+	void UpdateNavigationBounds();
+		
+	// Sorts pending build tiles by proximity to player, so tiles closer to player will get generated first
+	void SortPendingBuildTiles();
+
+	/** Instantiates dtNavMesh and configures it for tiles generation. Returns false if failed */
+	bool ConstructTiledNavMesh();
+
+	/** Determine bit masks for poly address */
+	void CalcNavMeshProperties(int32& MaxTiles, int32& MaxPolys);
+	
+	/** Marks grid tiles affected by specified areas as dirty */
+	void MarkDirtyTiles(const TArray<FNavigationDirtyArea>& DirtyAreas);
+	
+	/** Processes pending tile generation tasks */
+	TArray<uint32> ProcessTileTasks(const int32 NumTasksToSubmit);
+
+public:
+	/** Adds generated tiles to NavMesh, replacing old ones */
+	TArray<uint32> AddGeneratedTiles(FRecastTileGenerator& TileGenerator);
+
+public:
+	/** Removes all tiles at specified grid location */
+	TArray<uint32> RemoveTileLayers(const int32 TileX, const int32 TileY, TMap<int32, dtPolyRef>* OldLayerTileIdMap = nullptr);
+
+	void RemoveTiles(const TArray<FIntPoint>& Tiles);
+	void ReAddTiles(const TArray<FIntPoint>& Tiles);
+
+	bool IsBuildingRestrictedToActiveTiles() const { return bRestrictBuildingToActiveTiles; }
+
+	/** sets a limit to number of asynchronous tile generators running at one time
+	 *	@note if used at runtime will not result in killing tasks above limit count
+	 *	@mote function does not validate the input parameter - it's on caller */
+	void SetMaxTileGeneratorTasks(int32 NewLimit) { MaxTileGeneratorTasks = NewLimit; }
+
+	static void CalcPolyRefBits(ARecastNavMesh* NavMeshOwner, int32& MaxTileBits, int32& MaxPolyBits);
+
+protected:
+	bool IsInActiveSet(const FIntPoint& Tile) const;
+	void RestrictBuildingToActiveTiles(bool InRestrictBuildingToActiveTiles);
+	
+	/** Blocks until build for specified list of tiles is complete and discard results */
+	void DiscardCurrentBuildingTasks();
+
+	virtual TSharedRef<FRecastTileGenerator> CreateTileGenerator(const FIntPoint& Coord, const TArray<FBox>& DirtyAreas);
+
+	//----------------------------------------------------------------------//
+	// debug
+	//----------------------------------------------------------------------//
+	virtual uint32 LogMemUsed() const override;
+
+private:
+	friend ARecastNavMesh;
+
+	/** Parameters defining navmesh tiles */
+	FRecastBuildConfig Config;
+	
+	int32 NumActiveTiles;
+	/** the limit to number of asynchronous tile generators running at one time */
+	int32 MaxTileGeneratorTasks;
+	float AvgLayersPerTile;
+
+	/** Total bounding box that includes all volumes, in unreal units. */
+	FBox TotalNavBounds;
+
+	/** Bounding geometry definition. */
+	TNavStatArray<FBox> InclusionBounds;
+
+	/** Navigation mesh that owns this generator */
+	ARecastNavMesh*	DestNavMesh;
+	
+	/** List of dirty tiles that needs to be regenerated */
+	TNavStatArray<FPendingTileElement> PendingDirtyTiles;			
+	
+	/** List of dirty tiles currently being regenerated */
+	TNavStatArray<FRunningTileElement> RunningDirtyTiles;
+
+#if WITH_EDITOR
+	/** List of tiles that were recently regenerated */
+	TNavStatArray<FTileTimestamp> RecentlyBuiltTiles;
+#endif// WITH_EDITOR
+	
+	TArray<FIntPoint> ActiveTiles;
+
+	/** */
+	FRecastNavMeshCachedData AdditionalCachedData;
+
+	uint32 bInitialized:1;
+
+	uint32 bRestrictBuildingToActiveTiles:1;
+
+	/** Runtime generator's version, increased every time all tile generators get invalidated
+	 *	like when navmesh size changes */
+	uint32 Version;
+};
+
+#endif // WITH_RECAST
diff --git a/Engine/Source/Runtime/NavigationSystem/Public/NavMesh/RecastQueryFilter.h b/Engine/Source/Runtime/NavigationSystem/Public/NavMesh/RecastQueryFilter.h
new file mode 100644
--- /dev/null
+++ b/Engine/Source/Runtime/NavigationSystem/Public/NavMesh/RecastQueryFilter.h
@@ -0,0 +1,57 @@
+// Copyright 1998-2018 Epic Games, Inc. All Rights Reserved.
+
+#pragma once 
+
+#include "CoreMinimal.h"
+#include "NavFilters/NavigationQueryFilter.h"
+#include "AI/Navigation/NavigationTypes.h"
+
+#if WITH_RECAST
+#include "Detour/DetourNavMesh.h"
+#include "Detour/DetourNavMeshQuery.h"
+
+#define RECAST_VERY_SMALL_AGENT_RADIUS 0.0f
+
+class UNavigationSystemV1;
+
+class NAVIGATIONSYSTEM_API FRecastQueryFilter : public INavigationQueryFilterInterface, public dtQueryFilter
+{
+public:
+	FRecastQueryFilter(bool bIsVirtual = true);
+	virtual ~FRecastQueryFilter(){}
+
+	virtual void Reset() override;
+
+	virtual void SetAreaCost(uint8 AreaType, float Cost) override;
+	virtual void SetFixedAreaEnteringCost(uint8 AreaType, float Cost) override;
+	virtual void SetExcludedArea(uint8 AreaType) override;
+	virtual void SetAllAreaCosts(const float* CostArray, const int32 Count) override;
+	virtual void GetAllAreaCosts(float* CostArray, float* FixedCostArray, const int32 Count) const override;
+	virtual void SetBacktrackingEnabled(const bool bBacktracking) override;
+	virtual bool IsBacktrackingEnabled() const override;
+	virtual bool IsEqual(const INavigationQueryFilterInterface* Other) const override;
+	virtual void SetIncludeFlags(uint16 Flags) override;
+	virtual uint16 GetIncludeFlags() const override;
+	virtual void SetExcludeFlags(uint16 Flags) override;
+	virtual uint16 GetExcludeFlags() const override;
+	virtual FVector GetAdjustedEndLocation(const FVector& EndLocation) const override { return EndLocation; }
+	virtual INavigationQueryFilterInterface* CreateCopy() const override;
+
+	const dtQueryFilter* GetAsDetourQueryFilter() const { return this; }
+
+	/** note that it results in losing all area cost setup. Call it before setting anything else */
+	void SetIsVirtual(bool bIsVirtual);
+};
+
+struct NAVIGATIONSYSTEM_API FRecastSpeciaLinkFilter : public dtQuerySpecialLinkFilter
+{
+	FRecastSpeciaLinkFilter(UNavigationSystemV1* NavSystem, const UObject* Owner) : NavSys(NavSystem), SearchOwner(Owner), CachedOwnerOb(nullptr) {}
+	virtual bool isLinkAllowed(const int32 UserId) const override;
+	virtual void initialize() override;
+
+	UNavigationSystemV1* NavSys;
+	FWeakObjectPtr SearchOwner;
+	UObject* CachedOwnerOb;
+};
+
+#endif	// WITH_RECAST
diff --git a/Engine/Source/Runtime/NavigationSystem/Public/NavMesh/RecastVersion.h b/Engine/Source/Runtime/NavigationSystem/Public/NavMesh/RecastVersion.h
new file mode 100644
--- /dev/null
+++ b/Engine/Source/Runtime/NavigationSystem/Public/NavMesh/RecastVersion.h
@@ -0,0 +1,20 @@
+// Copyright 1998-2018 Epic Games, Inc. All Rights Reserved.
+
+#pragma once
+
+/** Versioning */
+#define NAVMESHVER_INITIAL				1
+#define NAVMESHVER_TILED_GENERATION		2
+#define NAVMESHVER_SEAMLESS_REBUILDING_1 3
+#define NAVMESHVER_AREA_CLASSES			4
+#define NAVMESHVER_CLUSTER_PATH			5
+#define NAVMESHVER_SEGMENT_LINKS		6
+#define NAVMESHVER_DYNAMIC_LINKS		7
+#define NAVMESHVER_64BIT				9
+#define NAVMESHVER_CLUSTER_SIMPLIFIED	10
+#define NAVMESHVER_OFFMESH_HEIGHT_BUG	11
+#define NAVMESHVER_LANDSCAPE_HEIGHT		13
+
+
+#define NAVMESHVER_LATEST				NAVMESHVER_LANDSCAPE_HEIGHT
+#define NAVMESHVER_MIN_COMPATIBLE		NAVMESHVER_LANDSCAPE_HEIGHT
diff --git a/Engine/Source/Runtime/NavigationSystem/Public/NavModifierComponent.h b/Engine/Source/Runtime/NavigationSystem/Public/NavModifierComponent.h
new file mode 100644
--- /dev/null
+++ b/Engine/Source/Runtime/NavigationSystem/Public/NavModifierComponent.h
@@ -0,0 +1,48 @@
+// Copyright 1998-2018 Epic Games, Inc. All Rights Reserved.
+
+#pragma once
+
+#include "CoreMinimal.h"
+#include "UObject/ObjectMacros.h"
+#include "Templates/SubclassOf.h"
+#include "NavAreas/NavArea.h"
+#include "NavRelevantComponent.h"
+#include "NavModifierComponent.generated.h"
+
+struct FNavigationRelevantData;
+
+UCLASS(ClassGroup = (Navigation), meta = (BlueprintSpawnableComponent), hidecategories = (Activation), config = Engine, defaultconfig)
+class NAVIGATIONSYSTEM_API UNavModifierComponent : public UNavRelevantComponent
+{
+	GENERATED_UCLASS_BODY()
+
+	UPROPERTY(EditAnywhere, BlueprintReadOnly, Category = Navigation)
+	TSubclassOf<UNavArea> AreaClass;
+
+	/** box extent used ONLY when owning actor doesn't have collision component */
+	UPROPERTY(EditAnywhere, Category = Navigation)
+	FVector FailsafeExtent;
+
+	/** Setting to 'true' will result in expanding lower bounding box of the nav 
+	 *	modifier by agent's height, before applying to navmesh */
+	UPROPERTY(config, EditAnywhere, Category = Navigation)
+	uint8 bIncludeAgentHeight : 1;
+
+	virtual void CalcAndCacheBounds() const override;
+	virtual void GetNavigationData(FNavigationRelevantData& Data) const override;
+
+	UFUNCTION(BlueprintCallable, Category = "AI|Navigation")
+	void SetAreaClass(TSubclassOf<UNavArea> NewAreaClass);
+
+protected:
+	struct FRotatedBox
+	{
+		FBox Box;
+		FQuat Quat;
+
+		FRotatedBox() {}
+		FRotatedBox(const FBox& InBox, const FQuat& InQuat) : Box(InBox), Quat(InQuat) {}
+	};
+
+	mutable TArray<FRotatedBox> ComponentBounds;
+};
diff --git a/Engine/Source/Runtime/NavigationSystem/Public/NavModifierVolume.h b/Engine/Source/Runtime/NavigationSystem/Public/NavModifierVolume.h
new file mode 100644
--- /dev/null
+++ b/Engine/Source/Runtime/NavigationSystem/Public/NavModifierVolume.h
@@ -0,0 +1,43 @@
+// Copyright 1998-2018 Epic Games, Inc. All Rights Reserved.
+
+#pragma once 
+
+#include "CoreMinimal.h"
+#include "UObject/ObjectMacros.h"
+#include "Templates/SubclassOf.h"
+#include "AI/Navigation/NavRelevantInterface.h"
+#include "NavAreas/NavArea.h"
+#include "GameFramework/Volume.h"
+#include "NavModifierVolume.generated.h"
+
+struct FNavigationRelevantData;
+
+/** 
+ *	Allows applying selected AreaClass to navmesh, using Volume's shape
+ */
+UCLASS(hidecategories=(Navigation))
+class NAVIGATIONSYSTEM_API ANavModifierVolume : public AVolume, public INavRelevantInterface
+{
+	GENERATED_BODY()
+
+protected:
+	UPROPERTY(EditAnywhere, BlueprintReadOnly, Category=Default)
+	TSubclassOf<UNavArea> AreaClass;
+
+public:
+	ANavModifierVolume(const FObjectInitializer& ObjectInitializer);
+
+	UFUNCTION(BlueprintCallable, Category = "AI|Navigation")
+	void SetAreaClass(TSubclassOf<UNavArea> NewAreaClass = nullptr);
+
+	TSubclassOf<UNavArea> GetAreaClass() const { return AreaClass; }
+
+	virtual void GetNavigationData(FNavigationRelevantData& Data) const override;
+	virtual FBox GetNavigationBounds() const override;
+	virtual void RebuildNavigationData() override;
+
+#if WITH_EDITOR
+	virtual void PostEditUndo() override;
+	virtual void PostEditChangeProperty(FPropertyChangedEvent& PropertyChangedEvent) override;
+#endif
+};
diff --git a/Engine/Source/Runtime/NavigationSystem/Public/NavRelevantComponent.h b/Engine/Source/Runtime/NavigationSystem/Public/NavRelevantComponent.h
new file mode 100644
--- /dev/null
+++ b/Engine/Source/Runtime/NavigationSystem/Public/NavRelevantComponent.h
@@ -0,0 +1,52 @@
+// Copyright 1998-2018 Epic Games, Inc. All Rights Reserved.
+
+#pragma once
+
+#include "CoreMinimal.h"
+#include "UObject/ObjectMacros.h"
+#include "Components/ActorComponent.h"
+#include "AI/Navigation/NavRelevantInterface.h"
+#include "NavRelevantComponent.generated.h"
+
+UCLASS()
+class NAVIGATIONSYSTEM_API UNavRelevantComponent : public UActorComponent, public INavRelevantInterface
+{
+	GENERATED_UCLASS_BODY()
+
+	//~ Begin UActorComponent Interface
+	virtual void OnRegister() override;
+	virtual void OnUnregister() override;
+	//~ End UActorComponent Interface
+
+	//~ Begin INavRelevantInterface Interface
+	virtual FBox GetNavigationBounds() const override;
+	virtual bool IsNavigationRelevant() const override;
+	virtual void UpdateNavigationBounds() override;
+	virtual UObject* GetNavigationParent() const override;
+	//~ End INavRelevantInterface Interface
+	
+	virtual void CalcAndCacheBounds() const;
+
+	UFUNCTION(BlueprintCallable, Category="AI|Navigation")
+	void SetNavigationRelevancy(bool bRelevant);
+
+	/** force relevancy and skip attaching navigation data to owner's root entry */
+	void ForceNavigationRelevancy(bool bForce);
+
+	/** force refresh in navigation octree */
+	void RefreshNavigationModifiers();
+	
+protected:
+
+	/** bounds for navigation octree */
+	mutable FBox Bounds;
+	
+	/** attach navigation data to entry for owner's root component (depends on its relevancy) */
+	UPROPERTY()
+	uint32 bAttachToOwnersRoot : 1;
+
+	mutable uint32 bBoundsInitialized : 1;
+
+	UPROPERTY(transient)
+	UObject* CachedNavParent;
+};
diff --git a/Engine/Source/Runtime/NavigationSystem/Public/NavigationInvokerComponent.h b/Engine/Source/Runtime/NavigationSystem/Public/NavigationInvokerComponent.h
new file mode 100644
--- /dev/null
+++ b/Engine/Source/Runtime/NavigationSystem/Public/NavigationInvokerComponent.h
@@ -0,0 +1,33 @@
+// Copyright 1998-2018 Epic Games, Inc. All Rights Reserved.
+#pragma once
+
+#include "CoreMinimal.h"
+#include "UObject/ObjectMacros.h"
+#include "UObject/UObjectGlobals.h"
+#include "Components/ActorComponent.h"
+#include "NavigationInvokerComponent.generated.h"
+
+class UNavigationSystemV1;
+
+UCLASS(ClassGroup = (Navigation), meta = (BlueprintSpawnableComponent), hidecategories = (Activation))
+class NAVIGATIONSYSTEM_API UNavigationInvokerComponent : public UActorComponent
+{
+	GENERATED_BODY()
+
+protected:
+
+	UPROPERTY(EditAnywhere, Category = Navigation, meta = (ClampMin = "0.1", UIMin = "0.1"))
+	float TileGenerationRadius;
+
+	UPROPERTY(EditAnywhere, Category = Navigation, meta = (ClampMin = "0.1", UIMin = "0.1"))
+	float TileRemovalRadius;
+
+public:
+	UNavigationInvokerComponent(const FObjectInitializer& ObjectInitializer = FObjectInitializer::Get());
+
+	void RegisterWithNavigationSystem(UNavigationSystemV1& NavSys);
+
+protected:
+	virtual void Activate(bool bReset = false) override;
+	virtual void Deactivate() override;
+};
diff --git a/Engine/Source/Runtime/NavigationSystem/Public/NavigationOctree.h b/Engine/Source/Runtime/NavigationSystem/Public/NavigationOctree.h
new file mode 100644
--- /dev/null
+++ b/Engine/Source/Runtime/NavigationSystem/Public/NavigationOctree.h
@@ -0,0 +1,158 @@
+// Copyright 1998-2018 Epic Games, Inc. All Rights Reserved.
+
+#pragma once
+
+#include "CoreMinimal.h"
+#include "Stats/Stats.h"
+#include "AI/Navigation/NavigationTypes.h"
+#include "GenericOctreePublic.h"
+#include "NavigationSystemTypes.h"
+#include "EngineStats.h"
+#include "AI/NavigationModifier.h"
+#include "AI/Navigation/NavRelevantInterface.h"
+#include "GenericOctree.h"
+
+class INavRelevantInterface;
+
+typedef FNavigationRelevantDataFilter FNavigationOctreeFilter;
+
+struct NAVIGATIONSYSTEM_API FNavigationOctreeElement
+{
+	FBoxSphereBounds Bounds;
+	TSharedRef<FNavigationRelevantData, ESPMode::ThreadSafe> Data;
+
+	FNavigationOctreeElement(UObject& SourceObject)
+		: Data(new FNavigationRelevantData(SourceObject))
+	{
+
+	}
+
+	FORCEINLINE bool IsEmpty() const
+	{
+		const FBox BBox = Bounds.GetBox();
+		return Data->IsEmpty() && (BBox.IsValid == 0 || BBox.GetSize().IsNearlyZero());
+	}
+
+	FORCEINLINE bool IsMatchingFilter(const FNavigationOctreeFilter& Filter) const
+	{
+		return Data->IsMatchingFilter(Filter);
+	}
+
+	/** 
+	 *	retrieves Modifier, if it doesn't contain any "Meta Navigation Areas". 
+	 *	If it does then retrieves a copy with meta areas substituted with
+	 *	appropriate non-meta areas, depending on NavAgent
+	 */
+	FORCEINLINE FCompositeNavModifier GetModifierForAgent(const struct FNavAgentProperties* NavAgent = NULL) const 
+	{ 
+		return Data->Modifiers.HasMetaAreas() ? Data->Modifiers.GetInstantiatedMetaModifier(NavAgent, Data->SourceObject) : Data->Modifiers;
+	}
+
+	FORCEINLINE bool ShouldUseGeometry(const FNavDataConfig& NavConfig) const
+	{ 
+		return !Data->ShouldUseGeometryDelegate.IsBound() || Data->ShouldUseGeometryDelegate.Execute(&NavConfig);
+	}
+
+	FORCEINLINE int32 GetAllocatedSize() const
+	{
+		return Data->GetAllocatedSize();
+	}
+
+	FORCEINLINE void Shrink()
+	{
+		Data->Shrink();
+	}
+
+	FORCEINLINE void ValidateAndShrink()
+	{
+		Data->ValidateAndShrink();
+	}
+
+	FORCEINLINE UObject* GetOwner() const { return Data->SourceObject.Get(); }
+};
+
+struct FNavigationOctreeSemantics
+{
+	enum { MaxElementsPerLeaf = 16 };
+	enum { MinInclusiveElementsPerNode = 7 };
+	enum { MaxNodeDepth = 12 };
+
+	typedef TInlineAllocator<MaxElementsPerLeaf> ElementAllocator;
+	//typedef FDefaultAllocator ElementAllocator;
+
+	FORCEINLINE static const FBoxSphereBounds& GetBoundingBox(const FNavigationOctreeElement& NavData)
+	{
+		return NavData.Bounds;
+	}
+
+	FORCEINLINE static bool AreElementsEqual(const FNavigationOctreeElement& A, const FNavigationOctreeElement& B)
+	{
+		return A.Data->SourceObject == B.Data->SourceObject;
+	}
+
+#if NAVSYS_DEBUG
+	FORCENOINLINE 
+#endif // NAVSYS_DEBUG
+	static void SetElementId(const FNavigationOctreeElement& Element, FOctreeElementId Id);
+};
+
+class FNavigationOctree : public TOctree<FNavigationOctreeElement, FNavigationOctreeSemantics>, public TSharedFromThis<FNavigationOctree, ESPMode::ThreadSafe>
+{
+public:
+	DECLARE_DELEGATE_TwoParams(FNavigableGeometryComponentExportDelegate, UActorComponent*, FNavigationRelevantData&);
+	FNavigableGeometryComponentExportDelegate ComponentExportDelegate;
+
+	enum ENavGeometryStoringMode {
+		SkipNavGeometry,
+		StoreNavGeometry,
+	};
+
+	FNavigationOctree(const FVector& Origin, float Radius);
+	virtual ~FNavigationOctree();
+
+	/** Add new node and fill it with navigation export data */
+	void AddNode(UObject* ElementOb, INavRelevantInterface* NavElement, const FBox& Bounds, FNavigationOctreeElement& Data);
+
+	/** Append new data to existing node */
+	void AppendToNode(const FOctreeElementId& Id, INavRelevantInterface* NavElement, const FBox& Bounds, FNavigationOctreeElement& Data);
+
+	/** Updates element bounds remove/add operation */
+	void UpdateNode(const FOctreeElementId& Id, const FBox& NewBounds);
+
+	/** Remove node */
+	void RemoveNode(const FOctreeElementId& Id);
+
+	void SetNavigableGeometryStoringMode(ENavGeometryStoringMode NavGeometryMode);
+
+	const FNavigationRelevantData* GetDataForID(const FOctreeElementId& Id) const;
+
+	ENavGeometryStoringMode GetNavGeometryStoringMode() const
+	{
+		return bGatherGeometry ? StoreNavGeometry : SkipNavGeometry;
+	}
+
+	void SetDataGatheringMode(ENavDataGatheringModeConfig Mode);
+	
+	// @hack! TO BE FIXED
+	void DemandLazyDataGathering(const FNavigationOctreeElement& Element);
+	void DemandLazyDataGathering(FNavigationRelevantData& ElementData);
+
+protected:
+	ENavDataGatheringMode DefaultGeometryGatheringMode;
+	uint32 bGatherGeometry : 1;
+	uint32 NodesMemory;
+};
+
+template<>
+FORCEINLINE_DEBUGGABLE void SetOctreeMemoryUsage(TOctree<FNavigationOctreeElement, FNavigationOctreeSemantics>* Octree, int32 NewSize)
+{
+	{
+		DEC_DWORD_STAT_BY( STAT_NavigationMemory, Octree->TotalSizeBytes );
+		DEC_DWORD_STAT_BY(STAT_Navigation_CollisionTreeMemory, Octree->TotalSizeBytes);
+	}
+	Octree->TotalSizeBytes = NewSize;
+	{
+		INC_DWORD_STAT_BY( STAT_NavigationMemory, NewSize );
+		INC_DWORD_STAT_BY(STAT_Navigation_CollisionTreeMemory, NewSize);
+	}
+}
diff --git a/Engine/Source/Runtime/NavigationSystem/Public/NavigationSystem.h b/Engine/Source/Runtime/NavigationSystem/Public/NavigationSystem.h
new file mode 100644
--- /dev/null
+++ b/Engine/Source/Runtime/NavigationSystem/Public/NavigationSystem.h
@@ -0,0 +1,971 @@
+// Copyright 1998-2018 Epic Games, Inc. All Rights Reserved.
+
+#pragma once
+
+#include "CoreMinimal.h"
+#include "UObject/ObjectMacros.h"
+#include "Templates/SubclassOf.h"
+#include "UObject/WeakObjectPtr.h"
+#include "Misc/CoreMisc.h"
+#include "Misc/CoreDelegates.h"
+#include "NavFilters/NavigationQueryFilter.h"
+#include "AI/Navigation/NavigationTypes.h"
+#include "NavigationSystemTypes.h"
+#include "GenericOctreePublic.h"
+#include "NavigationData.h"
+#include "AI/NavigationSystemBase.h"
+#include "Kismet/BlueprintFunctionLibrary.h"
+#include "NavigationOctree.h"
+#include "NavigationSystem.generated.h"
+
+
+#define NAV_USE_MAIN_NAVIGATION_DATA NULL
+
+class AController;
+class ANavMeshBoundsVolume;
+class AWorldSettings;
+class FEdMode;
+class FNavDataGenerator;
+class FNavigationOctree;
+class INavLinkCustomInterface;
+class INavRelevantInterface;
+class UCrowdManagerBase;
+class UNavArea;
+class UNavigationPath;
+struct FNavigationRelevantData;
+struct FNavigationOctreeElement;
+
+#if WITH_EDITOR
+class FEdMode;
+#endif // WITH_EDITOR
+
+/** delegate to let interested parties know that new nav area class has been registered */
+DECLARE_MULTICAST_DELEGATE_OneParam(FOnNavAreaChanged, const UClass* /*AreaClass*/);
+
+/** Delegate to let interested parties know that Nav Data has been registered */
+DECLARE_DYNAMIC_MULTICAST_DELEGATE_OneParam(FOnNavDataGenerigEvent, ANavigationData*, NavData);
+
+DECLARE_MULTICAST_DELEGATE(FOnNavigationInitDone);
+
+namespace NavigationDebugDrawing
+{
+	extern const NAVIGATIONSYSTEM_API float PathLineThickness;
+	extern const NAVIGATIONSYSTEM_API FVector PathOffset;
+	extern const NAVIGATIONSYSTEM_API FVector PathNodeBoxExtent;
+}
+
+namespace FNavigationSystem
+{	
+	/** 
+	 * Used to construct an ANavigationData instance for specified navigation data agent 
+	 */
+	typedef ANavigationData* (*FNavigationDataInstanceCreator)(UWorld*, const FNavDataConfig&);
+
+	struct NAVIGATIONSYSTEM_API FCustomLinkOwnerInfo
+	{
+		FWeakObjectPtr LinkOwner;
+		INavLinkCustomInterface* LinkInterface;
+
+		FCustomLinkOwnerInfo() : LinkInterface(nullptr) {}
+		FCustomLinkOwnerInfo(INavLinkCustomInterface* Link);
+
+		bool IsValid() const { return LinkOwner.IsValid(); }
+	};
+
+	bool NAVIGATIONSYSTEM_API ShouldLoadNavigationOnClient(ANavigationData& NavData);
+}
+
+struct FNavigationSystemExec: public FSelfRegisteringExec
+{
+	//~ Begin FExec Interface
+	virtual bool Exec(UWorld* Inworld, const TCHAR* Cmd, FOutputDevice& Ar) override;
+	//~ End FExec Interface
+};
+
+namespace ENavigationBuildLock
+{
+	enum Type
+	{
+		NoUpdateInEditor = 1 << 1,		// editor doesn't allow automatic updates
+		InitialLock = 1 << 2,			// initial lock, release manually after levels are ready for rebuild (e.g. streaming)
+		Custom = 1 << 3,
+	};
+}
+
+UCLASS(Within=World, config=Engine, defaultconfig)
+class NAVIGATIONSYSTEM_API UNavigationSystemV1 : public UNavigationSystemBase
+{
+	GENERATED_BODY()
+
+public:
+	UNavigationSystemV1(const FObjectInitializer& ObjectInitializer = FObjectInitializer::Get());
+	virtual ~UNavigationSystemV1();
+
+	UPROPERTY()
+	ANavigationData* MainNavData;
+
+	/** special navigation data for managing direct paths, not part of NavDataSet! */
+	UPROPERTY(Transient)
+	ANavigationData* AbstractNavData;
+
+protected:
+	UPROPERTY(EditAnywhere, BlueprintReadOnly, Category = Navigation)
+	TSubclassOf<UCrowdManagerBase> CrowdManagerClass;
+
+	/** Should navigation system spawn default Navigation Data when there's none and there are navigation bounds present? */
+	UPROPERTY(config, EditAnywhere, Category=NavigationSystem)
+	uint32 bAutoCreateNavigationData:1;
+
+	UPROPERTY(config, EditAnywhere, Category=NavigationSystem)
+	uint32 bAllowClientSideNavigation:1;
+
+	UPROPERTY(config, EditAnywhere, Category=NavigationSystem)
+	uint32 bTickWhilePaused:1;
+
+	/** gets set to true if gathering navigation data (like in navoctree) is required due to the need of navigation generation 
+	 *	Is always true in Editor Mode. In other modes it depends on bRebuildAtRuntime of every required NavigationData class' CDO
+	 */
+	UPROPERTY()
+	uint32 bSupportRebuilding : 1; 
+
+public:
+	/** if set to true will result navigation system not rebuild navigation until 
+	 *	a call to ReleaseInitialBuildingLock() is called. Does not influence 
+	 *	editor-time generation (i.e. does influence PIE and Game).
+	 *	Defaults to false.*/
+	UPROPERTY(config, EditAnywhere, Category=NavigationSystem)
+	uint32 bInitialBuildingLocked:1;
+
+	/** If set to true (default) navigation will be generated only within special navigation 
+	 *	bounds volumes (like ANavMeshBoundsVolume). Set to false means navigation should be generated
+	 *	everywhere.
+	 */
+	// @todo removing it from edition since it's currently broken and I'm not sure we want that at all
+	// since I'm not sure we can make it efficient in a generic case
+	//UPROPERTY(config, EditAnywhere, Category=NavigationSystem)
+	uint32 bWholeWorldNavigable:1;
+
+	/** false by default, if set to true will result in not caring about nav agent height 
+	 *	when trying to match navigation data to passed in nav agent */
+	UPROPERTY(config, EditAnywhere, Category=NavigationSystem)
+	uint32 bSkipAgentHeightCheckWhenPickingNavData:1;
+
+protected:
+	
+	UPROPERTY(EditDefaultsOnly, Category = "NavigationSystem", config)
+	ENavDataGatheringModeConfig DataGatheringMode;
+
+	/** If set to true navigation will be generated only around registered "navigation enforcers"
+	*	This has a range of consequences (including how navigation octree operates) so it needs to
+	*	be a conscious decision.
+	*	Once enabled results in whole world being navigable.
+	*	@see RegisterNavigationInvoker
+	*/
+	UPROPERTY(EditDefaultsOnly, Category = "Navigation Enforcing", config)
+	uint32 bGenerateNavigationOnlyAroundNavigationInvokers : 1;
+	
+	/** Minimal time, in seconds, between active tiles set update */
+	UPROPERTY(EditAnywhere, Category = "Navigation Enforcing", meta = (ClampMin = "0.1", UIMin = "0.1", EditCondition = "bGenerateNavigationOnlyAroundNavigationInvokers"), config)
+	float ActiveTilesUpdateInterval;
+
+	UPROPERTY(config, EditAnywhere, Category = Agents)
+	TArray<FNavDataConfig> SupportedAgents;
+
+public:
+	/** update frequency for dirty areas on navmesh */
+	UPROPERTY(config, EditAnywhere, Category=NavigationSystem)
+	float DirtyAreasUpdateFreq;
+
+	UPROPERTY()
+	TArray<ANavigationData*> NavDataSet;
+
+	UPROPERTY(transient)
+	TArray<ANavigationData*> NavDataRegistrationQueue;
+
+	TSet<FNavigationDirtyElement> PendingOctreeUpdates;
+
+	// List of pending navigation bounds update requests (add, remove, update size)
+	TArray<FNavigationBoundsUpdateRequest> PendingNavBoundsUpdates;
+
+ 	UPROPERTY(/*BlueprintAssignable, */Transient)
+	FOnNavDataGenerigEvent OnNavDataRegisteredEvent;
+
+	UPROPERTY(BlueprintAssignable, Transient, meta = (displayname = OnNavigationGenerationFinished))
+	FOnNavDataGenerigEvent OnNavigationGenerationFinishedDelegate;
+
+	FOnNavigationInitDone OnNavigationInitDone;
+	
+private:
+	TWeakObjectPtr<UCrowdManagerBase> CrowdManager;
+	
+	/** set to true when navigation processing was blocked due to missing nav bounds */
+	uint32 bNavDataRemovedDueToMissingNavBounds : 1;
+	
+	/** All areas where we build/have navigation */
+	TSet<FNavigationBounds> RegisteredNavBounds;
+
+	TMap<AActor*, FNavigationInvoker> Invokers;
+
+	float NextInvokersUpdateTime;
+	void UpdateInvokers();
+
+public:
+	//----------------------------------------------------------------------//
+	// Blueprint functions
+	//----------------------------------------------------------------------//
+
+	UFUNCTION(BlueprintPure, Category = "AI|Navigation", meta = (WorldContext = "WorldContextObject"))
+	static UNavigationSystemV1* GetNavigationSystem(UObject* WorldContextObject);
+	
+	/** Project a point onto the NavigationData */
+	UFUNCTION(BlueprintPure, Category = "AI|Navigation", meta = (WorldContext = "WorldContextObject", DisplayName = "ProjectPointToNavigation", ScriptName = "ProjectPointToNavigation"))
+	static bool K2_ProjectPointToNavigation(UObject* WorldContextObject, const FVector& Point, FVector& ProjectedLocation, ANavigationData* NavData, TSubclassOf<UNavigationQueryFilter> FilterClass, const FVector QueryExtent = FVector::ZeroVector);
+
+	/** Generates a random location reachable from given Origin location.
+	 *	@return Return Value represents if the call was successful */
+	UFUNCTION(BlueprintPure, Category = "AI|Navigation", meta = (WorldContext = "WorldContextObject", DisplayName = "GetRandomReachablePointInRadius", ScriptName = "GetRandomReachablePointInRadius"))
+	static bool K2_GetRandomReachablePointInRadius(UObject* WorldContextObject, const FVector& Origin, FVector& RandomLocation, float Radius, ANavigationData* NavData = NULL, TSubclassOf<UNavigationQueryFilter> FilterClass = NULL);
+
+	/** Generates a random location in navigable space within given radius of Origin.
+	 *	@return Return Value represents if the call was successful */
+	UFUNCTION(BlueprintPure, Category = "AI|Navigation", meta = (WorldContext = "WorldContextObject", DisplayName = "GetRandomPointInNavigableRadius", ScriptName = "GetRandomPointInNavigableRadius"))
+	static bool K2_GetRandomPointInNavigableRadius(UObject* WorldContextObject, const FVector& Origin, FVector& RandomLocation, float Radius, ANavigationData* NavData = NULL, TSubclassOf<UNavigationQueryFilter> FilterClass = NULL);
+	
+	/** Potentially expensive. Use with caution. Consider using UPathFollowingComponent::GetRemainingPathCost instead */
+	UFUNCTION(BlueprintPure, Category="AI|Navigation", meta=(WorldContext="WorldContextObject" ) )
+	static ENavigationQueryResult::Type GetPathCost(UObject* WorldContextObject, const FVector& PathStart, const FVector& PathEnd, float& PathCost, ANavigationData* NavData = NULL, TSubclassOf<UNavigationQueryFilter> FilterClass = NULL);
+
+	/** Potentially expensive. Use with caution */
+	UFUNCTION(BlueprintPure, Category="AI|Navigation", meta=(WorldContext="WorldContextObject" ) )
+	static ENavigationQueryResult::Type GetPathLength(UObject* WorldContextObject, const FVector& PathStart, const FVector& PathEnd, float& PathLength, ANavigationData* NavData = NULL, TSubclassOf<UNavigationQueryFilter> FilterClass = NULL);
+
+	UFUNCTION(BlueprintPure, Category="AI|Navigation", meta=(WorldContext="WorldContextObject" ) )
+	static bool IsNavigationBeingBuilt(UObject* WorldContextObject);
+
+	UFUNCTION(BlueprintPure, Category = "AI|Navigation", meta = (WorldContext = "WorldContextObject"))
+	static bool IsNavigationBeingBuiltOrLocked(UObject* WorldContextObject);
+
+	/** Finds path instantly, in a FindPath Synchronously. 
+	 *	@param PathfindingContext could be one of following: NavigationData (like Navmesh actor), Pawn or Controller. This parameter determines parameters of specific pathfinding query */
+	UFUNCTION(BlueprintCallable, Category = "AI|Navigation", meta = (WorldContext="WorldContextObject"))
+	static UNavigationPath* FindPathToLocationSynchronously(UObject* WorldContextObject, const FVector& PathStart, const FVector& PathEnd, AActor* PathfindingContext = NULL, TSubclassOf<UNavigationQueryFilter> FilterClass = NULL);
+
+	/** Finds path instantly, in a FindPath Synchronously. Main advantage over FindPathToLocationSynchronously is that 
+	 *	the resulting path will automatically get updated if goal actor moves more than TetherDistance away from last path node
+	 *	@param PathfindingContext could be one of following: NavigationData (like Navmesh actor), Pawn or Controller. This parameter determines parameters of specific pathfinding query */
+	UFUNCTION(BlueprintCallable, Category = "AI|Navigation", meta = (WorldContext="WorldContextObject"))
+	static UNavigationPath* FindPathToActorSynchronously(UObject* WorldContextObject, const FVector& PathStart, AActor* GoalActor, float TetherDistance = 50.f, AActor* PathfindingContext = NULL, TSubclassOf<UNavigationQueryFilter> FilterClass = NULL);
+
+	/** Performs navigation raycast on NavigationData appropriate for given Querier.
+	 *	@param Querier if not passed default navigation data will be used
+	 *	@param HitLocation if line was obstructed this will be set to hit location. Otherwise it contains SegmentEnd
+	 *	@return true if line from RayStart to RayEnd was obstructed. Also, true when no navigation data present */
+	UFUNCTION(BlueprintCallable, Category="AI|Navigation", meta=(WorldContext="WorldContextObject" ))
+	static bool NavigationRaycast(UObject* WorldContextObject, const FVector& RayStart, const FVector& RayEnd, FVector& HitLocation, TSubclassOf<UNavigationQueryFilter> FilterClass = NULL, AController* Querier = NULL);
+
+	/** will limit the number of simultaneously running navmesh tile generation jobs to specified number.
+	 *	@param MaxNumberOfJobs gets trimmed to be at least 1. You cannot use this function to pause navmesh generation */
+	UFUNCTION(BlueprintCallable, Category = "AI|Navigation")
+	void SetMaxSimultaneousTileGenerationJobsCount(int32 MaxNumberOfJobs);
+	
+	/** Brings limit of simultaneous navmesh tile generation jobs back to Project Setting's default value */
+	UFUNCTION(BlueprintCallable, Category = "AI|Navigation")
+	void ResetMaxSimultaneousTileGenerationJobsCount();
+
+	/** Registers given actor as a "navigation enforcer" which means navigation system will
+	 *	make sure navigation is being generated in specified radius around it.
+	 *	@note: you need NavigationSystem's GenerateNavigationOnlyAroundNavigationInvokers to be set to true
+	 *		to take advantage of this feature
+	 */
+	UFUNCTION(BlueprintCallable, Category = "AI|Navigation")
+	void RegisterNavigationInvoker(AActor* Invoker, float TileGenerationRadius = 3000, float TileRemovalRadius = 5000);
+
+	/** Removes given actor from the list of active navigation enforcers.
+	 *	@see RegisterNavigationInvoker for more details */
+	UFUNCTION(BlueprintCallable, Category = "AI|Navigation")
+	void UnregisterNavigationInvoker(AActor* Invoker);
+
+	UFUNCTION(BlueprintCallable, Category = "AI|Navigation|Generation")
+	void SetGeometryGatheringMode(ENavDataGatheringModeConfig NewMode);
+
+	FORCEINLINE bool IsActiveTilesGenerationEnabled() const{ return bGenerateNavigationOnlyAroundNavigationInvokers; }
+	
+	/** delegate type for events that dirty the navigation data ( Params: const FBox& DirtyBounds ) */
+	DECLARE_MULTICAST_DELEGATE_OneParam(FOnNavigationDirty, const FBox&);
+	/** called after navigation influencing event takes place*/
+	static FOnNavigationDirty NavigationDirtyEvent;
+
+	enum ERegistrationResult
+	{
+		RegistrationError,
+		RegistrationFailed_DataPendingKill,			// means navigation data being registered is marked as pending kill
+		RegistrationFailed_AgentAlreadySupported,	// this means that navigation agent supported by given nav data is already handled by some other, previously registered instance
+		RegistrationFailed_AgentNotValid,			// given instance contains navmesh that doesn't support any of expected agent types, or instance doesn't specify any agent
+		RegistrationSuccessful,
+	};
+
+	enum EOctreeUpdateMode
+	{
+		OctreeUpdate_Default = 0,						// regular update, mark dirty areas depending on exported content
+		OctreeUpdate_Geometry = 1,						// full update, mark dirty areas for geometry rebuild
+		OctreeUpdate_Modifiers = 2,						// quick update, mark dirty areas for modifier rebuild
+		OctreeUpdate_Refresh = 4,						// update is used for refresh, don't invalidate pending queue
+		OctreeUpdate_ParentChain = 8,					// update child nodes, don't remove anything
+	};
+
+	//~ Begin UObject Interface
+	virtual void PostInitProperties() override;
+	static void AddReferencedObjects(UObject* InThis, FReferenceCollector& Collector);
+#if WITH_EDITOR
+	virtual void PostEditChangeChainProperty(FPropertyChangedChainEvent& PropertyChangedEvent) override;
+	virtual void PostEditChangeProperty(struct FPropertyChangedEvent& PropertyChangedEvent) override;
+#endif // WITH_EDITOR
+	//~ End UObject Interface
+
+	virtual void Tick(float DeltaSeconds) override;	
+
+	UWorld* GetWorld() const override { return GetOuterUWorld(); }
+
+	UCrowdManagerBase* GetCrowdManager() const { return CrowdManager.Get(); }
+
+protected:
+	/** spawn new crowd manager */
+	virtual void CreateCrowdManager();
+
+	/** Used to properly set navigation class for indicated agent and propagate information to other places
+	 *	(like project settings) that may need this information 
+	 */
+	void SetSupportedAgentsNavigationClass(int32 AgentIndex, TSubclassOf<ANavigationData> NavigationDataClass);
+
+public:
+	//----------------------------------------------------------------------//
+	//~ Begin Public querying Interface                                                                
+	//----------------------------------------------------------------------//
+	/** 
+	 *	Synchronously looks for a path from @fLocation to @EndLocation for agent with properties @AgentProperties. NavData actor appropriate for specified 
+	 *	FNavAgentProperties will be found automatically
+	 *	@param ResultPath results are put here
+	 *	@param NavData optional navigation data that will be used instead of the one that would be deducted from AgentProperties
+	 *  @param Mode switch between normal and hierarchical path finding algorithms
+	 */
+	FPathFindingResult FindPathSync(const FNavAgentProperties& AgentProperties, FPathFindingQuery Query, EPathFindingMode::Type Mode = EPathFindingMode::Regular);
+
+	/** 
+	 *	Does a simple path finding from @StartLocation to @EndLocation on specified NavData. If none passed MainNavData will be used
+	 *	Result gets placed in ResultPath
+	 *	@param NavData optional navigation data that will be used instead main navigation data
+	 *  @param Mode switch between normal and hierarchical path finding algorithms
+	 */
+	FPathFindingResult FindPathSync(FPathFindingQuery Query, EPathFindingMode::Type Mode = EPathFindingMode::Regular);
+
+	/** 
+	 *	Asynchronously looks for a path from @StartLocation to @EndLocation for agent with properties @AgentProperties. NavData actor appropriate for specified 
+	 *	FNavAgentProperties will be found automatically
+	 *	@param ResultDelegate delegate that will be called once query has been processed and finished. Will be called even if query fails - in such case see comments for delegate's params
+	 *	@param NavData optional navigation data that will be used instead of the one that would be deducted from AgentProperties
+	 *	@param PathToFill if points to an actual navigation path instance than this instance will be filled with resulting path. Otherwise a new instance will be created and 
+	 *		used in call to ResultDelegate
+	 *  @param Mode switch between normal and hierarchical path finding algorithms
+	 *	@return request ID
+	 */
+	uint32 FindPathAsync(const FNavAgentProperties& AgentProperties, FPathFindingQuery Query, const FNavPathQueryDelegate& ResultDelegate, EPathFindingMode::Type Mode = EPathFindingMode::Regular);
+
+	/** Removes query indicated by given ID from queue of path finding requests to process. */
+	void AbortAsyncFindPathRequest(uint32 AsynPathQueryID);
+	
+	/** 
+	 *	Synchronously check if path between two points exists
+	 *  Does not return path object, but will run faster (especially in hierarchical mode)
+	 *  @param Mode switch between normal and hierarchical path finding algorithms. @note Hierarchical mode ignores QueryFilter
+	 *	@return true if path exists
+	 */
+	bool TestPathSync(FPathFindingQuery Query, EPathFindingMode::Type Mode = EPathFindingMode::Regular, int32* NumVisitedNodes = NULL) const;
+
+	/** Finds random point in navigable space
+	 *	@param ResultLocation Found point is put here
+	 *	@param NavData If NavData == NULL then MainNavData is used.
+	 *	@return true if any location found, false otherwise */
+	bool GetRandomPoint(FNavLocation& ResultLocation, ANavigationData* NavData = NULL, FSharedConstNavQueryFilter QueryFilter = NULL);
+
+	/** Finds random, reachable point in navigable space restricted to Radius around Origin
+	 *	@param ResultLocation Found point is put here
+	 *	@param NavData If NavData == NULL then MainNavData is used.
+	 *	@return true if any location found, false otherwise */
+	bool GetRandomReachablePointInRadius(const FVector& Origin, float Radius, FNavLocation& ResultLocation, ANavigationData* NavData = NULL, FSharedConstNavQueryFilter QueryFilter = NULL) const;
+
+	/** Finds random, point in navigable space restricted to Radius around Origin. Resulting location is not tested for reachability from the Origin
+	 *	@param ResultLocation Found point is put here
+	 *	@param NavData If NavData == NULL then MainNavData is used.
+	 *	@return true if any location found, false otherwise */
+	bool GetRandomPointInNavigableRadius(const FVector& Origin, float Radius, FNavLocation& ResultLocation, ANavigationData* NavData = NULL, FSharedConstNavQueryFilter QueryFilter = NULL) const;
+	
+	/** Calculates a path from PathStart to PathEnd and retrieves its cost. 
+	 *	@NOTE potentially expensive, so use it with caution */
+	ENavigationQueryResult::Type GetPathCost(const FVector& PathStart, const FVector& PathEnd, float& PathCost, const ANavigationData* NavData = NULL, FSharedConstNavQueryFilter QueryFilter = NULL) const;
+
+	/** Calculates a path from PathStart to PathEnd and retrieves its overestimated length.
+	 *	@NOTE potentially expensive, so use it with caution */
+	ENavigationQueryResult::Type GetPathLength(const FVector& PathStart, const FVector& PathEnd, float& PathLength, const ANavigationData* NavData = NULL, FSharedConstNavQueryFilter QueryFilter = NULL) const;
+
+	/** Calculates a path from PathStart to PathEnd and retrieves its overestimated length and cost.
+	 *	@NOTE potentially expensive, so use it with caution */
+	ENavigationQueryResult::Type GetPathLengthAndCost(const FVector& PathStart, const FVector& PathEnd, float& PathLength, float& PathCost, const ANavigationData* NavData = NULL, FSharedConstNavQueryFilter QueryFilter = NULL) const;
+
+	// @todo document
+	bool ProjectPointToNavigation(const FVector& Point, FNavLocation& OutLocation, const FVector& Extent = INVALID_NAVEXTENT, const FNavAgentProperties* AgentProperties = NULL, FSharedConstNavQueryFilter QueryFilter = NULL)
+	{
+		return ProjectPointToNavigation(Point, OutLocation, Extent, AgentProperties != NULL ? GetNavDataForProps(*AgentProperties) : GetDefaultNavDataInstance(FNavigationSystem::DontCreate), QueryFilter);
+	}
+
+	// @todo document
+	bool ProjectPointToNavigation(const FVector& Point, FNavLocation& OutLocation, const FVector& Extent = INVALID_NAVEXTENT, const ANavigationData* NavData = NULL, FSharedConstNavQueryFilter QueryFilter = NULL) const;
+
+	/** 
+	 * Looks for NavData generated for specified movement properties and returns it. NULL if not found;
+	 */
+	ANavigationData* GetNavDataForProps(const FNavAgentProperties& AgentProperties);
+
+	/** 
+	 * Looks for NavData generated for specified movement properties and returns it. NULL if not found; Const version.
+	 */
+	const ANavigationData* GetNavDataForProps(const FNavAgentProperties& AgentProperties) const;
+
+	/** Returns the world default navigation data instance. Creates one if it doesn't exist. */
+	ANavigationData* GetDefaultNavDataInstance(FNavigationSystem::ECreateIfMissing CreateNewIfNoneFound);
+	/** Returns the world default navigation data instance. */
+	virtual INavigationDataInterface* GetMainNavData() const override { return Cast<INavigationDataInterface>(GetDefaultNavDataInstance()); }
+	ANavigationData& GetMainNavDataChecked() const { check(MainNavData); return *MainNavData; }
+
+	ANavigationData* GetDefaultNavDataInstance() const { return MainNavData; }
+
+	ANavigationData* GetAbstractNavData() const { return AbstractNavData; }
+
+	/** constructs a navigation data instance of specified NavDataClass, in passed World
+	*	for supplied NavConfig */
+	virtual ANavigationData* CreateNavigationDataInstance(const FNavDataConfig& NavConfig);
+
+	FSharedNavQueryFilter CreateDefaultQueryFilterCopy() const;
+
+	/** Super-hacky safety feature for threaded navmesh building. Will be gone once figure out why keeping TSharedPointer to Navigation Generator doesn't 
+	 *	guarantee its existence */
+	bool ShouldGeneratorRun(const FNavDataGenerator* Generator) const;
+
+	virtual bool IsNavigationBuilt(const AWorldSettings* Settings) const override;
+
+	bool IsThereAnywhereToBuildNavigation() const;
+
+	bool ShouldGenerateNavigationEverywhere() const { return bWholeWorldNavigable; }
+
+	bool ShouldAllowClientSideNavigation() const { return bAllowClientSideNavigation; }
+	virtual bool ShouldLoadNavigationOnClient(ANavigationData* NavData = nullptr) const { return bAllowClientSideNavigation; }
+
+	FBox GetWorldBounds() const;
+	
+	FBox GetLevelBounds(ULevel* InLevel) const;
+
+	bool IsNavigationRelevant(const AActor* TestActor) const;
+
+	const TSet<FNavigationBounds>& GetNavigationBounds() const;
+
+	/** @return default walkable area class */
+	FORCEINLINE static TSubclassOf<UNavArea> GetDefaultWalkableArea() { return DefaultWalkableArea; }
+
+	/** @return default obstacle area class */
+	FORCEINLINE static TSubclassOf<UNavArea> GetDefaultObstacleArea() { return DefaultObstacleArea; }
+	
+	static const FNavDataConfig& GetDefaultSupportedAgent();
+	FORCEINLINE const FNavDataConfig& GetDefaultSupportedAgentConfig() const { check(SupportedAgents.Num() > 0);  return SupportedAgents[0]; }
+	FORCEINLINE const TArray<FNavDataConfig>& GetSupportedAgents() const { return SupportedAgents; }
+
+	virtual void ApplyWorldOffset(const FVector& InOffset, bool bWorldShift) override;
+
+	/** checks if navigation/navmesh is dirty and needs to be rebuilt */
+	bool IsNavigationDirty() const;
+
+	/** checks if dirty navigation data can rebuild itself */
+	bool CanRebuildDirtyNavigation() const;
+
+	FORCEINLINE bool SupportsNavigationGeneration() const { return bSupportRebuilding; }
+
+	static bool DoesPathIntersectBox(const FNavigationPath* Path, const FBox& Box, uint32 StartingIndex = 0, FVector* AgentExtent = NULL);
+	static bool DoesPathIntersectBox(const FNavigationPath* Path, const FBox& Box, const FVector& AgentLocation, uint32 StartingIndex = 0, FVector* AgentExtent = NULL);
+
+	//----------------------------------------------------------------------//
+	// Active tiles
+	//----------------------------------------------------------------------//
+	void RegisterInvoker(AActor& Invoker, float TileGenerationRadius, float TileRemovalRadius);
+	void UnregisterInvoker(AActor& Invoker);
+
+	static void RegisterNavigationInvoker(AActor& Invoker, float TileGenerationRadius, float TileRemovalRadius);
+	static void UnregisterNavigationInvoker(AActor& Invoker);
+
+	//----------------------------------------------------------------------//
+	// Bookkeeping 
+	//----------------------------------------------------------------------//
+	
+	// @todo document
+	virtual void UnregisterNavData(ANavigationData* NavData);
+
+	/** adds NavData to registration candidates queue - NavDataRegistrationQueue */
+	void RequestRegistration(ANavigationData* NavData, bool bTriggerRegistrationProcessing = true);
+
+protected:
+
+	/** Processes all NavigationData instances in UWorld owning navigation system instance, and registers
+	 *	all previously unregistered */
+	void RegisterNavigationDataInstances();
+
+	/** called in places where we need to spawn the NavOctree, but is checking additional conditions if we really want to do that
+	 *	depending on navigation data setup among others 
+	 *	@return true if NavOctree instance has been created, or if one is already present */
+	bool ConditionalPopulateNavOctree();
+
+	/** Processes registration of candidates queues via RequestRegistration and stored in NavDataRegistrationQueue */
+	virtual void ProcessRegistrationCandidates();
+
+	/** registers CustomLinks awaiting registration in PendingCustomLinkRegistration */
+	void ProcessCustomLinkPendingRegistration();
+
+	/** used to apply updates of nav volumes in navigation system's tick */
+	void PerformNavigationBoundsUpdate(const TArray<FNavigationBoundsUpdateRequest>& UpdateRequests);
+	
+	/** adds data to RegisteredNavBounds */
+	void AddNavigationBounds(const FNavigationBounds& NewBounds);
+
+	/** Searches for all valid navigation bounds in the world and stores them */
+	virtual void GatherNavigationBounds();
+
+	/** @return pointer to ANavigationData instance of given ID, or NULL if it was not found. Note it looks only through registered navigation data */
+	ANavigationData* GetNavDataWithID(const uint16 NavDataID) const;
+
+public:
+	void ReleaseInitialBuildingLock();
+
+	//----------------------------------------------------------------------//
+	// navigation octree related functions
+	//----------------------------------------------------------------------//
+	static void OnComponentRegistered(UActorComponent* Comp);
+	static void OnComponentUnregistered(UActorComponent* Comp);
+	static void OnActorRegistered(AActor* Actor);
+	static void OnActorUnregistered(AActor* Actor);
+
+	/** update navoctree entry for specified actor/component */
+	static void UpdateActorInNavOctree(AActor& Actor);
+	static void UpdateComponentInNavOctree(UActorComponent& Comp);
+	/** update all navoctree entries for actor and its components */
+	static void UpdateActorAndComponentsInNavOctree(AActor& Actor, bool bUpdateAttachedActors = true);
+	/** update all navoctree entries for actor and its non scene components after root movement */
+	static void UpdateNavOctreeAfterMove(USceneComponent* Comp);
+
+protected:
+	/** updates navoctree information on actors attached to RootActor */
+	static void UpdateAttachedActorsInNavOctree(AActor& RootActor);
+
+public:
+	/** removes all navoctree entries for actor and its components */
+	static void ClearNavOctreeAll(AActor* Actor);
+	/** updates bounds of all components implementing INavRelevantInterface */
+	static void UpdateNavOctreeBounds(AActor* Actor);
+
+	void AddDirtyArea(const FBox& NewArea, int32 Flags);
+	void AddDirtyAreas(const TArray<FBox>& NewAreas, int32 Flags);
+	bool HasDirtyAreasQueued() const;
+
+	const FNavigationOctree* GetNavOctree() const { return NavOctree.Get(); }
+	FNavigationOctree* GetMutableNavOctree() { return NavOctree.Get(); }
+
+	FORCEINLINE void SetObjectsNavOctreeId(UObject* Object, FOctreeElementId Id) { ObjectToOctreeId.Add(Object, Id); }
+	FORCEINLINE const FOctreeElementId* GetObjectsNavOctreeId(const UObject* Object) const { return ObjectToOctreeId.Find(Object); }
+	FORCEINLINE bool HasPendingObjectNavOctreeId(UObject* Object) const { return PendingOctreeUpdates.Contains(FNavigationDirtyElement(Object)); }
+	FORCEINLINE	void RemoveObjectsNavOctreeId(UObject* Object) { ObjectToOctreeId.Remove(Object); }
+	void RemoveNavOctreeElementId(const FOctreeElementId& ElementId, int32 UpdateFlags);
+
+	const FNavigationRelevantData* GetDataForObject(const UObject& Object) const;
+
+	/** find all elements in navigation octree within given box (intersection) */
+	void FindElementsInNavOctree(const FBox& QueryBox, const FNavigationOctreeFilter& Filter, TArray<FNavigationOctreeElement>& Elements);
+
+	/** update single element in navoctree */
+	void UpdateNavOctreeElement(UObject* ElementOwner, INavRelevantInterface* ElementInterface, int32 UpdateFlags);
+
+	/** force updating parent node and all its children */
+	void UpdateNavOctreeParentChain(UObject* ElementOwner, bool bSkipElementOwnerUpdate = false);
+
+	/** update component bounds in navigation octree and mark only specified area as dirty, doesn't re-export component geometry */
+	bool UpdateNavOctreeElementBounds(UActorComponent* Comp, const FBox& NewBounds, const FBox& DirtyArea);
+
+	//----------------------------------------------------------------------//
+	// Custom navigation links
+	//----------------------------------------------------------------------//
+	void RegisterCustomLink(INavLinkCustomInterface& CustomLink);
+	void UnregisterCustomLink(INavLinkCustomInterface& CustomLink);
+	
+	static void RequestCustomLinkRegistering(INavLinkCustomInterface& CustomLink, UObject* OwnerOb);
+	static void RequestCustomLinkUnregistering(INavLinkCustomInterface& CustomLink, UObject* ObjectOb);
+
+	/** find custom link by unique ID */
+	INavLinkCustomInterface* GetCustomLink(uint32 UniqueLinkId) const;
+
+	/** updates custom link for all active navigation data instances */
+	void UpdateCustomLink(const INavLinkCustomInterface* CustomLink);
+
+	//----------------------------------------------------------------------//
+	// Areas
+	//----------------------------------------------------------------------//
+	static void RequestAreaRegistering(UClass* NavAreaClass);
+	static void RequestAreaUnregistering(UClass* NavAreaClass);
+
+	/** find index in SupportedAgents array for given navigation data */
+	int32 GetSupportedAgentIndex(const ANavigationData* NavData) const;
+
+	/** find index in SupportedAgents array for agent type */
+	int32 GetSupportedAgentIndex(const FNavAgentProperties& NavAgent) const;
+
+	//----------------------------------------------------------------------//
+	// Filters
+	//----------------------------------------------------------------------//
+	
+	/** prepare descriptions of navigation flags in UNavigationQueryFilter class: using enum */
+	void DescribeFilterFlags(UEnum* FlagsEnum) const;
+
+	/** prepare descriptions of navigation flags in UNavigationQueryFilter class: using array */
+	void DescribeFilterFlags(const TArray<FString>& FlagsDesc) const;
+
+	/** removes cached filters from currently registered navigation data */
+	void ResetCachedFilter(TSubclassOf<UNavigationQueryFilter> FilterClass);
+
+	//----------------------------------------------------------------------//
+	// building
+	//----------------------------------------------------------------------//
+	
+	/** Triggers navigation building on all eligible navigation data. */
+	virtual void Build();
+
+	// @todo document
+	void OnPIEStart();
+	// @todo document
+	void OnPIEEnd();
+	
+	// @todo document
+	FORCEINLINE bool IsNavigationBuildingLocked() const { return NavBuildingLockFlags != 0; }
+
+	/** check if building is permanently locked to avoid showing navmesh building notify (due to queued dirty areas) */
+	FORCEINLINE bool IsNavigationBuildingPermanentlyLocked() const
+	{
+		return (NavBuildingLockFlags & ~ENavigationBuildLock::InitialLock) != 0; 
+	}
+
+	/** check if navigation octree updates are currently ignored */
+	FORCEINLINE bool IsNavigationOctreeLocked() const { return bNavOctreeLock; }
+
+	// @todo document
+	UFUNCTION(BlueprintCallable, Category = "AI|Navigation")
+	void OnNavigationBoundsUpdated(ANavMeshBoundsVolume* NavVolume);
+	virtual void OnNavigationBoundsAdded(ANavMeshBoundsVolume* NavVolume);
+	void OnNavigationBoundsRemoved(ANavMeshBoundsVolume* NavVolume);
+
+	/** Used to display "navigation building in progress" notify */
+	bool IsNavigationBuildInProgress(bool bCheckDirtyToo = true);
+
+	virtual void OnNavigationGenerationFinished(ANavigationData& NavData);
+
+	/** Used to display "navigation building in progress" counter */
+	int32 GetNumRemainingBuildTasks() const;
+
+	/** Number of currently running tasks */
+	int32 GetNumRunningBuildTasks() const;
+
+protected:
+	/** Sets up SuportedAgents and NavigationDataCreators. Override it to add additional setup, but make sure to call Super implementation */
+	virtual void DoInitialSetup();
+
+	/** spawn new crowd manager */
+	virtual void UpdateAbstractNavData();
+	
+public:
+	/** Called upon UWorld destruction to release what needs to be released */
+	virtual void CleanUp(const FNavigationSystem::ECleanupMode Mode = FNavigationSystem::ECleanupMode::CleanupUnsafe) override;
+
+	/** 
+	 *	Called when owner-UWorld initializes actors
+	 */
+	virtual void OnInitializeActors() override;
+
+	/** */
+	virtual void OnWorldInitDone(FNavigationSystemRunMode Mode);
+
+	FORCEINLINE bool IsInitialized() const { return bWorldInitDone; }
+
+	/** adds BSP collisions of currently streamed in levels to octree */
+	void InitializeLevelCollisions();
+
+	FORCEINLINE void AddNavigationBuildLock(uint8 Flags) { NavBuildingLockFlags |= Flags; }
+	void RemoveNavigationBuildLock(uint8 Flags, bool bSkipRebuildInEditor = false);
+
+	void SetNavigationOctreeLock(bool bLock) { bNavOctreeLock = bLock; }
+
+#if WITH_EDITOR
+	/** allow editor to toggle whether seamless navigation building is enabled */
+	static void SetNavigationAutoUpdateEnabled(bool bNewEnable, UNavigationSystemBase* InNavigationSystem);
+
+	/** check whether seamless navigation building is enabled*/
+	FORCEINLINE static bool GetIsNavigationAutoUpdateEnabled() { return bNavigationAutoUpdateEnabled; }
+
+	FORCEINLINE bool IsNavigationRegisterLocked() const { return NavUpdateLockFlags != 0; }
+	FORCEINLINE bool IsNavigationUnregisterLocked() const { return NavUpdateLockFlags && !(NavUpdateLockFlags & ENavigationLockReason::AllowUnregister); }
+	FORCEINLINE bool IsNavigationUpdateLocked() const { return IsNavigationRegisterLocked(); }
+	FORCEINLINE void AddNavigationUpdateLock(uint8 Flags) { NavUpdateLockFlags |= Flags; }
+	FORCEINLINE void RemoveNavigationUpdateLock(uint8 Flags) { NavUpdateLockFlags &= ~Flags; }
+
+	void UpdateLevelCollision(ULevel* InLevel);
+
+	virtual void OnEditorModeChanged(FEdMode* Mode, bool IsEntering);
+#endif // WITH_EDITOR
+
+	FORCEINLINE bool IsSetUpForLazyGeometryExporting() const { return bGenerateNavigationOnlyAroundNavigationInvokers; }
+
+	static UNavigationSystemV1* CreateNavigationSystem(UWorld* WorldOwner);
+
+	static UNavigationSystemV1* GetCurrent(UWorld* World);
+	static UNavigationSystemV1* GetCurrent(UObject* WorldContextObject);
+	
+	virtual void InitializeForWorld(UWorld& World, FNavigationSystemRunMode Mode) override;
+
+	// Fetch the array of all nav-agent properties.
+	void GetNavAgentPropertiesArray(TArray<FNavAgentProperties>& OutNavAgentProperties) const;
+
+	static FORCEINLINE bool ShouldUpdateNavOctreeOnComponentChange()
+	{
+		return (bUpdateNavOctreeOnComponentChange && !bStaticRuntimeNavigation)
+#if WITH_EDITOR
+			|| (GIsEditor && !GIsPlayInEditorWorld)
+#endif
+			;
+	}
+
+	static FORCEINLINE bool IsNavigationSystemStatic()
+	{
+		return bStaticRuntimeNavigation
+#if WITH_EDITOR
+			&& !(GIsEditor && !GIsPlayInEditorWorld)
+#endif
+			;
+	}
+
+	/** Use this function to signal the NavigationSystem it doesn't need to store
+	 *	any navigation-generation-related data at game runtime, because 
+	 *	nothing is going to use it anyway. This will short-circuit all code related 
+	 *	to navmesh rebuilding, so use it only if you have fully static navigation in 
+	 *	your game.
+	 *	Note: this is not a runtime switch. Call it before any actual game starts. */
+	static void ConfigureAsStatic();
+
+	static void SetUpdateNavOctreeOnComponentChange(bool bNewUpdateOnComponentChange);
+
+	/** 
+	 * Exec command handlers
+	 */
+	bool HandleCycleNavDrawnCommand( const TCHAR* Cmd, FOutputDevice& Ar );
+	bool HandleCountNavMemCommand( const TCHAR* Cmd, FOutputDevice& Ar );
+	
+	//----------------------------------------------------------------------//
+	// debug
+	//----------------------------------------------------------------------//
+	void CycleNavigationDataDrawn();
+
+protected:
+
+	UPROPERTY()
+	FNavigationSystemRunMode OperationMode;
+
+	TSharedPtr<FNavigationOctree, ESPMode::ThreadSafe> NavOctree;
+
+	TArray<FAsyncPathFindingQuery> AsyncPathFindingQueries;
+
+	FCriticalSection NavDataRegistration;
+
+	TMap<FNavAgentProperties, TWeakObjectPtr<ANavigationData> > AgentToNavDataMap;
+	
+	TMap<const UObject*, FOctreeElementId> ObjectToOctreeId;
+
+	/** Map of all objects that are tied to indexed navigation parent */
+	TMultiMap<UObject*, FWeakObjectPtr> OctreeChildNodesMap;
+
+	/** Map of all custom navigation links, that are relevant for path following */
+	TMap<uint32, FNavigationSystem::FCustomLinkOwnerInfo> CustomLinksMap;
+
+	/** stores areas marked as dirty throughout the frame, processes them 
+	 *	once a frame in Tick function */
+	TArray<FNavigationDirtyArea> DirtyAreas;
+
+	// async queries
+	FCriticalSection NavDataRegistrationSection;
+	static FCriticalSection CustomLinkRegistrationSection;
+	
+#if WITH_EDITOR
+	uint8 NavUpdateLockFlags;
+#endif
+	uint8 NavBuildingLockFlags;
+
+	/** set of locking flags applied on startup of navigation system */
+	uint8 InitialNavBuildingLockFlags;
+
+	/** if set, navoctree updates are ignored, use with caution! */
+	uint8 bNavOctreeLock : 1;
+
+	uint8 bInitialSetupHasBeenPerformed : 1;
+	uint8 bInitialLevelsAdded : 1;
+	uint8 bWorldInitDone : 1;
+	uint8 bAsyncBuildPaused : 1;
+	uint8 bCanAccumulateDirtyAreas : 1;
+#if !UE_BUILD_SHIPPING
+	uint8 bDirtyAreasReportedWhileAccumulationLocked : 1;
+#endif // !UE_BUILD_SHIPPING
+
+	/** cached navigable world bounding box*/
+	mutable FBox NavigableWorldBounds;
+
+	/** indicates which of multiple navigation data instances to draw*/
+	int32 CurrentlyDrawnNavDataIndex;
+
+	/** temporary cumulative time to calculate when we need to update dirty areas */
+	float DirtyAreasUpdateTime;
+	
+#if !UE_BUILD_SHIPPING
+	/** self-registering exec command to handle nav sys console commands */
+	static FNavigationSystemExec ExecHandler;
+#endif // !UE_BUILD_SHIPPING
+	
+	/** whether seamless navigation building is enabled */
+	static bool bNavigationAutoUpdateEnabled;
+	static bool bUpdateNavOctreeOnComponentChange;
+	static bool bStaticRuntimeNavigation;
+
+	static TMap<INavLinkCustomInterface*, FWeakObjectPtr> PendingCustomLinkRegistration;
+	TSet<const UClass*> NavAreaClasses;
+	static TSubclassOf<UNavArea> DefaultWalkableArea;
+	static TSubclassOf<UNavArea> DefaultObstacleArea;
+
+	/** delegate handler for PostLoadMap event */
+	void OnPostLoadMap(UWorld* LoadedWorld);
+#if WITH_EDITOR
+	/** delegate handler for ActorMoved events */
+	void OnActorMoved(AActor* Actor);
+#endif
+	/** delegate handler called when navigation is dirtied*/
+	void OnNavigationDirtied(const FBox& Bounds);
+
+#if WITH_HOT_RELOAD
+	FDelegateHandle HotReloadDelegateHandle;
+
+	/** called to notify NavigaitonSystem about finished hot reload */
+	virtual void OnHotReload(bool bWasTriggeredAutomatically);
+#endif // WITH_HOT_RELOAD
+
+	/** Registers given navigation data with this Navigation System.
+	 *	@return RegistrationSuccessful if registration was successful, other results mean it failed
+	 *	@see ERegistrationResult
+	 */
+	virtual ERegistrationResult RegisterNavData(ANavigationData* NavData);
+
+	/** tries to register navigation area */
+	void RegisterNavAreaClass(UClass* NavAreaClass);
+
+	/** tries to unregister navigation area */
+	void UnregisterNavAreaClass(UClass* NavAreaClass);
+
+	void OnNavigationAreaEvent(UClass* AreaClass, ENavAreaEvent::Type Event);
+	
+ 	FSetElementId RegisterNavOctreeElement(UObject* ElementOwner, INavRelevantInterface* ElementInterface, int32 UpdateFlags);
+	void UnregisterNavOctreeElement(UObject* ElementOwner, INavRelevantInterface* ElementInterface, int32 UpdateFlags);
+	
+	/** read element data from navigation octree */
+	bool GetNavOctreeElementData(UObject* NodeOwner, int32& DirtyFlags, FBox& DirtyBounds);
+
+	/** Adds given element to NavOctree. No check for owner's validity are performed, 
+	 *	nor its presence in NavOctree - function assumes callee responsibility 
+	 *	in this regard **/
+	void AddElementToNavOctree(const FNavigationDirtyElement& DirtyElement);
+
+	void SetCrowdManager(UCrowdManagerBase* NewCrowdManager);
+
+	/** Add BSP collision data to navigation octree */
+	void AddLevelCollisionToOctree(ULevel* Level);
+	
+	/** Remove BSP collision data from navigation octree */
+	void RemoveLevelCollisionFromOctree(ULevel* Level);
+
+	virtual void SpawnMissingNavigationData();
+
+private:
+	// adds navigation bounds update request to a pending list
+	void AddNavigationBoundsUpdateRequest(const FNavigationBoundsUpdateRequest& UpdateRequest);
+
+	/** Triggers navigation building on all eligible navigation data. */
+	void RebuildAll(bool bIsLoadTime = false);
+		 
+	/** Handler for FWorldDelegates::LevelAddedToWorld event */
+	 void OnLevelAddedToWorld(ULevel* InLevel, UWorld* InWorld);
+	 
+	/** Handler for FWorldDelegates::LevelRemovedFromWorld event */
+	void OnLevelRemovedFromWorld(ULevel* InLevel, UWorld* InWorld);
+
+	/** Adds given request to requests queue. Note it's to be called only on game thread only */
+	void AddAsyncQuery(const FAsyncPathFindingQuery& Query);
+		 
+	/** spawns a non-game-thread task to process requests given in PathFindingQueries.
+	 *	In the process PathFindingQueries gets copied. */
+	void TriggerAsyncQueries(TArray<FAsyncPathFindingQuery>& PathFindingQueries);
+
+	/** Processes pathfinding requests given in PathFindingQueries.*/
+	void PerformAsyncQueries(TArray<FAsyncPathFindingQuery> PathFindingQueries);
+
+	/** */
+	void DestroyNavOctree();
+
+	/** Whether Navigation system needs to populate nav octree. 
+	 *  Depends on runtime generation settings of each navigation data, always true in the editor
+	 */
+	bool RequiresNavOctree() const;
+
+	/** Return "Strongest" runtime generation type required by registered navigation data objects
+	 *  Depends on runtime generation settings of each navigation data, always ERuntimeGenerationType::Dynamic in the editor world
+	 */
+	ERuntimeGenerationType GetRuntimeGenerationType() const;
+	
+	//----------------------------------------------------------------------//
+	// new stuff
+	//----------------------------------------------------------------------//
+public:
+	void VerifyNavigationRenderingComponents(const bool bShow);
+	virtual int GetNavigationBoundsForNavData(const ANavigationData& NavData, TArray<FBox>& OutBounds) const;
+	static INavigationDataInterface* GetNavDataForActor(const AActor& Actor);
+	virtual void Configure(const UNavigationSystemConfig& Config);
+	
+#if !UE_BUILD_SHIPPING
+	void GetOnScreenMessages(TMultiMap<FCoreDelegates::EOnScreenMessageSeverity, FText>& OutMessages);
+#endif // !UE_BUILD_SHIPPING
+
+	//----------------------------------------------------------------------//
+	// DEPRECATED
+	//----------------------------------------------------------------------//
+public:
+	DEPRECATED(4.16, "This version of ProjectPointToNavigation is deprecated. Please use the new version")
+	UFUNCTION(BlueprintPure, Category = "AI|Navigation", meta = (WorldContext = "WorldContextObject", DisplayName = "ProjectPointToNavigation_DEPRECATED", ScriptNoExport, DeprecatedFunction, DeprecationMessage = "This version of ProjectPointToNavigation is deprecated. Please use the new version"))
+	static FVector ProjectPointToNavigation(UObject* WorldContextObject, const FVector& Point, ANavigationData* NavData = NULL, TSubclassOf<UNavigationQueryFilter> FilterClass = NULL, const FVector QueryExtent = FVector::ZeroVector);
+	DEPRECATED(4.16, "This version of GetRandomReachablePointInRadius is deprecated. Please use the new version")
+	UFUNCTION(BlueprintPure, Category = "AI|Navigation", meta = (WorldContext = "WorldContextObject", DisplayName = "GetRandomReachablePointInRadius_DEPRECATED", ScriptNoExport, DeprecatedFunction, DeprecationMessage = "This version of GetRandomReachablePointInRadius is deprecated. Please use the new version"))
+	static FVector GetRandomReachablePointInRadius(UObject* WorldContextObject, const FVector& Origin, float Radius, ANavigationData* NavData = NULL, TSubclassOf<UNavigationQueryFilter> FilterClass = NULL);
+	DEPRECATED(4.16, "This version of GetRandomPointInNavigableRadius is deprecated. Please use the new version")
+	UFUNCTION(BlueprintPure, Category = "AI|Navigation", meta = (WorldContext = "WorldContextObject", DisplayName = "GetRandomPointInNavigableRadius_DEPRECATED", ScriptNoExport, DeprecatedFunction, DeprecationMessage = "This version of GetRandomPointInNavigableRadius is deprecated. Please use the new version"))
+	static FVector GetRandomPointInNavigableRadius(UObject* WorldContextObject, const FVector& Origin, float Radius, ANavigationData* NavData = NULL, TSubclassOf<UNavigationQueryFilter> FilterClass = NULL);
+
+	DEPRECATED(4.20, "SimpleMoveToActor is deprecated. Use UAIBlueprintHelperLibrary::SimpleMoveToActor instead")
+	UFUNCTION(BlueprintCallable, Category = "AI|Navigation", meta = (DisplayName = "SimpleMoveToActor_DEPRECATED", ScriptNoExport, DeprecatedFunction, DeprecationMessage = "SimpleMoveToActor is deprecated. Use AIBlueprintHelperLibrary::SimpleMoveToActor instead"))
+	static void SimpleMoveToActor(AController* Controller, const AActor* Goal);
+
+	DEPRECATED(4.20, "SimpleMoveToLocation is deprecated. Use UAIBlueprintHelperLibrary::SimpleMoveToLocation instead")
+	UFUNCTION(BlueprintCallable, Category = "AI|Navigation", meta = (DisplayName = "SimpleMoveToLocation_DEPRECATED", ScriptNoExport, DeprecatedFunction, DeprecationMessage = "SimpleMoveToLocation is deprecated. Use AIBlueprintHelperLibrary::SimpleMoveToLocation instead"))
+	static void SimpleMoveToLocation(AController* Controller, const FVector& Goal);
+};
+
diff --git a/Engine/Source/Runtime/NavigationSystem/Public/NavigationSystemTypes.h b/Engine/Source/Runtime/NavigationSystem/Public/NavigationSystemTypes.h
new file mode 100644
--- /dev/null
+++ b/Engine/Source/Runtime/NavigationSystem/Public/NavigationSystemTypes.h
@@ -0,0 +1,176 @@
+// Copyright 1998-2018 Epic Games, Inc. All Rights Reserved.
+
+#pragma once
+
+#include "CoreMinimal.h"
+#include "NavFilters/NavigationQueryFilter.h"
+#include "AI/NavigationSystemConfig.h"
+#include "AI/Navigation/NavLinkDefinition.h"
+#include "GenericOctreePublic.h"
+#include "AI/NavigationModifier.h"
+
+
+#define NAVSYS_DEBUG (0 && UE_BUILD_DEBUG)
+
+class UBodySetup;
+class UNavCollision;
+struct FKAggregateGeom;
+class FNavigationOctree;
+class UNavigationPath;
+class ANavigationData;
+
+
+struct NAVIGATIONSYSTEM_API FPathFindingQueryData
+{
+	TWeakObjectPtr<const UObject> Owner;
+	FVector StartLocation;
+	FVector EndLocation;
+	FSharedConstNavQueryFilter QueryFilter;
+
+	/** additional flags passed to navigation data handling request */
+	int32 NavDataFlags;
+
+	/** if set, allow partial paths as a result */
+	uint32 bAllowPartialPaths : 1;
+
+	FPathFindingQueryData() : StartLocation(FNavigationSystem::InvalidLocation), EndLocation(FNavigationSystem::InvalidLocation), NavDataFlags(0), bAllowPartialPaths(true) {}
+	FPathFindingQueryData(const UObject* InOwner, const FVector& InStartLocation, const FVector& InEndLocation, FSharedConstNavQueryFilter InQueryFilter = nullptr, int32 InNavDataFlags = 0, bool bInAllowPartialPaths = true) :
+		Owner(InOwner), StartLocation(InStartLocation), EndLocation(InEndLocation), QueryFilter(InQueryFilter), NavDataFlags(InNavDataFlags), bAllowPartialPaths(bInAllowPartialPaths) {}
+};
+
+struct NAVIGATIONSYSTEM_API FPathFindingQuery : public FPathFindingQueryData
+{
+	TWeakObjectPtr<const ANavigationData> NavData;
+	FNavPathSharedPtr PathInstanceToFill;
+	FNavAgentProperties NavAgentProperties;
+
+	FPathFindingQuery() : FPathFindingQueryData() {}
+	FPathFindingQuery(const FPathFindingQuery& Source);
+	FPathFindingQuery(const UObject* InOwner, const ANavigationData& InNavData, const FVector& Start, const FVector& End, FSharedConstNavQueryFilter SourceQueryFilter = NULL, FNavPathSharedPtr InPathInstanceToFill = NULL);
+	FPathFindingQuery(const INavAgentInterface& InNavAgent, const ANavigationData& InNavData, const FVector& Start, const FVector& End, FSharedConstNavQueryFilter SourceQueryFilter = NULL, FNavPathSharedPtr InPathInstanceToFill = NULL);
+
+	explicit FPathFindingQuery(FNavPathSharedRef PathToRecalculate, const ANavigationData* NavDataOverride = NULL);
+
+	FPathFindingQuery& SetPathInstanceToUpdate(FNavPathSharedPtr InPathInstanceToFill) { PathInstanceToFill = InPathInstanceToFill; return *this; }
+	FPathFindingQuery& SetAllowPartialPaths(bool bAllow) { bAllowPartialPaths = bAllow; return *this; }
+	FPathFindingQuery& SetNavAgentProperties(const FNavAgentProperties& InNavAgentProperties) { NavAgentProperties = InNavAgentProperties; return *this; }
+};
+
+namespace EPathFindingMode
+{
+	enum Type
+	{
+		Regular,
+		Hierarchical,
+	};
+};
+
+////////////////////////////////////////////////////////////////////////////
+//// Custom path following data
+//
+///** Custom data passed to movement requests. */
+struct NAVIGATIONSYSTEM_API FMoveRequestCustomData
+{
+};
+
+typedef TSharedPtr<FMoveRequestCustomData, ESPMode::ThreadSafe> FCustomMoveSharedPtr;
+typedef TWeakPtr<FMoveRequestCustomData, ESPMode::ThreadSafe> FCustomMoveWeakPtr;
+
+//----------------------------------------------------------------------//
+// Active tiles 
+//----------------------------------------------------------------------//
+struct FNavigationInvokerRaw
+{
+	FVector Location;
+	float RadiusMin;
+	float RadiusMax;
+
+	FNavigationInvokerRaw(const FVector& InLocation, float Min, float Max)
+		: Location(InLocation), RadiusMin(Min), RadiusMax(Max)
+	{}
+};
+
+struct FNavigationInvoker
+{
+	TWeakObjectPtr<AActor> Actor;
+
+	/** tiles GenerationRadius away or close will be generated if they're not already present */
+	float GenerationRadius;
+	/** tiles over RemovalRadius will get removed.
+	*	@Note needs to be >= GenerationRadius or will get clampped */
+	float RemovalRadius;
+
+	FNavigationInvoker();
+	FNavigationInvoker(AActor& InActor, float InGenerationRadius, float InRemovalRadius);
+};
+
+namespace NavigationHelper
+{
+	void GatherCollision(UBodySetup* RigidBody, TNavStatArray<FVector>& OutVertexBuffer, TNavStatArray<int32>& OutIndexBuffer, const FTransform& ComponentToWorld = FTransform::Identity);
+	void GatherCollision(UBodySetup* RigidBody, UNavCollision* NavCollision);
+
+	/** gather collisions from aggregated geom, convex and tri mesh elements are not supported - use override with full UBodySetup param instead */
+	void GatherCollision(const FKAggregateGeom& AggGeom, UNavCollision& NavCollision);
+}
+
+namespace FNavigationSystem
+{
+	enum ECreateIfMissing
+	{
+		Invalid = -1,
+		DontCreate = 0,
+		Create = 1,
+	};
+
+	typedef ECreateIfMissing ECreateIfEmpty;
+}
+
+//----------------------------------------------------------------------//
+// 
+//----------------------------------------------------------------------//
+
+namespace NavigationHelper
+{
+	struct NAVIGATIONSYSTEM_API FNavLinkOwnerData
+	{
+		const AActor* Actor;
+		FTransform LinkToWorld;
+
+		FNavLinkOwnerData() : Actor(nullptr) {}
+		FNavLinkOwnerData(const AActor& InActor);
+		FNavLinkOwnerData(const USceneComponent& InComponent);
+	};
+
+	DECLARE_DELEGATE_ThreeParams(FNavLinkProcessorDelegate, FCompositeNavModifier*, const AActor*, const TArray<FNavigationLink>&);
+	DECLARE_DELEGATE_ThreeParams(FNavLinkSegmentProcessorDelegate, FCompositeNavModifier*, const AActor*, const TArray<FNavigationSegmentLink>&);
+
+	DECLARE_DELEGATE_ThreeParams(FNavLinkProcessorDataDelegate, FCompositeNavModifier*, const FNavLinkOwnerData&, const TArray<FNavigationLink>&);
+	DECLARE_DELEGATE_ThreeParams(FNavLinkSegmentProcessorDataDelegate, FCompositeNavModifier*, const FNavLinkOwnerData&, const TArray<FNavigationSegmentLink>&);
+
+	/** Set new implementation of nav link processor, a function that will be
+	*	be used to process/transform links before adding them to CompositeModifier.
+	*	This function is supposed to be called once during the engine/game
+	*	setup phase. Not intended to be toggled at runtime */
+	NAVIGATIONSYSTEM_API void SetNavLinkProcessorDelegate(const FNavLinkProcessorDataDelegate& NewDelegate);
+	NAVIGATIONSYSTEM_API void SetNavLinkSegmentProcessorDelegate(const FNavLinkSegmentProcessorDataDelegate& NewDelegate);
+
+	/** called to do any necessary processing on NavLinks and put results in CompositeModifier */
+	NAVIGATIONSYSTEM_API void ProcessNavLinkAndAppend(FCompositeNavModifier* OUT CompositeModifier, const AActor* Actor, const TArray<FNavigationLink>& IN NavLinks);
+	NAVIGATIONSYSTEM_API void ProcessNavLinkAndAppend(FCompositeNavModifier* OUT CompositeModifier, const FNavLinkOwnerData& OwnerData, const TArray<FNavigationLink>& IN NavLinks);
+
+	/** called to do any necessary processing on NavLinks and put results in CompositeModifier */
+	NAVIGATIONSYSTEM_API void ProcessNavLinkSegmentAndAppend(FCompositeNavModifier* OUT CompositeModifier, const AActor* Actor, const TArray<FNavigationSegmentLink>& IN NavLinks);
+	NAVIGATIONSYSTEM_API void ProcessNavLinkSegmentAndAppend(FCompositeNavModifier* OUT CompositeModifier, const FNavLinkOwnerData& OwnerData, const TArray<FNavigationSegmentLink>& IN NavLinks);
+
+	NAVIGATIONSYSTEM_API void DefaultNavLinkProcessorImpl(FCompositeNavModifier* OUT CompositeModifier, const FNavLinkOwnerData& OwnerData, const TArray<FNavigationLink>& IN NavLinks);
+	NAVIGATIONSYSTEM_API void DefaultNavLinkSegmentProcessorImpl(FCompositeNavModifier* OUT CompositeModifier, const FNavLinkOwnerData& OwnerData, const TArray<FNavigationSegmentLink>& IN NavLinks);
+
+	NAVIGATIONSYSTEM_API bool IsBodyNavigationRelevant(const UBodySetup& IN BodySetup);
+
+	// deprecated functions	
+	DEPRECATED_FORGAME(4.17, "This function is now deprecated, please use override with OwnerData argument.")
+	NAVIGATIONSYSTEM_API void DefaultNavLinkProcessorImpl(FCompositeNavModifier* OUT CompositeModifier, const AActor* Actor, const TArray<FNavigationLink>& IN NavLinks);
+
+	DEPRECATED_FORGAME(4.17, "This function is now deprecated, please use override with OwnerData argument.")
+	NAVIGATIONSYSTEM_API void DefaultNavLinkSegmentProcessorImpl(FCompositeNavModifier* OUT CompositeModifier, const AActor* Actor, const TArray<FNavigationSegmentLink>& IN NavLinks);
+}
\ No newline at end of file
diff --git a/Engine/Source/Runtime/Slate/Public/Framework/Commands/InputChord.h b/Engine/Source/Runtime/Slate/Public/Framework/Commands/InputChord.h
--- a/Engine/Source/Runtime/Slate/Public/Framework/Commands/InputChord.h
+++ b/Engine/Source/Runtime/Slate/Public/Framework/Commands/InputChord.h
@@ -53,7 +53,8 @@ struct SLATE_API FInputChord
 	ERelationshipType GetRelationship(const FInputChord& OtherChord) const;
 
 	FInputChord()
-		: bShift(false)
+		: Key(EKeys::Invalid)
+		, bShift(false)
 		, bCtrl(false)
 		, bAlt(false)
 		, bCmd(false)
diff --git a/Engine/Source/Runtime/UMG/Private/WidgetBlueprintLibrary.cpp b/Engine/Source/Runtime/UMG/Private/WidgetBlueprintLibrary.cpp
--- a/Engine/Source/Runtime/UMG/Private/WidgetBlueprintLibrary.cpp
+++ b/Engine/Source/Runtime/UMG/Private/WidgetBlueprintLibrary.cpp
@@ -141,7 +141,7 @@ void UWidgetBlueprintLibrary::SetFocusToGameViewport()
 	FSlateApplication::Get().SetAllUserFocusToGameViewport();
 }
 
-void UWidgetBlueprintLibrary::DrawBox(UPARAM(ref) FPaintContext& Context, FVector2D Position, FVector2D Size, USlateBrushAsset* Brush, FLinearColor Tint)
+void UWidgetBlueprintLibrary::DrawBox(FPaintContext& Context, FVector2D Position, FVector2D Size, USlateBrushAsset* Brush, FLinearColor Tint)
 {
 	Context.MaxLayer++;
 
@@ -157,7 +157,7 @@ void UWidgetBlueprintLibrary::DrawBox(UPARAM(ref) FPaintContext& Context, FVecto
 	}
 }
 
-void UWidgetBlueprintLibrary::DrawLine(UPARAM(ref) FPaintContext& Context, FVector2D PositionA, FVector2D PositionB, FLinearColor Tint, bool bAntiAlias)
+void UWidgetBlueprintLibrary::DrawLine(FPaintContext& Context, FVector2D PositionA, FVector2D PositionB, FLinearColor Tint, bool bAntiAlias)
 {
 	Context.MaxLayer++;
 
@@ -175,7 +175,7 @@ void UWidgetBlueprintLibrary::DrawLine(UPARAM(ref) FPaintContext& Context, FVect
 		bAntiAlias);
 }
 
-void UWidgetBlueprintLibrary::DrawLines(UPARAM(ref) FPaintContext& Context, const TArray<FVector2D>& Points, FLinearColor Tint, bool bAntiAlias)
+void UWidgetBlueprintLibrary::DrawLines(FPaintContext& Context, const TArray<FVector2D>& Points, FLinearColor Tint, bool bAntiAlias)
 {
 	Context.MaxLayer++;
 
@@ -189,7 +189,7 @@ void UWidgetBlueprintLibrary::DrawLines(UPARAM(ref) FPaintContext& Context, cons
 		bAntiAlias);
 }
 
-void UWidgetBlueprintLibrary::DrawText(UPARAM(ref) FPaintContext& Context, const FString& InString, FVector2D Position, FLinearColor Tint)
+void UWidgetBlueprintLibrary::DrawText(FPaintContext& Context, const FString& InString, FVector2D Position, FLinearColor Tint)
 {
 	Context.MaxLayer++;
 
@@ -206,7 +206,7 @@ void UWidgetBlueprintLibrary::DrawText(UPARAM(ref) FPaintContext& Context, const
 		Tint);
 }
 
-void UWidgetBlueprintLibrary::DrawTextFormatted(UPARAM(ref) FPaintContext& Context, const FText& Text, FVector2D Position, UFont* Font, int32 FontSize, FName FontTypeFace, FLinearColor Tint)
+void UWidgetBlueprintLibrary::DrawTextFormatted(FPaintContext& Context, const FText& Text, FVector2D Position, UFont* Font, int32 FontSize, FName FontTypeFace, FLinearColor Tint)
 {
 	if ( Font )
 	{
@@ -242,7 +242,7 @@ FEventReply UWidgetBlueprintLibrary::Unhandled()
 	return Reply;
 }
 
-FEventReply UWidgetBlueprintLibrary::CaptureMouse(UPARAM(ref) FEventReply& Reply, UWidget* CapturingWidget)
+FEventReply UWidgetBlueprintLibrary::CaptureMouse(FEventReply& Reply, UWidget* CapturingWidget)
 {
 	if ( CapturingWidget )
 	{
@@ -256,7 +256,7 @@ FEventReply UWidgetBlueprintLibrary::CaptureMouse(UPARAM(ref) FEventReply& Reply
 	return Reply;
 }
 
-FEventReply UWidgetBlueprintLibrary::ReleaseMouseCapture(UPARAM(ref) FEventReply& Reply)
+FEventReply UWidgetBlueprintLibrary::ReleaseMouseCapture(FEventReply& Reply)
 {
 	Reply.NativeReply = Reply.NativeReply.ReleaseMouseCapture();
 
@@ -298,24 +298,24 @@ FEventReply UWidgetBlueprintLibrary::SetUserFocus( UPARAM( ref ) FEventReply& Re
 	return Reply;
 }
 
-FEventReply UWidgetBlueprintLibrary::CaptureJoystick(UPARAM(ref) FEventReply& Reply, UWidget* CapturingWidget, bool bInAllJoysticks/* = false*/)
+FEventReply UWidgetBlueprintLibrary::CaptureJoystick(FEventReply& Reply, UWidget* CapturingWidget, bool bInAllJoysticks/* = false*/)
 {
 	return SetUserFocus(Reply, CapturingWidget, bInAllJoysticks);
 }
 
-FEventReply UWidgetBlueprintLibrary::ClearUserFocus(UPARAM(ref) FEventReply& Reply, bool bInAllUsers /*= false*/)
+FEventReply UWidgetBlueprintLibrary::ClearUserFocus(FEventReply& Reply, bool bInAllUsers /*= false*/)
 {
 	Reply.NativeReply = Reply.NativeReply.ClearUserFocus(bInAllUsers);
 
 	return Reply;
 }
 
-FEventReply UWidgetBlueprintLibrary::ReleaseJoystickCapture(UPARAM(ref) FEventReply& Reply, bool bInAllJoysticks /*= false*/)
+FEventReply UWidgetBlueprintLibrary::ReleaseJoystickCapture(FEventReply& Reply, bool bInAllJoysticks /*= false*/)
 {
 	return ClearUserFocus(Reply, bInAllJoysticks);
 }
 
-FEventReply UWidgetBlueprintLibrary::SetMousePosition(UPARAM(ref) FEventReply& Reply, FVector2D NewMousePosition)
+FEventReply UWidgetBlueprintLibrary::SetMousePosition(FEventReply& Reply, FVector2D NewMousePosition)
 {
 	FIntPoint NewPoint((int32)NewMousePosition.X, (int32)NewMousePosition.Y);
 	Reply.NativeReply = Reply.NativeReply.SetMousePos(NewPoint);
@@ -323,7 +323,7 @@ FEventReply UWidgetBlueprintLibrary::SetMousePosition(UPARAM(ref) FEventReply& R
 	return Reply;
 }
 
-FEventReply UWidgetBlueprintLibrary::DetectDrag(UPARAM(ref) FEventReply& Reply, UWidget* WidgetDetectingDrag, FKey DragKey)
+FEventReply UWidgetBlueprintLibrary::DetectDrag(FEventReply& Reply, UWidget* WidgetDetectingDrag, FKey DragKey)
 {
 	if ( WidgetDetectingDrag )
 	{
@@ -348,7 +348,7 @@ FEventReply UWidgetBlueprintLibrary::DetectDragIfPressed(const FPointerEvent& Po
 	return UWidgetBlueprintLibrary::Unhandled();
 }
 
-FEventReply UWidgetBlueprintLibrary::EndDragDrop(UPARAM(ref) FEventReply& Reply)
+FEventReply UWidgetBlueprintLibrary::EndDragDrop(FEventReply& Reply)
 {
 	Reply.NativeReply = Reply.NativeReply.EndDragDrop();
 
@@ -425,17 +425,17 @@ FSlateBrush UWidgetBlueprintLibrary::MakeBrushFromMaterial(UMaterialInterface* M
 	return FSlateNoResource();
 }
 
-UObject* UWidgetBlueprintLibrary::GetBrushResource(FSlateBrush& Brush)
+UObject* UWidgetBlueprintLibrary::GetBrushResource(const FSlateBrush& Brush)
 {
 	return Brush.GetResourceObject();
 }
 
-UTexture2D* UWidgetBlueprintLibrary::GetBrushResourceAsTexture2D(UPARAM(ref) FSlateBrush& Brush)
+UTexture2D* UWidgetBlueprintLibrary::GetBrushResourceAsTexture2D(const FSlateBrush& Brush)
 {
 	return Cast<UTexture2D>(Brush.GetResourceObject());
 }
 
-UMaterialInterface* UWidgetBlueprintLibrary::GetBrushResourceAsMaterial(UPARAM(ref) FSlateBrush& Brush)
+UMaterialInterface* UWidgetBlueprintLibrary::GetBrushResourceAsMaterial(const FSlateBrush& Brush)
 {
 	return Cast<UMaterialInterface>(Brush.GetResourceObject());
 }
