Fix NavigationHelper::IsBodyNavigationRelevant for WITH_PHYSX = 0
#rnx

diff --git a/Engine/Source/Runtime/NavigationSystem/Private/NavigationSystemTypes.cpp b/Engine/Source/Runtime/NavigationSystem/Private/NavigationSystemTypes.cpp
--- a/Engine/Source/Runtime/NavigationSystem/Private/NavigationSystemTypes.cpp
+++ b/Engine/Source/Runtime/NavigationSystem/Private/NavigationSystemTypes.cpp
@@ -233,8 +233,14 @@ namespace NavigationHelper
 
 	bool IsBodyNavigationRelevant(const UBodySetup& BodySetup)
 	{
+#if WITH_PHYSX
+		const bool bBodyHasGeometry = (BodySetup.AggGeom.GetElementCount() > 0 || BodySetup.TriMeshes.Num() > 0);
+#else
+		const bool bBodyHasGeometry = (BodySetup.AggGeom.GetElementCount() > 0);
+#endif
+
 		// has any colliding geometry
-		return (BodySetup.AggGeom.GetElementCount() > 0 || BodySetup.TriMeshes.Num() > 0)
+		return bBodyHasGeometry
 			// AND blocks any of Navigation-relevant 
 			&& (BodySetup.DefaultInstance.GetResponseToChannel(ECC_Pawn) == ECR_Block || BodySetup.DefaultInstance.GetResponseToChannel(ECC_Vehicle) == ECR_Block)
 			// AND has full colliding capabilities 
