Fixed new navigation system not respecting 'Allow Client Side Navigation' setting #UE4

diff --git a/Engine/Config/BaseEngine.ini b/Engine/Config/BaseEngine.ini
--- a/Engine/Config/BaseEngine.ini
+++ b/Engine/Config/BaseEngine.ini
@@ -52,6 +52,7 @@ GameViewportClientClassName=/Script/Engine.GameViewportClient
 LocalPlayerClassName=/Script/Engine.LocalPlayer
 WorldSettingsClassName=/Script/Engine.WorldSettings
 NavigationSystemClassName=/Script/NavigationSystem.NavigationSystemV1
+NavigationSystemConfigClassName=/Script/NavigationSystem.NavigationSystemModuleConfig
 AvoidanceManagerClassName=/Script/Engine.AvoidanceManager
 PhysicsCollisionHandlerClassName=/Script/Engine.PhysicsCollisionHandler
 LevelScriptActorClassName=/Script/Engine.LevelScriptActor
@@ -547,9 +548,10 @@ bUseStreamingPause=false
 +ClassRedirects=(OldName="NavigationGraphNodeComponent",NewName="/Script/NavigationSystem.NavigationGraphNodeComponent")
 +ClassRedirects=(OldName="NavigationPathGenerator",NewName="/Script/NavigationSystem.NavigationPathGenerator")
 +ClassRedirects=(OldName="NavLinkCustomInterface",NewName="/Script/NavigationSystem.NavLinkCustomInterface")
-+ClassRedirects=(OldName="NavLinkHostInterface",NewName="/Script/NavigationSystem.NavLinkHostInterface")
++ClassRedirects=(OldName="NavLinkHostInterface",NewName="/Script/NavigationSystemConfig.NavLinkHostInterface")
 +ClassRedirects=(OldName="NavNodeInterface",NewName="/Script/NavigationSystem.NavNodeInterface")
 +ClassRedirects=(OldName="NavLinkProxy",NewName="/Script/AIModule.NavLinkProxy")
++ClassRedirects=(OldName="NavigationSystemConfig", NewName="/Script/NavigationSystem.NavigationSystemModuleConfig")
 +StructRedirects=(OldName="NavigationFilterArea",NewName="/Script/NavigationSystem.NavigationFilterArea")
 +StructRedirects=(OldName="NavigationFilterFlags",NewName="/Script/NavigationSystem.NavigationFilterFlags")
 +StructRedirects=(OldName="NavGraphEdge",NewName="/Script/NavigationSystem.NavGraphEdge")
diff --git a/Engine/Source/Runtime/Engine/Classes/AI/NavigationSystemConfig.h b/Engine/Source/Runtime/Engine/Classes/AI/NavigationSystemConfig.h
--- a/Engine/Source/Runtime/Engine/Classes/AI/NavigationSystemConfig.h
+++ b/Engine/Source/Runtime/Engine/Classes/AI/NavigationSystemConfig.h
@@ -19,14 +19,7 @@ class ENGINE_API UNavigationSystemConfig : public UObject
 public:
 	UPROPERTY(EditAnywhere, Category=Navigation, meta = (MetaClass = "NavigationSystemBase", NoResetToDefault))
 	FSoftClassPath NavigationSystemClass;
-
-	/** Whether at game runtime we expect any kind of dynamic navigation generation */
-	UPROPERTY(EditAnywhere, Category = Navigation)
-	uint32 bStrictlyStatic : 1;
-
-	UPROPERTY(EditAnywhere, Category = Navigation)
-	uint32 bCreateOnClient : 1;
-
+	
 public:
 	UNavigationSystemConfig(const FObjectInitializer& ObjectInitializer = FObjectInitializer::Get());
 
@@ -35,6 +28,8 @@ public:
 #if WITH_EDITOR
 	virtual void PostEditChangeProperty(FPropertyChangedEvent& PropertyChangedEvent) override;
 #endif // WITH_EDITOR
+
+	static TSubclassOf<UNavigationSystemConfig> GetDefaultConfigClass();
 };
 
 UCLASS(MinimalAPI, HideCategories=Navigation)
diff --git a/Engine/Source/Runtime/Engine/Classes/Engine/Engine.h b/Engine/Source/Runtime/Engine/Classes/Engine/Engine.h
--- a/Engine/Source/Runtime/Engine/Classes/Engine/Engine.h
+++ b/Engine/Source/Runtime/Engine/Classes/Engine/Engine.h
@@ -738,6 +738,9 @@ public:
 	/** The class for NavigationSystem **/
 	UPROPERTY()
 	TSubclassOf<class UNavigationSystemBase>  NavigationSystemClass;
+
+	UPROPERTY(globalconfig, noclear, meta = (MetaClass = "NavigationSystem", DisplayName = "Navigation System Config Class"))
+	FSoftClassPath NavigationSystemConfigClassName;
 	
 	/** Name of behavior tree manager class */
 	UPROPERTY(globalconfig, noclear, meta=(MetaClass="AvoidanceManager", DisplayName="Avoidance Manager Class"))
diff --git a/Engine/Source/Runtime/Engine/Private/AI/NavigationSystemConfig.cpp b/Engine/Source/Runtime/Engine/Private/AI/NavigationSystemConfig.cpp
--- a/Engine/Source/Runtime/Engine/Private/AI/NavigationSystemConfig.cpp
+++ b/Engine/Source/Runtime/Engine/Private/AI/NavigationSystemConfig.cpp
@@ -15,6 +15,17 @@ UNavigationSystemConfig::UNavigationSystemConfig(const FObjectInitializer& Objec
 	}
 }
 
+TSubclassOf<UNavigationSystemConfig> UNavigationSystemConfig::GetDefaultConfigClass()
+{
+	TSubclassOf<UNavigationSystemConfig> NavSysConfigClass;
+	if (GEngine && GEngine->NavigationSystemConfigClassName.IsValid())
+	{
+		GEngine->NavigationSystemConfigClassName.TryLoad();
+		NavSysConfigClass = GEngine->NavigationSystemConfigClassName.ResolveClass();
+	}
+	return NavSysConfigClass;
+}
+
 UNavigationSystemBase* UNavigationSystemConfig::CreateAndConfigureNavigationSystem(UWorld& World) const
 {
 	UNavigationSystemBase* NavSys = nullptr;
@@ -22,7 +33,7 @@ UNavigationSystemBase* UNavigationSystemConfig::CreateAndConfigureNavigationSyst
 	NavigationSystemClass.TryLoad();
 	TSubclassOf<UNavigationSystemBase> NavSysClass = NavigationSystemClass.ResolveClass();
 	
-	if (NavSysClass && (bCreateOnClient == true || World.GetNetMode() != NM_Client))
+	if (NavSysClass)
 	{
 		NavSys = NewObject<UNavigationSystemBase>(&World, NavSysClass);
 		if (NavSys)
@@ -54,10 +65,6 @@ void UNavigationSystemConfig::PostEditChangeProperty(FPropertyChangedEvent& Prop
 			{
 				NavigationSystemClass.TryLoad();
 				TSubclassOf<UNavigationSystemBase> NavSysClass = NavigationSystemClass.ResolveClass();
-				if (NavSysClass)
-				{
-					//bCreateOnClient = NavSysClass->GetDefaultObject<UNavigationSystem>()->ShouldAllowClientSideNavigation();
-				}
 			}
 		}
 	}
diff --git a/Engine/Source/Runtime/Engine/Private/WorldSettings.cpp b/Engine/Source/Runtime/Engine/Private/WorldSettings.cpp
--- a/Engine/Source/Runtime/Engine/Private/WorldSettings.cpp
+++ b/Engine/Source/Runtime/Engine/Private/WorldSettings.cpp
@@ -105,7 +105,11 @@ void AWorldSettings::PostInitProperties()
 	Super::PostInitProperties();
 	if (HasAnyFlags(RF_NeedLoad|RF_WasLoaded|RF_ClassDefaultObject) == false)
 	{
-		NavigationSystemConfig = NewObject<UNavigationSystemConfig>(this);
+		TSubclassOf<UNavigationSystemConfig> NavSystemConfigClass = UNavigationSystemConfig::GetDefaultConfigClass();
+		if (*NavSystemConfigClass)
+		{
+			NavigationSystemConfig = NewObject<UNavigationSystemConfig>(this, NavSystemConfigClass);
+		}
 	}
 
 	if (MinGlobalTimeDilation < 0)
@@ -413,7 +417,11 @@ void AWorldSettings::PostLoad()
 		ULevel* Level = GetLevel();
 		if (Level && Level->IsPersistentLevel())
 		{
-			NavigationSystemConfig = NewObject<UNavigationSystemConfig>(this);
+			TSubclassOf<UNavigationSystemConfig> NavSystemConfigClass = UNavigationSystemConfig::GetDefaultConfigClass();
+			if (*NavSystemConfigClass)
+			{
+				NavigationSystemConfig = NewObject<UNavigationSystemConfig>(this, NavSystemConfigClass);
+			}
 			bEnableNavigationSystem = false;
 		}
 	}
diff --git a/Engine/Source/Runtime/NavigationSystem/Private/NavigationSystem.cpp b/Engine/Source/Runtime/NavigationSystem/Private/NavigationSystem.cpp
--- a/Engine/Source/Runtime/NavigationSystem/Private/NavigationSystem.cpp
+++ b/Engine/Source/Runtime/NavigationSystem/Private/NavigationSystem.cpp
@@ -116,14 +116,35 @@ namespace FNavigationSystem
 
 	bool ShouldLoadNavigationOnClient(ANavigationData& NavData)
 	{
-		const UNavigationSystemV1* NavSysCDO = (*GEngine->NavigationSystemClass != nullptr) 
-			? (GEngine->NavigationSystemClass->GetDefaultObject<const UNavigationSystemV1>()) 
-			: (const UNavigationSystemV1*)nullptr;
-		return NavSysCDO && NavSysCDO->ShouldLoadNavigationOnClient(&NavData);
+		const UWorld* World = NavData.GetWorld();
+
+		if (World && World->GetNavigationSystem())
+		{
+			const UNavigationSystemV1* NavSys = Cast<UNavigationSystemV1>(World->GetNavigationSystem());
+			return NavSys && NavSys->ShouldLoadNavigationOnClient(&NavData);
+		}
+		else
+		{
+			const UNavigationSystemV1* NavSysCDO = (*GEngine->NavigationSystemClass != nullptr)
+				? (GEngine->NavigationSystemClass->GetDefaultObject<const UNavigationSystemV1>())
+				: (const UNavigationSystemV1*)nullptr;
+			return NavSysCDO && NavSysCDO->ShouldLoadNavigationOnClient(&NavData);
+		}
 	}
 
 	bool ShouldDiscardSubLevelNavData(ANavigationData& NavData)
 	{
+		const UWorld* World = NavData.GetWorld();
+
+		if (World && World->GetNavigationSystem())
+		{
+			const UNavigationSystemV1* NavSys = Cast<UNavigationSystemV1>(World->GetNavigationSystem());
+			if (NavSys)
+			{
+				return NavSys->ShouldDiscardSubLevelNavData(&NavData);
+			}
+		}
+
 		const UNavigationSystemV1* NavSysCDO = (*GEngine->NavigationSystemClass != nullptr)
 			? (GEngine->NavigationSystemClass->GetDefaultObject<const UNavigationSystemV1>())
 			: (const UNavigationSystemV1*)nullptr;
@@ -446,7 +467,10 @@ void UNavigationSystemV1::UpdateAbstractNavData()
 		FNavDataConfig DummyConfig;
 		DummyConfig.NavigationDataClass = AAbstractNavData::StaticClass();
 		AbstractNavData = CreateNavigationDataInstance(DummyConfig);
-		AbstractNavData->SetFlags(RF_Transient);
+		if (AbstractNavData)
+		{
+			AbstractNavData->SetFlags(RF_Transient);
+		}
 	}
 }
 
@@ -626,6 +650,7 @@ void UNavigationSystemV1::PostEditChangeChainProperty(FPropertyChangedChainEvent
 {
 	static const FName NAME_NavigationDataClass = GET_MEMBER_NAME_CHECKED(FNavDataConfig, NavigationDataClass);
 	static const FName NAME_SupportedAgents = GET_MEMBER_NAME_CHECKED(UNavigationSystemV1, SupportedAgents);
+	static const FName NAME_AllowClientSideNavigation = GET_MEMBER_NAME_CHECKED(UNavigationSystemV1, bAllowClientSideNavigation);
 
 	Super::PostEditChangeChainProperty(PropertyChangedEvent);
 
@@ -643,6 +668,13 @@ void UNavigationSystemV1::PostEditChangeChainProperty(FPropertyChangedChainEvent
 				SaveConfig();
 			}
 		}
+		else if (PropName == NAME_AllowClientSideNavigation && HasAnyFlags(RF_ClassDefaultObject))
+		{
+			for (FObjectIterator It(UNavigationSystemModuleConfig::StaticClass()); It; ++It)
+			{
+				((UNavigationSystemModuleConfig*)*It)->UpdateWithNavSysCDO(*this);
+			}
+		}
 	}
 }
 
@@ -4340,12 +4372,30 @@ UNavigationSystemModuleConfig::UNavigationSystemModuleConfig(const FObjectInitia
 	: Super(ObjectInitializer)
 {
 	const UNavigationSystemV1* NavSysCDO = GetDefault<UNavigationSystemV1>();
-	bAutoSpawnMissingNavData = NavSysCDO && NavSysCDO->bAutoCreateNavigationData;
-	bSpawnNavDataInNavBoundsLevel = NavSysCDO && NavSysCDO->bSpawnNavDataInNavBoundsLevel;
+	if (NavSysCDO)
+	{
+		UpdateWithNavSysCDO(*NavSysCDO);
+	}
+}
+
+void UNavigationSystemModuleConfig::UpdateWithNavSysCDO(const UNavigationSystemV1& NavSysCDO)
+{
+	if (NavigationSystemClass.TryLoad() == NavSysCDO.GetClass())
+	{
+		bStrictlyStatic = NavSysCDO.bStaticRuntimeNavigation;
+		bCreateOnClient = NavSysCDO.bAllowClientSideNavigation;
+		bAutoSpawnMissingNavData = NavSysCDO.bAutoCreateNavigationData;
+		bSpawnNavDataInNavBoundsLevel = NavSysCDO.bSpawnNavDataInNavBoundsLevel;
+	}
 }
 
 UNavigationSystemBase* UNavigationSystemModuleConfig::CreateAndConfigureNavigationSystem(UWorld& World) const
 {
+	if (bCreateOnClient == false && World.GetNetMode() == NM_Client)
+	{
+		return nullptr;
+	}
+	
 	UNavigationSystemV1* NavSysInstance = Cast<UNavigationSystemV1>(Super::CreateAndConfigureNavigationSystem(World));
 	UE_CLOG(NavSysInstance == nullptr, LogNavigation, Error
 		, TEXT("Unable to spawn navsys instance of class %s - unable to cast to UNavigationSystemV1")
@@ -4356,9 +4406,46 @@ UNavigationSystemBase* UNavigationSystemModuleConfig::CreateAndConfigureNavigati
 	{
 		NavSysInstance->bAutoCreateNavigationData = bAutoSpawnMissingNavData;
 		NavSysInstance->bSpawnNavDataInNavBoundsLevel = bSpawnNavDataInNavBoundsLevel;
+		if (bStrictlyStatic)
+		{
+			NavSysInstance->ConfigureAsStatic();
+		}
 	}
 
 	return NavSysInstance;
 }
 
+#if WITH_EDITOR
+void UNavigationSystemModuleConfig::PostEditChangeProperty(FPropertyChangedEvent& PropertyChangedEvent)
+{
+	static const FName NAME_NavigationSystemClass = GET_MEMBER_NAME_CHECKED(UNavigationSystemConfig, NavigationSystemClass);
+
+	Super::PostEditChangeProperty(PropertyChangedEvent);
+
+	if (PropertyChangedEvent.Property)
+	{
+		FName PropName = PropertyChangedEvent.Property->GetFName();
+		if (PropName == NAME_NavigationSystemClass)
+		{
+			if (NavigationSystemClass.IsValid() == false)
+			{
+				NavigationSystemClass = *GEngine->NavigationSystemClass;
+			}
+			else
+			{
+				NavigationSystemClass.TryLoad();
+				TSubclassOf<UNavigationSystemBase> NavSysClass = NavigationSystemClass.ResolveClass();
+				const UNavigationSystemV1* NavSysCDO = *NavSysClass
+					? NavSysClass->GetDefaultObject<UNavigationSystemV1>()
+					: (UNavigationSystemV1*)nullptr;
+				if (NavSysCDO)
+				{
+					UpdateWithNavSysCDO(*NavSysCDO);
+				}
+			}
+		}
+	}
+}
+#endif // WITH_EDITOR
+
 #undef LOCTEXT_NAMESPACE
diff --git a/Engine/Source/Runtime/NavigationSystem/Public/NavigationSystem.h b/Engine/Source/Runtime/NavigationSystem/Public/NavigationSystem.h
--- a/Engine/Source/Runtime/NavigationSystem/Public/NavigationSystem.h
+++ b/Engine/Source/Runtime/NavigationSystem/Public/NavigationSystem.h
@@ -11,7 +11,6 @@
 #include "NavFilters/NavigationQueryFilter.h"
 #include "AI/Navigation/NavigationTypes.h"
 #include "NavigationSystemTypes.h"
-#include "GenericOctreePublic.h"
 #include "NavigationData.h"
 #include "AI/NavigationSystemBase.h"
 #include "Kismet/BlueprintFunctionLibrary.h"
@@ -991,6 +990,13 @@ class NAVIGATIONSYSTEM_API UNavigationSystemModuleConfig : public UNavigationSys
 	GENERATED_BODY()
 
 protected:
+	/** Whether at game runtime we expect any kind of dynamic navigation generation */
+	UPROPERTY(EditAnywhere, Category = Navigation)
+	uint32 bStrictlyStatic : 1;
+
+	UPROPERTY(EditAnywhere, Category = Navigation)
+	uint32 bCreateOnClient : 1;
+
 	UPROPERTY(EditAnywhere, Category = Navigation)
 	uint32 bAutoSpawnMissingNavData : 1;
 
@@ -1000,4 +1006,14 @@ protected:
 public:
 	UNavigationSystemModuleConfig(const FObjectInitializer& ObjectInitializer = FObjectInitializer::Get());
 	virtual UNavigationSystemBase* CreateAndConfigureNavigationSystem(UWorld& World) const override;
+
+#if WITH_EDITOR
+	virtual void PostEditChangeProperty(FPropertyChangedEvent& PropertyChangedEvent) override;
+#endif // WITH_EDITOR
+
+protected:
+#if WITH_EDITOR
+	friend UNavigationSystemV1;
+#endif // WITH_EDITOR
+	void UpdateWithNavSysCDO(const UNavigationSystemV1& NavSysCDO);
 };
