Made sure we pass the Querier over to GetQueryFilter call when we're constructing a meta navigation filter with UNavigationQueryFilter::GetQueryFilter #UE4

diff --git a/Engine/Source/Runtime/NavigationSystem/Private/NavFilters/NavigationQueryFilter.cpp b/Engine/Source/Runtime/NavigationSystem/Private/NavFilters/NavigationQueryFilter.cpp
--- a/Engine/Source/Runtime/NavigationSystem/Private/NavFilters/NavigationQueryFilter.cpp
+++ b/Engine/Source/Runtime/NavigationSystem/Private/NavFilters/NavigationQueryFilter.cpp
@@ -29,7 +29,7 @@ FSharedConstNavQueryFilter UNavigationQueryFilter::GetQueryFilter(const ANavigat
 			check(DefFilterOb);
 			if (DefFilterOb->bIsMetaFilter == false)
 			{
-				return DefFilterOb->GetQueryFilter(NavData, nullptr);
+				return DefFilterOb->GetQueryFilter(NavData, Querier);
 			}
 		}
 	}
