Cleaning up AI and navigation-related [REDACTED_HANDLE] #UE4

Most of them related to recent nav sys refactor, but some are just old todos. Changes include:
- plugged in navmesh vislogging
- implemented querying and stopping movement from the engine code (via FNavigationSystem's binds)
- fixed UNavigationSystemV1::ConfigureAsStatic
- implemented NavigationSystem's own nav config class (with additional options available)

Also added some helpful comments in couple of places.

diff --git a/Engine/Source/Runtime/AIModule/Classes/AIController.h b/Engine/Source/Runtime/AIModule/Classes/AIController.h
--- a/Engine/Source/Runtime/AIModule/Classes/AIController.h
+++ b/Engine/Source/Runtime/AIModule/Classes/AIController.h
@@ -441,6 +441,11 @@ public:
 	UBrainComponent* GetBrainComponent() const { return BrainComponent; }
 	const UBlackboardComponent* GetBlackboardComponent() const { return Blackboard; }
 	UBlackboardComponent* GetBlackboardComponent() { return Blackboard; }
+
+	/** Note that his function does not do any pathfollowing state transfer. 
+	 *	Intended to be called as part of initialization/setup process */
+	UFUNCTION(BlueprintCallable, Category = "AI|Navigation")
+	void SetPathFollowingComponent(UPathFollowingComponent* NewPFComponent);
 };
 
 //----------------------------------------------------------------------//
diff --git a/Engine/Source/Runtime/AIModule/Private/AIController.cpp b/Engine/Source/Runtime/AIModule/Private/AIController.cpp
--- a/Engine/Source/Runtime/AIModule/Private/AIController.cpp
+++ b/Engine/Source/Runtime/AIModule/Private/AIController.cpp
@@ -22,16 +22,8 @@
 #include "BehaviorTree/BehaviorTree.h"
 #include "GameplayTasksComponent.h"
 #include "Tasks/GameplayTask_ClaimResource.h"
+#include "NetworkingDistanceConstants.h"
 
-// mz@todo these need to be removed, legacy code
-#define CLOSEPROXIMITY					500.f
-#define NEARSIGHTTHRESHOLD				2000.f
-#define MEDSIGHTTHRESHOLD				3162.f
-#define FARSIGHTTHRESHOLD				8000.f
-#define CLOSEPROXIMITYSQUARED			(CLOSEPROXIMITY*CLOSEPROXIMITY)
-#define NEARSIGHTTHRESHOLDSQUARED		(NEARSIGHTTHRESHOLD*NEARSIGHTTHRESHOLD)
-#define MEDSIGHTTHRESHOLDSQUARED		(MEDSIGHTTHRESHOLD*MEDSIGHTTHRESHOLD)
-#define FARSIGHTTHRESHOLDSQUARED		(FARSIGHTTHRESHOLD*FARSIGHTTHRESHOLD)
 
 //----------------------------------------------------------------------//
 // AAIController
@@ -1069,6 +1061,17 @@ void AAIController::OnGameplayTaskResourcesClaimed(FGameplayResourceSet NewlyCla
 	}
 }
 
+void AAIController::SetPathFollowingComponent(UPathFollowingComponent* NewPFComponent)
+{ 
+	PathFollowingComponent = NewPFComponent; 
+#if ENABLE_VISUAL_LOG
+	if (NewPFComponent)
+	{
+		REDIRECT_OBJECT_TO_VLOG(NewPFComponent, this);
+	}
+#endif // ENABLE_VISUAL_LOG
+}
+
 //----------------------------------------------------------------------//
 // IGenericTeamAgentInterface
 //----------------------------------------------------------------------//
diff --git a/Engine/Source/Runtime/AIModule/Private/Navigation/CrowdFollowingComponent.cpp b/Engine/Source/Runtime/AIModule/Private/Navigation/CrowdFollowingComponent.cpp
--- a/Engine/Source/Runtime/AIModule/Private/Navigation/CrowdFollowingComponent.cpp
+++ b/Engine/Source/Runtime/AIModule/Private/Navigation/CrowdFollowingComponent.cpp
@@ -734,14 +734,14 @@ void LogPathPartHelper(AActor* LogOwner, FNavMeshPath* NavMeshPath, int32 StartI
 			CenterPt /= Verts.Num();
 
 			const UNavArea* DefArea = AreaClass ? ((UClass*)AreaClass)->GetDefaultObject<UNavArea>() : NULL;
-			const FColor PolygonColor = AreaClass != UNavigationSystemV1::GetDefaultWalkableArea() ? (DefArea ? DefArea->DrawColor : NavMesh->GetConfig().Color) : FColorList::LightSteelBlue;
+			const FColor PolygonColor = AreaClass != FNavigationSystem::GetDefaultWalkableArea() ? (DefArea ? DefArea->DrawColor : NavMesh->GetConfig().Color) : FColorList::LightSteelBlue;
 
 			CorridorPoly.SetColor(PolygonColor.WithAlpha(100));
 			CorridorPoly.Points.Reset();
 			CorridorPoly.Points.Append(Verts);
 			Snapshot->ElementsToDraw.Add(CorridorPoly);
 
-			if (AreaClass && AreaClass != UNavigationSystemV1::GetDefaultWalkableArea())
+			if (AreaClass && AreaClass != FNavigationSystem::GetDefaultWalkableArea())
 			{
 				FVisualLogShapeElement AreaMarkElem(EVisualLoggerShapeElement::Segment);
 				AreaMarkElem.SetColor(FColorList::Orange.WithAlpha(100));
diff --git a/Engine/Source/Runtime/AIModule/Private/Navigation/MetaNavMeshPath.cpp b/Engine/Source/Runtime/AIModule/Private/Navigation/MetaNavMeshPath.cpp
--- a/Engine/Source/Runtime/AIModule/Private/Navigation/MetaNavMeshPath.cpp
+++ b/Engine/Source/Runtime/AIModule/Private/Navigation/MetaNavMeshPath.cpp
@@ -222,7 +222,7 @@ float FMetaNavMeshPath::GetLengthFromPosition(FVector SegmentStart, uint32 NextP
 float FMetaNavMeshPath::GetCostFromIndex(int32 PathPointIndex) const
 {
 	// return approximation of full path * default cost, there's not enough data to give accurate value
-	const UNavArea* DefaultAreaOb = UNavigationSystemV1::GetDefaultWalkableArea().GetDefaultObject();
+	const UNavArea* DefaultAreaOb = static_cast<const UNavArea*>(FNavigationSystem::GetDefaultWalkableArea().GetDefaultObject());
 	const float DefaultAreaCost = DefaultAreaOb ? DefaultAreaOb->DefaultCost : 1.0f;
 	return ApproximateLength * DefaultAreaCost;
 }
diff --git a/Engine/Source/Runtime/Engine/Classes/AI/Navigation/NavigationTypes.h b/Engine/Source/Runtime/Engine/Classes/AI/Navigation/NavigationTypes.h
--- a/Engine/Source/Runtime/Engine/Classes/AI/Navigation/NavigationTypes.h
+++ b/Engine/Source/Runtime/Engine/Classes/AI/Navigation/NavigationTypes.h
@@ -558,7 +558,8 @@ struct FNavigationProjectionWork
 	// if set, projection function scoring will be biased for 2D work (e.g. in case of navmesh, findNearestPoly2D)
 	uint32 bHintProjection2D : 1;
 
-	// result of projection function
+	// result of projection function. 'true' means nav projection was successful
+	// and OutLocation contains Point projected to the nav data's surface
 	uint32 bResult : 1;
 
 	// if set, data in this structure is valid
diff --git a/Engine/Source/Runtime/Engine/Classes/AI/NavigationSystemBase.h b/Engine/Source/Runtime/Engine/Classes/AI/NavigationSystemBase.h
--- a/Engine/Source/Runtime/Engine/Classes/AI/NavigationSystemBase.h
+++ b/Engine/Source/Runtime/Engine/Classes/AI/NavigationSystemBase.h
@@ -19,6 +19,7 @@ class IPathFollowingAgentInterface;
 class AWorldSettings;
 class ULevel;
 class AController;
+class UNavAreaBase;
 
 ENGINE_API DECLARE_LOG_CATEGORY_EXTERN(LogNavigation, Warning, All);
 
@@ -145,6 +146,7 @@ namespace FNavigationSystem
 	ENGINE_API const FNavDataConfig& GetDefaultSupportedAgent();
 
 	ENGINE_API TSubclassOf<UNavAreaBase> GetDefaultWalkableArea();
+	ENGINE_API TSubclassOf<UNavAreaBase> GetDefaultObstacleArea();
 
 	ENGINE_API const FTransform& GetCoordTransformTo(const ENavigationCoordSystem::Type CoordType);
 	ENGINE_API const FTransform& GetCoordTransformFrom(const ENavigationCoordSystem::Type CoordType);
@@ -183,7 +185,7 @@ namespace FNavigationSystem
 	DECLARE_DELEGATE_OneParam(FLevelBasedSignature, ULevel& /*Level*/);
 	DECLARE_DELEGATE_OneParam(FControllerBasedSignature, const AController& /*Controller*/);
 	DECLARE_DELEGATE_TwoParams(FNavigationAutoUpdateEnableSignature, const bool /*bNewEnable*/, UNavigationSystemBase* /*InNavigationSystem*/);
-	DECLARE_DELEGATE_RetVal_OneParam(bool, FBoolControllerBasedSignature, const AController& /*World*/);
+	DECLARE_DELEGATE_RetVal_OneParam(bool, FBoolControllerBasedSignature, const AController& /*Controller*/);
 	DECLARE_DELEGATE_RetVal_OneParam(bool, FBoolActorComponentBasedSignature, UActorComponent& /*Comp*/);
 	DECLARE_DELEGATE_RetVal(TSubclassOf<UNavAreaBase>, FNavAreaBasedSignature);
 	DECLARE_DELEGATE_RetVal(const FNavDataConfig&, FNavDatConfigBasedSignature);
@@ -209,9 +211,6 @@ public:
 	virtual void CleanUp(const FNavigationSystem::ECleanupMode Mode) PURE_VIRTUAL(UNavigationSystemBase::CleanUp, );
 	virtual void Configure(const UNavigationSystemConfig& Config) PURE_VIRTUAL(UNavigationSystemBase::Configure, );
 
-	// possibly editor-only
-	
-
 	/**
 	*	Called when owner-UWorld initializes actors
 	*/
@@ -237,6 +236,9 @@ protected:
 	
 	static void SetCoordTransformTo(const ENavigationCoordSystem::Type CoordType, const FTransform& Transform);
 	static void SetCoordTransformFrom(const ENavigationCoordSystem::Type CoordType, const FTransform& Transform);
+	static void SetWantsComponentChangeNotifies(const bool bEnable);
+	static void SetDefaultWalkableArea(TSubclassOf<UNavAreaBase> InAreaClass);
+	static void SetDefaultObstacleArea(TSubclassOf<UNavAreaBase> InAreaClass);
 
 	static FNavigationSystem::FActorBasedSignature& UpdateActorDataDelegate();
 	static FNavigationSystem::FActorComponentBasedSignature& UpdateComponentDataDelegate();
@@ -250,11 +252,9 @@ protected:
 	static FNavigationSystem::FActorComponentBasedSignature& OnComponentUnregisteredDelegate();
 	static FNavigationSystem::FActorBasedSignature& RemoveActorDataDelegate();
 	static FNavigationSystem::FBoolActorComponentBasedSignature& HasComponentDataDelegate();
-	static FNavigationSystem::FNavAreaBasedSignature& GetDefaultWalkableAreaDelegate();
 	static FNavigationSystem::FNavDatConfigBasedSignature& GetDefaultSupportedAgentDelegate();
 	static FNavigationSystem::FActorBooleBasedSignature& UpdateActorAndComponentDataDelegate();
 	static FNavigationSystem::FComponentBoundsChangeSignature& OnComponentBoundsChangedDelegate();
-	//static FNavigationSystem::FNavDataForPropsSignature& GetNavDataForPropsDelegate();
 	static FNavigationSystem::FNavDataForActorSignature& GetNavDataForActorDelegate();
 	static FNavigationSystem::FNavDataClassFetchSignature& GetDefaultNavDataClassDelegate();
 	static FNavigationSystem::FWorldBoolBasedSignature& VerifyNavigationRenderingComponentsDelegate();
@@ -317,10 +317,16 @@ public:
 	static bool ShouldUpdateNavOctreeOnComponentChange();
 	DEPRECATED(4.20, "GetDefaultWalkableArea is deprecated. Use FNavigationSystem::GetDefaultWalkableArea instead")
 	static TSubclassOf<UNavAreaBase> GetDefaultWalkableArea();
-	DEPRECATED(4.20, "UNavigationSystem::GetNavDataForProps is deprecated. Use UNavigationSystemV1::GetDefaultObstacleArea instead")
+	DEPRECATED(4.20, "UNavigationSystem::GetNavDataForProps is deprecated. Use FNavigationSystem::GetDefaultObstacleArea instead")
 	static TSubclassOf<UNavAreaBase> GetDefaultObstacleArea();
 	DEPRECATED(4.20, "UNavigationSystem::K2_GetRandomReachablePointInRadius is deprecated. Use UNavigationSystemV1::K2_GetRandomReachablePointInRadius instead")
 	static bool K2_GetRandomReachablePointInRadius(UObject* WorldContextObject, const FVector& Origin, FVector& RandomLocation, float Radius, UObject* NavData = NULL, TSubclassOf<UObject> FilterClass = NULL) { return false; }
+	DEPRECATED(4.20, "UNavigationSystem::SimpleMoveToActor is deprecated. Use UAIBlueprintHelperLibrary::SimpleMoveToActor instead")
+	UFUNCTION(BlueprintCallable, Category = "AI|Navigation", meta = (DisplayName = "SimpleMoveToActor_DEPRECATED", ScriptNoExport, DeprecatedFunction, DeprecationMessage = "SimpleMoveToActor is deprecated. Use AIBlueprintHelperLibrary::SimpleMoveToActor instead"))
+	static void SimpleMoveToActor(AController* Controller, const AActor* Goal) {}
+	DEPRECATED(4.20, "UNavigationSystem::SimpleMoveToLocation is deprecated. Use UAIBlueprintHelperLibrary::SimpleMoveToLocation instead")
+	UFUNCTION(BlueprintCallable, Category = "AI|Navigation", meta = (DisplayName = "SimpleMoveToLocation_DEPRECATED", ScriptNoExport, DeprecatedFunction, DeprecationMessage = "SimpleMoveToLocation is deprecated. Use AIBlueprintHelperLibrary::SimpleMoveToLocation instead"))
+	static void SimpleMoveToLocation(AController* Controller, const FVector& Goal) {}
 
 	DEPRECATED(4.20, "UpdateNavOctreeElementBounds is deprecated. Use FNavigationSystem::OnComponentBoundsChanged instead")
 	bool UpdateNavOctreeElementBounds(UActorComponent* Comp, const FBox& NewBounds, const FBox& DirtyArea);
diff --git a/Engine/Source/Runtime/Engine/Private/AI/NavigationSystemBase.cpp b/Engine/Source/Runtime/Engine/Private/AI/NavigationSystemBase.cpp
--- a/Engine/Source/Runtime/Engine/Private/AI/NavigationSystemBase.cpp
+++ b/Engine/Source/Runtime/Engine/Private/AI/NavigationSystemBase.cpp
@@ -66,7 +66,7 @@ namespace FNavigationSystem
 		static FNavDataConfig FallbackSupportedAgent;
 		return FallbackSupportedAgent; 
 	}
-	// mz@todo plug this into the "configure navigation system as static" mechanism
+
 	bool bWantsComponentChangeNotifies = true;
 	
 	class FDelegates
@@ -86,7 +86,6 @@ namespace FNavigationSystem
 		FControllerBasedSignature StopMovement;
 		FBoolControllerBasedSignature IsFollowingAPath;
 		FBoolActorComponentBasedSignature HasComponentData;
-		FNavAreaBasedSignature GetDefaultWalkableArea;
 		FNavDatConfigBasedSignature GetDefaultSupportedAgent;
 		FActorBooleBasedSignature UpdateActorAndComponentData;
 		FComponentBoundsChangeSignature OnComponentBoundsChanged;
@@ -120,7 +119,6 @@ namespace FNavigationSystem
 			StopMovement.BindLambda([](const AController&) {});
 			IsFollowingAPath.BindLambda([](const AController&) { return false; });
 			HasComponentData.BindLambda([](UActorComponent&) { return false; });
-			GetDefaultWalkableArea.BindLambda([]() { return UNavAreaBase::StaticClass(); });
 			GetDefaultSupportedAgent.BindStatic(&GetFallbackSupportedAgent);
 			UpdateActorAndComponentData.BindLambda([](AActor&, bool) {});
 			OnComponentBoundsChanged.BindLambda([](UActorComponent&, const FBox&, const FBox&) {});
@@ -157,7 +155,12 @@ namespace FNavigationSystem
 	void RemoveActorData(AActor& Actor) { Delegates.RemoveActorData.Execute(Actor); }
 	bool HasComponentData(UActorComponent& Comp) { return Delegates.HasComponentData.Execute(Comp);	}
 	const FNavDataConfig& GetDefaultSupportedAgent() { return Delegates.GetDefaultSupportedAgent.Execute(); }
-	TSubclassOf<UNavAreaBase> GetDefaultWalkableArea() { return Delegates.GetDefaultWalkableArea.Execute(); }
+
+
+	TSubclassOf<UNavAreaBase> DefaultWalkableArea; 
+	TSubclassOf<UNavAreaBase> DefaultObstacleArea;
+	TSubclassOf<UNavAreaBase> GetDefaultWalkableArea() { return DefaultWalkableArea; }
+	TSubclassOf<UNavAreaBase> GetDefaultObstacleArea() { return DefaultObstacleArea; }
 		
 	bool WantsComponentChangeNotifies()
 	{
@@ -322,6 +325,22 @@ void UNavigationSystemBase::SetCoordTransformFrom(const ENavigationCoordSystem::
 	FNavigationSystem::CoordTypeTransformsTo[uint8(CoordType)] = Transform;
 }
 
+void UNavigationSystemBase::SetWantsComponentChangeNotifies(const bool bEnable)
+{
+	FNavigationSystem::bWantsComponentChangeNotifies = bEnable;
+}
+
+void UNavigationSystemBase::SetDefaultWalkableArea(TSubclassOf<UNavAreaBase> InAreaClass)
+{
+	FNavigationSystem::DefaultWalkableArea = InAreaClass;
+}
+
+void UNavigationSystemBase::SetDefaultObstacleArea(TSubclassOf<UNavAreaBase> InAreaClass)
+{
+	FNavigationSystem::DefaultObstacleArea = InAreaClass;
+}
+
+
 FNavigationSystem::FActorBasedSignature& UNavigationSystemBase::UpdateActorDataDelegate() { return FNavigationSystem::Delegates.UpdateActorData; }
 FNavigationSystem::FActorComponentBasedSignature& UNavigationSystemBase::UpdateComponentDataDelegate() { return FNavigationSystem::Delegates.UpdateComponentData; }
 FNavigationSystem::FSceneComponentBasedSignature& UNavigationSystemBase::UpdateComponentDataAfterMoveDelegate() { return FNavigationSystem::Delegates.UpdateComponentDataAfterMove; }
@@ -334,11 +353,9 @@ FNavigationSystem::FActorComponentBasedSignature& UNavigationSystemBase::OnCompo
 FNavigationSystem::FActorComponentBasedSignature& UNavigationSystemBase::OnComponentUnregisteredDelegate() { return FNavigationSystem::Delegates.OnComponentUnregistered; }
 FNavigationSystem::FActorBasedSignature& UNavigationSystemBase::RemoveActorDataDelegate() { return FNavigationSystem::Delegates.RemoveActorData; }
 FNavigationSystem::FBoolActorComponentBasedSignature& UNavigationSystemBase::HasComponentDataDelegate() { return FNavigationSystem::Delegates.HasComponentData; }
-FNavigationSystem::FNavAreaBasedSignature& UNavigationSystemBase::GetDefaultWalkableAreaDelegate() { return FNavigationSystem::Delegates.GetDefaultWalkableArea; }
 FNavigationSystem::FNavDatConfigBasedSignature& UNavigationSystemBase::GetDefaultSupportedAgentDelegate() { return FNavigationSystem::Delegates.GetDefaultSupportedAgent; }
 FNavigationSystem::FActorBooleBasedSignature& UNavigationSystemBase::UpdateActorAndComponentDataDelegate() { return FNavigationSystem::Delegates.UpdateActorAndComponentData; }
 FNavigationSystem::FComponentBoundsChangeSignature& UNavigationSystemBase::OnComponentBoundsChangedDelegate() { return FNavigationSystem::Delegates.OnComponentBoundsChanged; }
-//FNavigationSystem::FNavDataForPropsSignature& UNavigationSystemBase::GetNavDataForPropsDelegate() { return FNavigationSystem::Delegates.GetNavDataForProps; }
 FNavigationSystem::FNavDataForActorSignature& UNavigationSystemBase::GetNavDataForActorDelegate() { return FNavigationSystem::Delegates.GetNavDataForActor; }
 FNavigationSystem::FNavDataClassFetchSignature& UNavigationSystemBase::GetDefaultNavDataClassDelegate() { return FNavigationSystem::Delegates.GetDefaultNavDataClass; }
 FNavigationSystem::FWorldBoolBasedSignature& UNavigationSystemBase::VerifyNavigationRenderingComponentsDelegate() { return FNavigationSystem::Delegates.VerifyNavigationRenderingComponents; }
diff --git a/Engine/Source/Runtime/Engine/Private/Components/ShapeComponent.cpp b/Engine/Source/Runtime/Engine/Private/Components/ShapeComponent.cpp
--- a/Engine/Source/Runtime/Engine/Private/Components/ShapeComponent.cpp
+++ b/Engine/Source/Runtime/Engine/Private/Components/ShapeComponent.cpp
@@ -3,6 +3,7 @@
 
 #include "Components/ShapeComponent.h"
 #include "AI/Navigation/NavAreaBase.h"
+#include "AI/NavigationSystemBase.h"
 #include "PhysicsEngine/BoxElem.h"
 #include "PhysicsEngine/SphereElem.h"
 #include "PhysicsEngine/SphylElem.h"
@@ -31,8 +32,7 @@ UShapeComponent::UShapeComponent(const FObjectInitializer& ObjectInitializer)
 	bHasCustomNavigableGeometry = EHasCustomNavigableGeometry::Yes;
 	bCanEverAffectNavigation = true;
 	bDynamicObstacle = false;
-	// mz@todo
-	//AreaClass = UNavArea_Obstacle::StaticClass();
+	AreaClass = FNavigationSystem::GetDefaultObstacleArea();
 
 	// Ignore streaming updates since GetUsedMaterials() is not implemented.
 	bIgnoreStreamingManagerUpdate = true;
diff --git a/Engine/Source/Runtime/Engine/Private/VisualLogger/VisualLogger.cpp b/Engine/Source/Runtime/Engine/Private/VisualLogger/VisualLogger.cpp
--- a/Engine/Source/Runtime/Engine/Private/VisualLogger/VisualLogger.cpp
+++ b/Engine/Source/Runtime/Engine/Private/VisualLogger/VisualLogger.cpp
@@ -338,7 +338,16 @@ void FVisualLogger::NavigationDataDump(const UObject* Object, const FLogCategory
 void FVisualLogger::NavigationDataDump(const UObject* Object, const FName& CategoryName, const ELogVerbosity::Type Verbosity, const FBox& Box)
 {
 	SCOPE_CYCLE_COUNTER(STAT_VisualLog);
-	NavigationDataDumpDelegate.Broadcast(Object, CategoryName, Verbosity, Box);
+
+	UWorld* World = nullptr;
+	FVisualLogEntry* CurrentEntry = nullptr;
+	if (CheckVisualLogInputInternal(Object, CategoryName, Verbosity, &World, &CurrentEntry) == false 
+		|| CurrentEntry == nullptr)
+	{
+		return;
+	}
+	
+	NavigationDataDumpDelegate.Broadcast(Object, CategoryName, Verbosity, Box, *World, *CurrentEntry);
 }
 
 
diff --git a/Engine/Source/Runtime/Engine/Public/VisualLogger/VisualLogger.h b/Engine/Source/Runtime/Engine/Public/VisualLogger/VisualLogger.h
--- a/Engine/Source/Runtime/Engine/Public/VisualLogger/VisualLogger.h
+++ b/Engine/Source/Runtime/Engine/Public/VisualLogger/VisualLogger.h
@@ -370,7 +370,7 @@ public:
 	static void NavigationDataDump(const UObject* LogOwner, const FLogCategoryBase& Category, const ELogVerbosity::Type Verbosity, const FBox& Box);
 	static void NavigationDataDump(const UObject* LogOwner, const FName& CategoryName, const ELogVerbosity::Type Verbosity, const FBox& Box);
 
-	DECLARE_MULTICAST_DELEGATE_FourParams(FNavigationDataDump, const UObject* /*LogOwner*/, const FName& /*CategoryName*/, const ELogVerbosity::Type /*Verbosity*/, const FBox& /*Box*/);
+	DECLARE_MULTICAST_DELEGATE_SixParams(FNavigationDataDump, const UObject* /*Object*/, const FName& /*CategoryName*/, const ELogVerbosity::Type /*Verbosity*/, const FBox& /*Box*/, const UWorld& /*World*/, FVisualLogEntry& /*CurrentEntry*/);
 	static FNavigationDataDump NavigationDataDumpDelegate;
 
 	/** Log events */
diff --git a/Engine/Source/Runtime/NavigationSystem/Private/NavCollision.cpp b/Engine/Source/Runtime/NavigationSystem/Private/NavCollision.cpp
--- a/Engine/Source/Runtime/NavigationSystem/Private/NavCollision.cpp
+++ b/Engine/Source/Runtime/NavigationSystem/Private/NavCollision.cpp
@@ -270,7 +270,7 @@ void UNavCollision::GetNavigationModifier(FCompositeNavModifier& Modifier, const
 {
 	QUICK_SCOPE_CYCLE_COUNTER(STAT_NavCollision_GetNavigationModifier);
 
-	const TSubclassOf<UNavArea> UseAreaClass = AreaClass ? AreaClass : UNavigationSystemV1::GetDefaultObstacleArea();
+	const TSubclassOf<UNavArea> UseAreaClass = AreaClass ? AreaClass : (const TSubclassOf<UNavArea>)(FNavigationSystem::GetDefaultObstacleArea());
 
 	// rebuild collision data if needed
 	if (!bHasConvexGeometry)
diff --git a/Engine/Source/Runtime/NavigationSystem/Private/NavMesh/NavMeshPath.cpp b/Engine/Source/Runtime/NavigationSystem/Private/NavMesh/NavMeshPath.cpp
--- a/Engine/Source/Runtime/NavigationSystem/Private/NavMesh/NavMeshPath.cpp
+++ b/Engine/Source/Runtime/NavigationSystem/Private/NavMesh/NavMeshPath.cpp
@@ -816,7 +816,7 @@ void FNavMeshPath::DescribeSelfToVisLog(FVisualLogEntry* Snapshot) const
 		}
 
 		const UNavArea* DefArea = AreaClass ? ((UClass*)AreaClass)->GetDefaultObject<UNavArea>() : NULL;
-		TSubclassOf<UNavAreaBase> DefaultWalkableArea = FNavigationSystem::GetDefaultWalkableArea();
+		const TSubclassOf<UNavAreaBase> DefaultWalkableArea = FNavigationSystem::GetDefaultWalkableArea();
 		const FColor PolygonColor = AreaClass != DefaultWalkableArea ? (DefArea ? DefArea->DrawColor : NavMesh->GetConfig().Color) : FColorList::Cyan;
 
 		CorridorPoly.SetColor(PolygonColor.WithAlpha(100));
diff --git a/Engine/Source/Runtime/NavigationSystem/Private/NavMesh/RecastNavMesh.cpp b/Engine/Source/Runtime/NavigationSystem/Private/NavMesh/RecastNavMesh.cpp
--- a/Engine/Source/Runtime/NavigationSystem/Private/NavMesh/RecastNavMesh.cpp
+++ b/Engine/Source/Runtime/NavigationSystem/Private/NavMesh/RecastNavMesh.cpp
@@ -522,7 +522,7 @@ void ARecastNavMesh::OnNavAreaChanged()
 
 int32 ARecastNavMesh::GetNewAreaID(const UClass* AreaClass) const
 {
-	if (AreaClass == UNavigationSystemV1::GetDefaultWalkableArea())
+	if (AreaClass == FNavigationSystem::GetDefaultWalkableArea())
 	{
 		return RECAST_DEFAULT_AREA;
 	}
@@ -1756,9 +1756,43 @@ void ARecastNavMesh::OnNavMeshGenerationFinished()
 uint32 ARecastNavMesh::LogMemUsed() const 
 {
 	const uint32 SuperMemUsed = Super::LogMemUsed();
+	const int32 headerSize = dtAlign4(sizeof(dtMeshHeader));
+
 	uint32 MemUsed = 0;
 
-	UE_LOG(LogNavigation, Display, TEXT("%s: ARecastNavMesh: %u\n    self: %d"), *GetName(), MemUsed, sizeof(ARecastNavMesh));	
+	if (RecastNavMeshImpl && RecastNavMeshImpl->DetourNavMesh)
+	{
+		const dtNavMesh* const ConstNavMesh = RecastNavMeshImpl->DetourNavMesh;
+
+		for (int TileIndex = 0; TileIndex < RecastNavMeshImpl->DetourNavMesh->getMaxTiles(); ++TileIndex)
+		{
+			const dtMeshTile* Tile = ConstNavMesh->getTile(TileIndex);
+			if (Tile && Tile->header)
+			{
+				dtMeshHeader* const H = (dtMeshHeader*)(Tile->header);
+				const int32 vertsSize = dtAlign4(sizeof(float) * 3 * H->vertCount);
+				const int32 polysSize = dtAlign4(sizeof(dtPoly) * H->polyCount);
+				const int32 linksSize = dtAlign4(sizeof(dtLink) * H->maxLinkCount);
+				const int32 detailMeshesSize = dtAlign4(sizeof(dtPolyDetail) * H->detailMeshCount);
+				const int32 detailVertsSize = dtAlign4(sizeof(float) * 3 * H->detailVertCount);
+				const int32 detailTrisSize = dtAlign4(sizeof(unsigned char) * 4 * H->detailTriCount);
+				const int32 bvTreeSize = dtAlign4(sizeof(dtBVNode) * H->bvNodeCount);
+				const int32 offMeshConsSize = dtAlign4(sizeof(dtOffMeshConnection) * H->offMeshConCount);
+				const int32 offMeshSegsSize = dtAlign4(sizeof(dtOffMeshSegmentConnection) * H->offMeshSegConCount);
+				const int32 clusterSize = dtAlign4(sizeof(dtCluster) * H->clusterCount);
+				const int32 polyClustersSize = dtAlign4(sizeof(unsigned short) * H->detailMeshCount);
+
+				const int32 TileDataSize = headerSize + vertsSize + polysSize + linksSize +
+					detailMeshesSize + detailVertsSize + detailTrisSize +
+					bvTreeSize + offMeshConsSize + offMeshSegsSize +
+					clusterSize + polyClustersSize;
+
+				MemUsed += TileDataSize;
+			}
+		}
+	}
+
+	UE_LOG(LogNavigation, Warning, TEXT("%s: ARecastNavMesh: %u\n    self: %d"), *GetName(), MemUsed, sizeof(ARecastNavMesh));	
 
 	return MemUsed + SuperMemUsed;
 }
@@ -2196,7 +2230,9 @@ void ARecastNavMesh::PostEditChangeProperty(FPropertyChangedEvent& PropertyChang
 				UpdatePolyRefBitsPreview();
 			}
 
-			if (!HasAnyFlags(RF_ClassDefaultObject) && UNavigationSystemV1::GetIsNavigationAutoUpdateEnabled())
+			if (!HasAnyFlags(RF_ClassDefaultObject) 
+				&& UNavigationSystemV1::GetIsNavigationAutoUpdateEnabled()
+				&& PropName != GET_MEMBER_NAME_CHECKED(ARecastNavMesh, MaxSimultaneousTileGenerationJobsCount))
 			{
 				RebuildAll();
 			}
diff --git a/Engine/Source/Runtime/NavigationSystem/Private/NavMesh/RecastNavMeshGenerator.cpp b/Engine/Source/Runtime/NavigationSystem/Private/NavMesh/RecastNavMeshGenerator.cpp
--- a/Engine/Source/Runtime/NavigationSystem/Private/NavMesh/RecastNavMeshGenerator.cpp
+++ b/Engine/Source/Runtime/NavigationSystem/Private/NavMesh/RecastNavMeshGenerator.cpp
@@ -4472,7 +4472,7 @@ void FRecastNavMeshGenerator::GrabDebugSnapshot(struct FVisualLogEntry* Snapshot
 						const uint8 AreaId = NavData->GetAreaID(AreaMods[i].GetAreaClass());
 						const UClass* AreaClass = NavData->GetAreaClass(AreaId);
 						const UNavArea* DefArea = AreaClass ? ((UClass*)AreaClass)->GetDefaultObject<UNavArea>() : NULL;
-						const FColor PolygonColor = AreaClass != UNavigationSystemV1::GetDefaultWalkableArea() ? (DefArea ? DefArea->DrawColor : NavData->GetConfig().Color) : FColorList::Cyan;
+						const FColor PolygonColor = AreaClass != FNavigationSystem::GetDefaultWalkableArea() ? (DefArea ? DefArea->DrawColor : NavData->GetConfig().Color) : FColorList::Cyan;
 
 						if (AreaMods[i].GetShapeType() == ENavigationShapeType::Box)
 						{
diff --git a/Engine/Source/Runtime/NavigationSystem/Private/NavModifierVolume.cpp b/Engine/Source/Runtime/NavigationSystem/Private/NavModifierVolume.cpp
--- a/Engine/Source/Runtime/NavigationSystem/Private/NavModifierVolume.cpp
+++ b/Engine/Source/Runtime/NavigationSystem/Private/NavModifierVolume.cpp
@@ -26,7 +26,7 @@ ANavModifierVolume::ANavModifierVolume(const FObjectInitializer& ObjectInitializ
 
 void ANavModifierVolume::GetNavigationData(FNavigationRelevantData& Data) const
 {
-	if (Brush && AreaClass && AreaClass != UNavigationSystemV1::GetDefaultWalkableArea())
+	if (Brush && AreaClass && AreaClass != FNavigationSystem::GetDefaultWalkableArea())
 	{
 		FAreaNavModifier AreaMod(GetBrushComponent(), AreaClass);
 		Data.Modifiers.Add(AreaMod);
diff --git a/Engine/Source/Runtime/NavigationSystem/Private/NavigationData.cpp b/Engine/Source/Runtime/NavigationSystem/Private/NavigationData.cpp
--- a/Engine/Source/Runtime/NavigationSystem/Private/NavigationData.cpp
+++ b/Engine/Source/Runtime/NavigationSystem/Private/NavigationData.cpp
@@ -7,8 +7,6 @@
 #include "AI/NavDataGenerator.h"
 #include "NavigationSystem.h"
 #include "Engine/Engine.h"
-// mz@todo why was this needed?
-//#include "AITypes.h"
 #include "NavAreas/NavArea.h"
 #include "AI/Navigation/NavAreaBase.h"
 #include "VisualLogger/VisualLogger.h"
diff --git a/Engine/Source/Runtime/NavigationSystem/Private/NavigationOctree.cpp b/Engine/Source/Runtime/NavigationSystem/Private/NavigationOctree.cpp
--- a/Engine/Source/Runtime/NavigationSystem/Private/NavigationOctree.cpp
+++ b/Engine/Source/Runtime/NavigationSystem/Private/NavigationOctree.cpp
@@ -175,7 +175,7 @@ void FNavigationOctree::UpdateNode(const FOctreeElementId& Id, const FBox& NewBo
 
 void FNavigationOctree::RemoveNode(const FOctreeElementId& Id)
 {
-	FNavigationOctreeElement& Element = GetElementById(Id);
+	const FNavigationOctreeElement& Element = GetElementById(Id);
 	const int32 ElementMemory = Element.GetAllocatedSize();
 	NodesMemory -= ElementMemory;
 	DEC_MEMORY_STAT_BY(STAT_Navigation_CollisionTreeMemory, ElementMemory);
@@ -205,6 +205,10 @@ void FNavigationOctreeSemantics::SetElementId(const FNavigationOctreeElement& El
 {
 	UWorld* World = NULL;
 	UObject* ElementOwner = Element.GetOwner();
+	if (ElementOwner == nullptr)
+	{
+		return;
+	}
 
 	if (AActor* Actor = Cast<AActor>(ElementOwner))
 	{
@@ -222,6 +226,6 @@ void FNavigationOctreeSemantics::SetElementId(const FNavigationOctreeElement& El
 	UNavigationSystemV1* NavSys = FNavigationSystem::GetCurrent<UNavigationSystemV1>(World);
 	if (NavSys)
 	{
-		NavSys->SetObjectsNavOctreeId(ElementOwner, Id);
+		NavSys->SetObjectsNavOctreeId(*ElementOwner, Id);
 	}
 }
diff --git a/Engine/Source/Runtime/NavigationSystem/Private/NavigationSystem.cpp b/Engine/Source/Runtime/NavigationSystem/Private/NavigationSystem.cpp
--- a/Engine/Source/Runtime/NavigationSystem/Private/NavigationSystem.cpp
+++ b/Engine/Source/Runtime/NavigationSystem/Private/NavigationSystem.cpp
@@ -180,8 +180,6 @@ namespace
 bool UNavigationSystemV1::bNavigationAutoUpdateEnabled = true;
 TMap<INavLinkCustomInterface*, FWeakObjectPtr> UNavigationSystemV1::PendingCustomLinkRegistration;
 FCriticalSection UNavigationSystemV1::CustomLinkRegistrationSection;
-TSubclassOf<UNavArea> UNavigationSystemV1::DefaultWalkableArea = NULL;
-TSubclassOf<UNavArea> UNavigationSystemV1::DefaultObstacleArea = NULL;
 #if !(UE_BUILD_SHIPPING || UE_BUILD_TEST)
 FNavigationSystemExec UNavigationSystemV1::ExecHandler;
 #endif // !UE_BUILD_SHIPPING
@@ -259,9 +257,8 @@ UNavigationSystemV1::UNavigationSystemV1(const FObjectInitializer& ObjectInitial
 			UNavigationSystemBase::RemoveActorDataDelegate().BindLambda([](AActor& Actor) { UNavigationSystemV1::ClearNavOctreeAll(&Actor); });
 			UNavigationSystemBase::HasComponentDataDelegate().BindLambda([](UActorComponent& Comp) {
 				UNavigationSystemV1* NavSys = FNavigationSystem::GetCurrent<UNavigationSystemV1>(Comp.GetWorld());
-				return (NavSys && (NavSys->GetObjectsNavOctreeId(&Comp) || NavSys->HasPendingObjectNavOctreeId(&Comp)));
+				return (NavSys && (NavSys->GetObjectsNavOctreeId(Comp) || NavSys->HasPendingObjectNavOctreeId(&Comp)));
 			});
-			UNavigationSystemBase::GetDefaultWalkableAreaDelegate().BindLambda([]() { return UNavigationSystemV1::GetDefaultWalkableArea(); });
 			UNavigationSystemBase::GetDefaultSupportedAgentDelegate().BindStatic(&UNavigationSystemV1::GetDefaultSupportedAgent);
 			UNavigationSystemBase::UpdateActorAndComponentDataDelegate().BindStatic(&UNavigationSystemV1::UpdateActorAndComponentsInNavOctree);
 			UNavigationSystemBase::OnComponentBoundsChangedDelegate().BindLambda([](UActorComponent& Comp, const FBox& NewBounds, const FBox& DirtyArea) {
@@ -332,6 +329,10 @@ UNavigationSystemV1::UNavigationSystemV1(const FObjectInitializer& ObjectInitial
 				}
 			});
 #endif // WITH_EDITOR
+
+#if ENABLE_VISUAL_LOG
+			FVisualLogger::NavigationDataDumpDelegate.AddStatic(&NavigationDataDump);
+#endif // ENABLE_VISUAL_LOG
 		}
 	};
 	static FDelegatesInitializer DelegatesInitializer;
@@ -365,8 +366,8 @@ UNavigationSystemV1::UNavigationSystemV1(const FObjectInitializer& ObjectInitial
 	}
 	else if (GetClass() == UNavigationSystemV1::StaticClass())
 	{
-		DefaultWalkableArea = UNavArea_Default::StaticClass();
-		DefaultObstacleArea = UNavArea_Null::StaticClass();
+		SetDefaultWalkableArea(UNavArea_Default::StaticClass());
+		SetDefaultObstacleArea(UNavArea_Obstacle::StaticClass());
 		
 		const FTransform RecastToUnrealTransfrom(Recast2UnrealMatrix());
 		SetCoordTransformFrom(ENavigationCoordSystem::Recast, RecastToUnrealTransfrom);
@@ -401,6 +402,7 @@ UNavigationSystemV1::~UNavigationSystemV1()
 void UNavigationSystemV1::ConfigureAsStatic()
 {
 	bStaticRuntimeNavigation = true;
+	SetWantsComponentChangeNotifies(false);
 }
 
 void UNavigationSystemV1::SetUpdateNavOctreeOnComponentChange(bool bNewUpdateOnComponentChange)
@@ -2349,7 +2351,7 @@ FSetElementId UNavigationSystemV1::RegisterNavOctreeElement(UObject* ElementOwne
 		}
 		else
 		{
-			const FOctreeElementId* ElementId = GetObjectsNavOctreeId(ElementOwner);
+			const FOctreeElementId* ElementId = GetObjectsNavOctreeId(*ElementOwner);
 			bCanAdd = (ElementId == NULL);
 		}
 
@@ -2400,7 +2402,7 @@ void UNavigationSystemV1::AddElementToNavOctree(const FNavigationDirtyElement& D
 	{
 		// check if parent node is waiting in queue
 		const FSetElementId ParentRequestId = PendingOctreeUpdates.FindId(FNavigationDirtyElement(NavigationParent));
-		const FOctreeElementId* ParentId = GetObjectsNavOctreeId(NavigationParent);
+		const FOctreeElementId* ParentId = GetObjectsNavOctreeId(*NavigationParent);
 		if (ParentRequestId.IsValidId() && ParentId == NULL)
 		{
 			FNavigationDirtyElement& ParentNode = PendingOctreeUpdates[ParentRequestId];
@@ -2410,7 +2412,7 @@ void UNavigationSystemV1::AddElementToNavOctree(const FNavigationDirtyElement& D
 			ParentNode.bInvalidRequest = true;
 		}
 
-		const FOctreeElementId* UseParentId = ParentId ? ParentId : GetObjectsNavOctreeId(NavigationParent);
+		const FOctreeElementId* UseParentId = ParentId ? ParentId : GetObjectsNavOctreeId(*NavigationParent);
 		if (UseParentId && NavOctree->IsValidElementId(*UseParentId))
 		{
 			UE_LOG(LogNavOctree, Log, TEXT("ADD %s to %s"), *GetNameSafe(ElementOwner), *GetNameSafe(NavigationParent));
@@ -2437,7 +2439,7 @@ void UNavigationSystemV1::AddElementToNavOctree(const FNavigationDirtyElement& D
 	}
 }
 
-bool UNavigationSystemV1::GetNavOctreeElementData(UObject* NodeOwner, int32& DirtyFlags, FBox& DirtyBounds)
+bool UNavigationSystemV1::GetNavOctreeElementData(const UObject& NodeOwner, int32& DirtyFlags, FBox& DirtyBounds)
 {
 	const FOctreeElementId* ElementId = GetObjectsNavOctreeId(NodeOwner);
 	if (ElementId != NULL)
@@ -2475,13 +2477,13 @@ void UNavigationSystemV1::UnregisterNavOctreeElement(UObject* ElementOwner, INav
 		return;
 	}
 
-	const FOctreeElementId* ElementId = GetObjectsNavOctreeId(ElementOwner);
+	const FOctreeElementId* ElementId = GetObjectsNavOctreeId(*ElementOwner);
 	UE_LOG(LogNavOctree, Log, TEXT("UNREG %s %s"), *GetNameSafe(ElementOwner), ElementId ? TEXT("[exists]") : TEXT(""));
 
 	if (ElementId != NULL)
 	{
 		RemoveNavOctreeElementId(*ElementId, UpdateFlags);
-		RemoveObjectsNavOctreeId(ElementOwner);
+		RemoveObjectsNavOctreeId(*ElementOwner);
 	}
 	else
 	{
@@ -2514,9 +2516,9 @@ void UNavigationSystemV1::RemoveNavOctreeElementId(const FOctreeElementId& Eleme
 {
 	if (NavOctree->IsValidElementId(ElementId))
 	{
-		// mark area occupied by given actor as dirty
-		FNavigationOctreeElement& ElementData = NavOctree->GetElementById(ElementId);
+		const FNavigationOctreeElement& ElementData = NavOctree->GetElementById(ElementId);
 		const int32 DirtyFlag = GetDirtyFlagHelper(UpdateFlags, ElementData.Data->GetDirtyFlag());
+		// mark area occupied by given actor as dirty
 		AddDirtyArea(ElementData.Bounds.GetBox(), DirtyFlag);
 		NavOctree->RemoveNode(ElementId);
 	}
@@ -2526,7 +2528,7 @@ const FNavigationRelevantData* UNavigationSystemV1::GetDataForObject(const UObje
 {
 	check(NavOctree.IsValid());
 
-	const FOctreeElementId* OctreeID = GetObjectsNavOctreeId(&Object);
+	const FOctreeElementId* OctreeID = GetObjectsNavOctreeId(Object);
 
 	if (OctreeID != nullptr && OctreeID->IsValidId() == true)
 	{
@@ -2701,11 +2703,15 @@ void UNavigationSystemV1::UpdateNavOctreeElement(UObject* ElementOwner, INavRele
 		UE_LOG(LogNavOctree, Log, TEXT("IGNORE(UpdateNavOctreeElement) %s"), *GetPathNameSafe(ElementOwner));
 		return;
 	}
+	else if (ElementOwner == nullptr)
+	{
+		return;
+	}
 
 	// grab existing octree data
 	FBox CurrentBounds;
 	int32 CurrentFlags;
-	const bool bAlreadyExists = GetNavOctreeElementData(ElementOwner, CurrentFlags, CurrentBounds);
+	const bool bAlreadyExists = GetNavOctreeElementData(*ElementOwner, CurrentFlags, CurrentBounds);
 
 	// don't invalidate pending requests
 	UpdateFlags |= OctreeUpdate_Refresh;
@@ -2786,7 +2792,12 @@ void UNavigationSystemV1::UpdateNavOctreeParentChain(UObject* ElementOwner, bool
 
 bool UNavigationSystemV1::UpdateNavOctreeElementBounds(UActorComponent* Comp, const FBox& NewBounds, const FBox& DirtyArea)
 {
-	const FOctreeElementId* ElementId = GetObjectsNavOctreeId(Comp);
+	if (Comp == nullptr)
+	{
+		return false;
+	}
+
+	const FOctreeElementId* ElementId = GetObjectsNavOctreeId(*Comp);
 	if (ElementId && ElementId->IsValidId())
 	{
 		NavOctree->UpdateNode(*ElementId, NewBounds);
@@ -2794,7 +2805,7 @@ bool UNavigationSystemV1::UpdateNavOctreeElementBounds(UActorComponent* Comp, co
 		// Add dirty area
 		if (DirtyArea.IsValid)
 		{
-			ElementId = GetObjectsNavOctreeId(Comp);
+			ElementId = GetObjectsNavOctreeId(*Comp);
 			if (ElementId && ElementId->IsValidId())
 			{
 				FNavigationOctreeElement& ElementData = NavOctree->GetElementById(*ElementId);
@@ -3540,7 +3551,7 @@ void UNavigationSystemV1::AddLevelCollisionToOctree(ULevel* Level)
 		NavOctree->GetNavGeometryStoringMode() == FNavigationOctree::StoreNavGeometry)
 	{
 		const TArray<FVector>* LevelGeom = Level->GetStaticNavigableGeometry();
-		const FOctreeElementId* ElementId = GetObjectsNavOctreeId(Level);
+		const FOctreeElementId* ElementId = GetObjectsNavOctreeId(*Level);
 
 		if (!ElementId && LevelGeom && LevelGeom->Num() > 0)
 		{
@@ -3564,7 +3575,7 @@ void UNavigationSystemV1::RemoveLevelCollisionFromOctree(ULevel* Level)
 {
 	if (Level && NavOctree.IsValid())
 	{
-		const FOctreeElementId* ElementId = GetObjectsNavOctreeId(Level);
+		const FOctreeElementId* ElementId = GetObjectsNavOctreeId(*Level);
 		UE_LOG(LogNavOctree, Log, TEXT("UNREG %s %s"), *GetNameSafe(Level), ElementId ? TEXT("[exists]") : TEXT(""));
 
 		if (ElementId != NULL)
@@ -3577,7 +3588,7 @@ void UNavigationSystemV1::RemoveLevelCollisionFromOctree(ULevel* Level)
 			}
 
 			NavOctree->RemoveNode(*ElementId);
-			RemoveObjectsNavOctreeId(Level);
+			RemoveObjectsNavOctreeId(*Level);
 		}
 	}
 }
diff --git a/Engine/Source/Runtime/NavigationSystem/Public/NavMesh/RecastQueryFilter.h b/Engine/Source/Runtime/NavigationSystem/Public/NavMesh/RecastQueryFilter.h
--- a/Engine/Source/Runtime/NavigationSystem/Public/NavMesh/RecastQueryFilter.h
+++ b/Engine/Source/Runtime/NavigationSystem/Public/NavMesh/RecastQueryFilter.h
@@ -41,6 +41,11 @@ public:
 
 	/** note that it results in losing all area cost setup. Call it before setting anything else */
 	void SetIsVirtual(bool bIsVirtual);
+
+	//----------------------------------------------------------------------//
+	// @note you might also want to override following functions from dtQueryFilter	
+	// virtual bool passVirtualFilter(const dtPolyRef ref, const dtMeshTile* tile, const dtPoly* poly) const;
+	// virtual float getVirtualCost(const float* pa, const float* pb, const dtPolyRef prevRef, const dtMeshTile* prevTile, const dtPoly* prevPoly, const dtPolyRef curRef, const dtMeshTile* curTile, const dtPoly* curPoly, const dtPolyRef nextRef, const dtMeshTile* nextTile, const dtPoly* nextPoly) const;
 };
 
 struct NAVIGATIONSYSTEM_API FRecastSpeciaLinkFilter : public dtQuerySpecialLinkFilter
diff --git a/Engine/Source/Runtime/NavigationSystem/Public/NavigationData.h b/Engine/Source/Runtime/NavigationSystem/Public/NavigationData.h
--- a/Engine/Source/Runtime/NavigationSystem/Public/NavigationData.h
+++ b/Engine/Source/Runtime/NavigationSystem/Public/NavigationData.h
@@ -510,7 +510,7 @@ protected:
 	uint32 bForceRebuildOnLoad : 1;
 
 	/** If set, navigation data can act as default one in navigation system's queries */
-	UPROPERTY(config, VisibleAnywhere, Category = Runtime, AdvancedDisplay)
+	UPROPERTY(config, EditAnywhere, Category = Runtime, AdvancedDisplay)
 	uint32 bCanBeMainNavData : 1;
 
 	/** If set, navigation data will be spawned in persistent level during rebuild if actor doesn't exist */
diff --git a/Engine/Source/Runtime/NavigationSystem/Public/NavigationSystem.h b/Engine/Source/Runtime/NavigationSystem/Public/NavigationSystem.h
--- a/Engine/Source/Runtime/NavigationSystem/Public/NavigationSystem.h
+++ b/Engine/Source/Runtime/NavigationSystem/Public/NavigationSystem.h
@@ -480,12 +480,6 @@ public:
 
 	const TSet<FNavigationBounds>& GetNavigationBounds() const;
 
-	/** @return default walkable area class */
-	FORCEINLINE static TSubclassOf<UNavArea> GetDefaultWalkableArea() { return DefaultWalkableArea; }
-
-	/** @return default obstacle area class */
-	FORCEINLINE static TSubclassOf<UNavArea> GetDefaultObstacleArea() { return DefaultObstacleArea; }
-	
 	static const FNavDataConfig& GetDefaultSupportedAgent();
 	FORCEINLINE const FNavDataConfig& GetDefaultSupportedAgentConfig() const { check(SupportedAgents.Num() > 0);  return SupportedAgents[0]; }
 	FORCEINLINE const TArray<FNavDataConfig>& GetSupportedAgents() const { return SupportedAgents; }
@@ -590,10 +584,20 @@ public:
 	const FNavigationOctree* GetNavOctree() const { return NavOctree.Get(); }
 	FNavigationOctree* GetMutableNavOctree() { return NavOctree.Get(); }
 
-	FORCEINLINE void SetObjectsNavOctreeId(UObject* Object, FOctreeElementId Id) { ObjectToOctreeId.Add(Object, Id); }
-	FORCEINLINE const FOctreeElementId* GetObjectsNavOctreeId(const UObject* Object) const { return ObjectToOctreeId.Find(Object); }
+	// coppied from GetObjectOuterHash
+	FORCEINLINE static int32 HashObject(const UObject& Object)
+	{
+		return ((Object.GetFName().GetComparisonIndex() ^ Object.GetFName().GetNumber()) ^ (PTRINT(Object.GetOuter()) >> 6));
+	}
+	FORCEINLINE void SetObjectsNavOctreeId(const UObject& Object, FOctreeElementId Id) { ObjectToOctreeId.Add(HashObject(Object), Id); }
+	FORCEINLINE const FOctreeElementId* GetObjectsNavOctreeId(const UObject& Object) const { return ObjectToOctreeId.Find(HashObject(Object)); }
+	FORCEINLINE bool HasPendingObjectNavOctreeId(UObject* Object) const { return PendingOctreeUpdates.Contains(FNavigationDirtyElement(Object)); }
+	FORCEINLINE	void RemoveObjectsNavOctreeId(const UObject& Object) { ObjectToOctreeId.Remove(HashObject(Object)); }
+
+	/*FORCEINLINE void SetObjectsNavOctreeId(UObject* Object, FOctreeElementId Id) { ObjectToOctreeId.Add(HashObject(Object), Id); }
+	FORCEINLINE const FOctreeElementId* GetObjectsNavOctreeId(const UObject* Object) const { return ObjectToOctreeId.Find(HashObject(Object)); }
 	FORCEINLINE bool HasPendingObjectNavOctreeId(UObject* Object) const { return PendingOctreeUpdates.Contains(FNavigationDirtyElement(Object)); }
-	FORCEINLINE	void RemoveObjectsNavOctreeId(UObject* Object) { ObjectToOctreeId.Remove(Object); }
+	FORCEINLINE	void RemoveObjectsNavOctreeId(UObject* Object) { ObjectToOctreeId.Remove(HashObject(Object)); }*/
 	void RemoveNavOctreeElementId(const FOctreeElementId& ElementId, int32 UpdateFlags);
 
 	const FNavigationRelevantData* GetDataForObject(const UObject& Object) const;
@@ -802,7 +806,7 @@ protected:
 
 	TMap<FNavAgentProperties, TWeakObjectPtr<ANavigationData> > AgentToNavDataMap;
 	
-	TMap<const UObject*, FOctreeElementId> ObjectToOctreeId;
+	TMap<int32, FOctreeElementId> ObjectToOctreeId;
 
 	/** Map of all objects that are tied to indexed navigation parent */
 	TMultiMap<UObject*, FWeakObjectPtr> OctreeChildNodesMap;
@@ -860,8 +864,6 @@ protected:
 
 	static TMap<INavLinkCustomInterface*, FWeakObjectPtr> PendingCustomLinkRegistration;
 	TSet<const UClass*> NavAreaClasses;
-	static TSubclassOf<UNavArea> DefaultWalkableArea;
-	static TSubclassOf<UNavArea> DefaultObstacleArea;
 
 	/** delegate handler for PostLoadMap event */
 	void OnPostLoadMap(UWorld* LoadedWorld);
@@ -897,7 +899,8 @@ protected:
 	void UnregisterNavOctreeElement(UObject* ElementOwner, INavRelevantInterface* ElementInterface, int32 UpdateFlags);
 	
 	/** read element data from navigation octree */
-	bool GetNavOctreeElementData(UObject* NodeOwner, int32& DirtyFlags, FBox& DirtyBounds);
+	bool GetNavOctreeElementData(const UObject& NodeOwner, int32& DirtyFlags, FBox& DirtyBounds);
+	//bool GetNavOctreeElementData(UObject* NodeOwner, int32& DirtyFlags, FBox& DirtyBounds);
 
 	/** Adds given element to NavOctree. No check for owner's validity are performed, 
 	 *	nor its presence in NavOctree - function assumes callee responsibility 
@@ -982,6 +985,10 @@ public:
 	DEPRECATED(4.20, "SimpleMoveToLocation is deprecated. Use UAIBlueprintHelperLibrary::SimpleMoveToLocation instead")
 	UFUNCTION(BlueprintCallable, Category = "AI|Navigation", meta = (DisplayName = "SimpleMoveToLocation_DEPRECATED", ScriptNoExport, DeprecatedFunction, DeprecationMessage = "SimpleMoveToLocation is deprecated. Use AIBlueprintHelperLibrary::SimpleMoveToLocation instead"))
 	static void SimpleMoveToLocation(AController* Controller, const FVector& Goal);
+	DEPRECATED(4.20, "UNavigationSystem::GetDefaultWalkableArea is deprecated. Use FNavigationSystem::GetDefaultWalkableArea instead")
+	static TSubclassOf<UNavAreaBase> GetDefaultWalkableArea() { return FNavigationSystem::GetDefaultWalkableArea(); }
+	DEPRECATED(4.20, "UNavigationSystem::GetDefaultObstacleArea is deprecated. Use FNavigationSystem::GetDefaultObstacleArea instead")
+	static TSubclassOf<UNavAreaBase> GetDefaultObstacleArea() { return FNavigationSystem::GetDefaultObstacleArea(); }
 };
 
 
