Fixed 'Navmesh needs to be rebuild' warning erroneously flashing during PIE #UE4

diff --git a/Engine/Source/Runtime/NavigationSystem/Private/NavigationSystem.cpp b/Engine/Source/Runtime/NavigationSystem/Private/NavigationSystem.cpp
--- a/Engine/Source/Runtime/NavigationSystem/Private/NavigationSystem.cpp
+++ b/Engine/Source/Runtime/NavigationSystem/Private/NavigationSystem.cpp
@@ -4273,8 +4273,9 @@ void UNavigationSystemV1::GetOnScreenMessages(TMultiMap<FCoreDelegates::EOnScree
 #else
 	const bool bIsNavigationAutoUpdateEnabled = true;
 #endif
-	if (IsNavigationDirty() 
-		&& (!bIsNavigationAutoUpdateEnabled || !SupportsNavigationGeneration() || !CanRebuildDirtyNavigation()))
+	if (IsNavigationDirty()
+		&& ((OperationMode == FNavigationSystemRunMode::EditorMode && !bIsNavigationAutoUpdateEnabled)
+			|| !SupportsNavigationGeneration() || !CanRebuildDirtyNavigation()))
 	{
 		OutMessages.Add(FCoreDelegates::EOnScreenMessageSeverity::Error
 			, LOCTEXT("NAVMESHERROR", "NAVMESH NEEDS TO BE REBUILT"));
