(4.20.2) PR #4980: Fix: CrowdManagerClass property in UNavigationSystem is now correctly marked with "config" specifier. (Contributed by [REDACTED_PERSON])
#rnx

diff --git a/Engine/Source/Runtime/NavigationSystem/Public/NavigationSystem.h b/Engine/Source/Runtime/NavigationSystem/Public/NavigationSystem.h
--- a/Engine/Source/Runtime/NavigationSystem/Public/NavigationSystem.h
+++ b/Engine/Source/Runtime/NavigationSystem/Public/NavigationSystem.h
@@ -111,7 +111,7 @@ public:
 	ANavigationData* AbstractNavData;
 
 protected:
-	UPROPERTY(EditAnywhere, BlueprintReadOnly, Category = Navigation)
+	UPROPERTY(config, EditAnywhere, BlueprintReadOnly, Category = Navigation)
 	TSubclassOf<UCrowdManagerBase> CrowdManagerClass;
 
 	/** Should navigation system spawn default Navigation Data when there's none and there are navigation bounds present? */
