Adding some perf starts to RegisterNavOctreeElement / UnregisterNavOctreeElement. These functions have proven to be quite expensive when there is a lot of geometry in the NavOctree.

diff --git a/Engine/Source/Runtime/Engine/Public/EngineStats.h b/Engine/Source/Runtime/Engine/Public/EngineStats.h
--- a/Engine/Source/Runtime/Engine/Public/EngineStats.h
+++ b/Engine/Source/Runtime/Engine/Public/EngineStats.h
@@ -114,6 +114,8 @@ DECLARE_CYCLE_STAT_EXTERN(TEXT("Sync queries"),STAT_Navigation_QueriesTimeSync,S
 DECLARE_CYCLE_STAT_EXTERN(TEXT("Sync meta nav area preparation"),STAT_Navigation_MetaAreaTranslation,STATGROUP_Navigation, ); 
 DECLARE_CYCLE_STAT_EXTERN(TEXT("Async nav areas sorting"),STAT_Navigation_TileNavAreaSorting,STATGROUP_Navigation, );
 DECLARE_DWORD_COUNTER_STAT_EXTERN(TEXT("Num update nav octree"),STAT_Navigation_UpdateNavOctree,STATGROUP_Navigation, );
+DECLARE_CYCLE_STAT_EXTERN(TEXT("Registering Nav Octree Element"), STAT_Navigation_RegisterNavOctreeElement, STATGROUP_Navigation, );
+DECLARE_CYCLE_STAT_EXTERN(TEXT("Unregistering Nav Octree Element"), STAT_Navigation_UnregisterNavOctreeElement, STATGROUP_Navigation, )
 DECLARE_CYCLE_STAT_EXTERN(TEXT("Adding actors to navoctree"),STAT_Navigation_AddingActorsToNavOctree,STATGROUP_Navigation, );
 DECLARE_CYCLE_STAT_EXTERN(TEXT("Adjusting nav links"),STAT_Navigation_AdjustingNavLinks,STATGROUP_Navigation, );
 DECLARE_CYCLE_STAT_EXTERN(TEXT("Recast: sync add generated tiles"), STAT_Navigation_RecastAddGeneratedTiles,STATGROUP_Navigation, );
diff --git a/Engine/Source/Runtime/NavigationSystem/Private/NavigationSystem.cpp b/Engine/Source/Runtime/NavigationSystem/Private/NavigationSystem.cpp
--- a/Engine/Source/Runtime/NavigationSystem/Private/NavigationSystem.cpp
+++ b/Engine/Source/Runtime/NavigationSystem/Private/NavigationSystem.cpp
@@ -80,6 +80,8 @@ DEFINE_STAT(STAT_Navigation_GatheringNavigationModifiersSync);
 DEFINE_STAT(STAT_Navigation_ActorsGeometryExportSync);
 DEFINE_STAT(STAT_Navigation_ProcessingActorsForNavMeshBuilding);
 DEFINE_STAT(STAT_Navigation_AdjustingNavLinks);
+DEFINE_STAT(STAT_Navigation_RegisterNavOctreeElement);
+DEFINE_STAT(STAT_Navigation_UnregisterNavOctreeElement);
 DEFINE_STAT(STAT_Navigation_AddingActorsToNavOctree);
 DEFINE_STAT(STAT_Navigation_RecastAddGeneratedTiles);
 DEFINE_STAT(STAT_Navigation_RecastTick);
@@ -2386,6 +2388,8 @@ int32 GetDirtyFlagHelper(int32 UpdateFlags, int32 DefaultValue)
 
 FSetElementId UNavigationSystemV1::RegisterNavOctreeElement(UObject* ElementOwner, INavRelevantInterface* ElementInterface, int32 UpdateFlags)
 {
+	SCOPE_CYCLE_COUNTER(STAT_Navigation_RegisterNavOctreeElement);
+
 	FSetElementId SetId;
 
 #if WITH_EDITOR
@@ -2529,6 +2533,8 @@ bool UNavigationSystemV1::GetNavOctreeElementData(const UObject& NodeOwner, int3
 
 void UNavigationSystemV1::UnregisterNavOctreeElement(UObject* ElementOwner, INavRelevantInterface* ElementInterface, int32 UpdateFlags)
 {
+	SCOPE_CYCLE_COUNTER(STAT_Navigation_UnregisterNavOctreeElement);
+
 #if WITH_EDITOR
 	if (IsNavigationUnregisterLocked())
 	{
