Added a way to ask the navmesh generator if there's any generation taking place at given moment inside a specified bounding box #UE4

diff --git a/Engine/Source/Runtime/NavigationSystem/Private/NavMesh/RecastNavMeshGenerator.cpp b/Engine/Source/Runtime/NavigationSystem/Private/NavMesh/RecastNavMeshGenerator.cpp
--- a/Engine/Source/Runtime/NavigationSystem/Private/NavMesh/RecastNavMeshGenerator.cpp
+++ b/Engine/Source/Runtime/NavigationSystem/Private/NavMesh/RecastNavMeshGenerator.cpp
@@ -59,6 +59,28 @@ FORCEINLINE bool DoesBoxContainBox(const FBox& BigBox, const FBox& SmallBox)
 	return DoesBoxContainOrOverlapVector(BigBox, SmallBox.Min) && DoesBoxContainOrOverlapVector(BigBox, SmallBox.Max);
 }
 
+struct FRcTileBox
+{
+	int32 XMin, XMax, YMin, YMax;
+
+	FRcTileBox(const FBox& UnrealBounds, const FVector& RcNavMeshOrigin, const float TileSizeInWorldUnits)
+	{
+		check(TileSizeInWorldUnits > 0);
+
+		const FBox RcAreaBounds = Unreal2RecastBox(UnrealBounds);
+		XMin = FMath::FloorToInt((RcAreaBounds.Min.X - RcNavMeshOrigin.X) / TileSizeInWorldUnits);
+		XMax = FMath::FloorToInt((RcAreaBounds.Max.X - RcNavMeshOrigin.X) / TileSizeInWorldUnits);
+		YMin = FMath::FloorToInt((RcAreaBounds.Min.Z - RcNavMeshOrigin.Z) / TileSizeInWorldUnits);
+		YMax = FMath::FloorToInt((RcAreaBounds.Max.Z - RcNavMeshOrigin.Z) / TileSizeInWorldUnits);
+	}
+
+	FORCEINLINE bool Contains(const FIntPoint& Point) const 
+	{
+		return Point.X >= XMin && Point.X <= XMax
+			&& Point.Y >= YMin && Point.Y <= YMax;
+	}
+};
+
 int32 GetTilesCountHelper(const dtNavMesh* DetourMesh)
 {
 	int32 NumTiles = 0;
@@ -4067,6 +4089,49 @@ namespace
 	}
 }
 
+bool FRecastNavMeshGenerator::HasDirtyTiles(const FBox& AreaBounds) const
+{
+	if (HasDirtyTiles() == false)
+	{
+		return false;
+	}
+
+	bool bRetDirty = false;
+	const float TileSizeInWorldUnits = Config.tileSize * Config.cs;
+	const FRcTileBox TileBox(AreaBounds, RcNavMeshOrigin, TileSizeInWorldUnits);
+		
+	for (int32 Index = 0; bRetDirty == false && Index < PendingDirtyTiles.Num(); ++Index)
+	{
+		bRetDirty = TileBox.Contains(PendingDirtyTiles[Index].Coord);
+	}
+	for (int32 Index = 0; bRetDirty == false && Index < RunningDirtyTiles.Num(); ++Index)
+	{
+		bRetDirty = TileBox.Contains(RunningDirtyTiles[Index].Coord);
+	}
+
+	return bRetDirty;
+}
+
+int32 FRecastNavMeshGenerator::GetDirtyTilesCount(const FBox& AreaBounds) const
+{
+	const float TileSizeInWorldUnits = Config.tileSize * Config.cs;
+	const FRcTileBox TileBox(AreaBounds, RcNavMeshOrigin, TileSizeInWorldUnits);
+
+	int32 DirtyPendingCount = 0;
+	for (const FPendingTileElement& PendingElement : PendingDirtyTiles)
+	{
+		DirtyPendingCount += TileBox.Contains(PendingElement.Coord) ? 1 : 0;
+	}
+
+	int32 RunningCount = 0;
+	for (const FRunningTileElement& RunningElement : RunningDirtyTiles)
+	{
+		RunningCount += TileBox.Contains(RunningElement.Coord) ? 1 : 0;
+	}
+
+	return DirtyPendingCount + RunningCount;
+}
+
 void FRecastNavMeshGenerator::MarkDirtyTiles(const TArray<FNavigationDirtyArea>& DirtyAreas)
 {
 	QUICK_SCOPE_CYCLE_COUNTER(STAT_RecastNavMeshGenerator_MarkDirtyTiles);
@@ -4114,16 +4179,12 @@ void FRecastNavMeshGenerator::MarkDirtyTiles(const TArray<FNavigationDirtyArea>&
 			// then FindInclusionBoundEncapsulatingBox can produce false negative
 			bDoTileInclusionTest = FindInclusionBoundEncapsulatingBox(CutDownArea) == INDEX_NONE;
 		}
-				
-		const FBox RcAreaBounds = Unreal2RecastBox(AdjustedAreaBounds);
-		const int32 XMin = FMath::FloorToInt((RcAreaBounds.Min.X - RcNavMeshOrigin.X) / TileSizeInWorldUnits);
-		const int32 XMax = FMath::FloorToInt((RcAreaBounds.Max.X - RcNavMeshOrigin.X) / TileSizeInWorldUnits);
-		const int32 YMin = FMath::FloorToInt((RcAreaBounds.Min.Z - RcNavMeshOrigin.Z) / TileSizeInWorldUnits);
-		const int32 YMax = FMath::FloorToInt((RcAreaBounds.Max.Z - RcNavMeshOrigin.Z) / TileSizeInWorldUnits);
+		
+		const FRcTileBox TileBox(AdjustedAreaBounds, RcNavMeshOrigin, TileSizeInWorldUnits);
 
-		for (int32 TileY = YMin; TileY <= YMax; ++TileY)
+		for (int32 TileY = TileBox.YMin; TileY <= TileBox.YMax; ++TileY)
 		{
-			for (int32 TileX = XMin; TileX <= XMax; ++TileX)
+			for (int32 TileX = TileBox.XMin; TileX <= TileBox.XMax; ++TileX)
 			{
 				if (IsInActiveSet(FIntPoint(TileX, TileY)) == false)
 				{
@@ -4132,10 +4193,10 @@ void FRecastNavMeshGenerator::MarkDirtyTiles(const TArray<FNavigationDirtyArea>&
 
 				if (DirtyArea.HasFlag(ENavigationDirtyFlag::NavigationBounds) == false && bDoTileInclusionTest == true)
 				{
-					const FBox TileBox = CalculateTileBounds(TileX, TileY, RcNavMeshOrigin, TotalNavBounds, TileSizeInWorldUnits);
+					const FBox TileBounds = CalculateTileBounds(TileX, TileY, RcNavMeshOrigin, TotalNavBounds, TileSizeInWorldUnits);
 
 					// do per tile check since we can have lots of tiles inbetween navigable bounds volumes
-					if (IntersectBounds(TileBox, InclusionBounds) == false)
+					if (IntersectBounds(TileBounds, InclusionBounds) == false)
 					{
 						// Skip this tile
 						continue;
diff --git a/Engine/Source/Runtime/NavigationSystem/Public/NavMesh/RecastNavMeshGenerator.h b/Engine/Source/Runtime/NavigationSystem/Public/NavMesh/RecastNavMeshGenerator.h
--- a/Engine/Source/Runtime/NavigationSystem/Public/NavMesh/RecastNavMeshGenerator.h
+++ b/Engine/Source/Runtime/NavigationSystem/Public/NavMesh/RecastNavMeshGenerator.h
@@ -451,6 +451,8 @@ public:
 	const FRecastNavMeshCachedData& GetAdditionalCachedData() const { return AdditionalCachedData; }
 
 	bool HasDirtyTiles() const;
+	bool HasDirtyTiles(const FBox& AreaBounds) const;
+	int32 GetDirtyTilesCount(const FBox& AreaBounds) const;
 
 	bool GatherGeometryOnGameThread() const;
 
