Work to provide ConditionalPopulateNavOctreeActor Delgate for the Nav System.

diff --git a/Engine/Source/Runtime/NavigationSystem/Private/NavigationSystem.cpp b/Engine/Source/Runtime/NavigationSystem/Private/NavigationSystem.cpp
--- a/Engine/Source/Runtime/NavigationSystem/Private/NavigationSystem.cpp
+++ b/Engine/Source/Runtime/NavigationSystem/Private/NavigationSystem.cpp
@@ -654,7 +654,7 @@ bool UNavigationSystemV1::ConditionalPopulateNavOctree()
 					const bool bLegalActor = Actor && !Actor->IsPendingKill();
 					if (bLegalActor)
 					{
-						UpdateActorAndComponentsInNavOctree(*Actor);
+						ConditionalPopulateNavOctreeActor(*Actor);
 					}
 				}
 			}
@@ -2977,6 +2977,11 @@ void UNavigationSystemV1::OnActorUnregistered(AActor* Actor)
 	}
 }
 
+void UNavigationSystemV1::ConditionalPopulateNavOctreeActor(AActor& Actor)
+{
+	UpdateActorAndComponentsInNavOctree(Actor);
+}
+
 void UNavigationSystemV1::FindElementsInNavOctree(const FBox& QueryBox, const FNavigationOctreeFilter& Filter, TArray<FNavigationOctreeElement>& Elements)
 {
 	if (NavOctree.IsValid() == false)
diff --git a/Engine/Source/Runtime/NavigationSystem/Public/NavigationSystem.h b/Engine/Source/Runtime/NavigationSystem/Public/NavigationSystem.h
--- a/Engine/Source/Runtime/NavigationSystem/Public/NavigationSystem.h
+++ b/Engine/Source/Runtime/NavigationSystem/Public/NavigationSystem.h
@@ -700,6 +700,9 @@ protected:
 
 	/** spawn new crowd manager */
 	virtual void UpdateAbstractNavData();
+
+	/** Called when ConditionalPopulateNavOctree processes each Actor */
+	virtual void ConditionalPopulateNavOctreeActor(AActor& Actor);
 	
 public:
 	/** Called upon UWorld destruction to release what needs to be released */
