Fix for floats being used instead of doubles in nav tile time slicing causing calculation errors which can result in long nav mesh build times.

Removed unnecessary float.

diff --git a/Engine/Source/Runtime/NavigationSystem/Private/NavMesh/RecastNavMeshGenerator.cpp b/Engine/Source/Runtime/NavigationSystem/Private/NavMesh/RecastNavMeshGenerator.cpp
--- a/Engine/Source/Runtime/NavigationSystem/Private/NavMesh/RecastNavMeshGenerator.cpp
+++ b/Engine/Source/Runtime/NavigationSystem/Private/NavMesh/RecastNavMeshGenerator.cpp
@@ -4416,8 +4416,8 @@ TArray<uint32> FRecastNavMeshGenerator::ProcessTileTasksAsync(const int32 NumTas
 #if TIME_SLICE_NAV_REGEN
 bool IsTimeSliceDurationExceeded(const double StartTime)
 {
-	const float TimeSliceDuration = 0.0025f;
-	const float	CurTime = FPlatformTime::Seconds();
+	const double TimeSliceDuration = 0.0025f;
+	const double CurTime = FPlatformTime::Seconds();
 	
 	return CurTime - StartTime >= TimeSliceDuration;
 }
@@ -4433,7 +4433,6 @@ TArray<uint32> FRecastNavMeshGenerator::ProcessTileTasksSync(const int32 NumTask
 
 #if TIME_SLICE_NAV_REGEN
 	const double StartTime = FPlatformTime::Seconds();
-	double CurTime = StartTime;
 
 	//if we are time slice processing a tile (ie we have already done some processing on this tile last frame)
 	bool bIsTimeSliceProcessingTile = TileGeneratorSync.Get() != nullptr;
