Allow navigation tile TimeSliceDuration to be overriden by child classes.

diff --git a/Engine/Source/Runtime/NavigationSystem/Private/NavMesh/RecastNavMeshGenerator.cpp b/Engine/Source/Runtime/NavigationSystem/Private/NavMesh/RecastNavMeshGenerator.cpp
--- a/Engine/Source/Runtime/NavigationSystem/Private/NavMesh/RecastNavMeshGenerator.cpp
+++ b/Engine/Source/Runtime/NavigationSystem/Private/NavMesh/RecastNavMeshGenerator.cpp
@@ -3221,6 +3221,10 @@ FRecastNavMeshGenerator::FRecastNavMeshGenerator(ARecastNavMesh& InDestNavMesh)
 	, bRestrictBuildingToActiveTiles(false)
 	, Version(0)
 {
+#if TIME_SLICE_NAV_REGEN
+	TimeSliceDuration = 0.0025;
+#endif
+
 	INC_DWORD_STAT_BY(STAT_NavigationMemory, sizeof(*this));
 }
 
@@ -4414,9 +4418,8 @@ TArray<uint32> FRecastNavMeshGenerator::ProcessTileTasksAsync(const int32 NumTas
 #endif
 
 #if TIME_SLICE_NAV_REGEN
-bool IsTimeSliceDurationExceeded(const double StartTime)
+bool FRecastNavMeshGenerator::IsTimeSliceDurationExceeded(const double StartTime) const
 {
-	const double TimeSliceDuration = 0.0025f;
 	const double CurTime = FPlatformTime::Seconds();
 	
 	return CurTime - StartTime >= TimeSliceDuration;
diff --git a/Engine/Source/Runtime/NavigationSystem/Public/NavMesh/RecastNavMeshGenerator.h b/Engine/Source/Runtime/NavigationSystem/Public/NavMesh/RecastNavMeshGenerator.h
--- a/Engine/Source/Runtime/NavigationSystem/Public/NavMesh/RecastNavMeshGenerator.h
+++ b/Engine/Source/Runtime/NavigationSystem/Public/NavMesh/RecastNavMeshGenerator.h
@@ -535,6 +535,10 @@ public:
 	/** Adds generated tiles to NavMesh, replacing old ones */
 	TArray<uint32> AddGeneratedTiles(FRecastTileGenerator& TileGenerator);
 
+#if TIME_SLICE_NAV_REGEN
+	bool IsTimeSliceDurationExceeded(const double StartTime) const;
+#endif
+
 public:
 	/** Removes all tiles at specified grid location */
 	TArray<uint32> RemoveTileLayers(const int32 TileX, const int32 TileY, TMap<int32, dtPolyRef>* OldLayerTileIdMap = nullptr);
@@ -622,6 +626,10 @@ protected:
 	/** Runtime generator's version, increased every time all tile generators get invalidated
 	 *	like when navmesh size changes */
 	uint32 Version;
+
+#if TIME_SLICE_NAV_REGEN
+	double TimeSliceDuration;
+#endif
 };
 
 #endif // WITH_RECAST
