Added functions for quick replacing of nav area information in NavOctree and RecastNavMesh #UE4

ARecastNavMesh::ReplaceAreaInTileBounds replaces occurences of given area directly on navmesh polys (checking and setting areaID) in tiles overlapped by given bounding box.
UNavigationSystemV1::ReplaceAreaInOctreeData replaces occurences of given area in object's nav representation in the NavOctree.

Also extended navmesh drawing to be able to optionally mark polys excluded by navmesh's default query filter (via dtQueryFilter.data.m_excludeFlags).

review-5764395

diff --git a/Engine/Source/Runtime/Engine/Public/AI/NavigationModifier.h b/Engine/Source/Runtime/Engine/Public/AI/NavigationModifier.h
--- a/Engine/Source/Runtime/Engine/Public/AI/NavigationModifier.h
+++ b/Engine/Source/Runtime/Engine/Public/AI/NavigationModifier.h
@@ -306,6 +306,10 @@ struct ENGINE_API FCompositeNavModifier : public FNavigationModifier
 	// Should be called only on game thread
 	void GetPerInstanceTransforms(const FBox& AreaBox, TArray<FTransform>& PerInstanceTransforms) const;
 
+	TArray<FAreaNavModifier>& GetMutableAreas() { return Areas; }
+	TArray<FSimpleLinkNavModifier>& GetSimpleLinks() { return SimpleLinks; }
+	TArray<FCustomLinkNavModifier>& GetCustomLinks() { return CustomLinks; }
+
 public:
 	// Gathers per instance data for navigation area modifiers in a specified area box
 	FNavDataPerInstanceTransformDelegate NavDataPerInstanceTransformDelegate;
diff --git a/Engine/Source/Runtime/NavigationSystem/Private/NavMesh/NavMeshRenderingComponent.cpp b/Engine/Source/Runtime/NavigationSystem/Private/NavMesh/NavMeshRenderingComponent.cpp
--- a/Engine/Source/Runtime/NavigationSystem/Private/NavMesh/NavMeshRenderingComponent.cpp
+++ b/Engine/Source/Runtime/NavigationSystem/Private/NavMesh/NavMeshRenderingComponent.cpp
@@ -31,6 +31,7 @@ static const FColor NavMeshRenderColor_TileBounds(255, 255, 64, 255);
 static const FColor NavMeshRenderColor_PathCollidingGeom(255, 255, 255, 40);
 static const FColor NavMeshRenderColor_RecastTileBeingRebuilt(255, 0, 0, 64);
 static const FColor NavMeshRenderColor_OffMeshConnectionInvalid(64, 64, 64);
+static const FColor NavMeshRenderColor_PolyForbidden(FColorList::Black);
 
 static const float DefaultEdges_LineThickness = 0.0f;
 static const float PolyEdges_LineThickness = 1.5f;
@@ -109,6 +110,18 @@ namespace FNavMeshRenderingHelpers
 		DebugLines.Add(FDebugRenderSceneProxy::FDebugLine(Tip, FVector(Tip.X + Ay.X*Size - Ax.X*Size / 3, Tip.Y + Ay.Y*Size - Ax.Y*Size / 3, Tip.Z + Ay.Z*Size - Ax.Z*Size / 3), Color.ToFColor(true)));
 	}
 
+	void CacheLink(TArray<FDebugRenderSceneProxy::FDebugLine>& DebugLines, const FVector V0, const FVector V1, const FColor LinkColor, const uint8 LinkDirection)
+	{
+		FNavMeshRenderingHelpers::CacheArc(DebugLines, V0, V1, 0.4f, 4, LinkColor, LinkLines_LineThickness);
+
+		const FVector VOffset(0, 0, FVector::Dist(V0, V1) * 1.333f);
+		FNavMeshRenderingHelpers::CacheArrowHead(DebugLines, V1, V0 + VOffset, 30.f, LinkColor, LinkLines_LineThickness);
+		if (LinkDirection)
+		{
+			FNavMeshRenderingHelpers::CacheArrowHead(DebugLines, V0, V1 + VOffset, 30.f, LinkColor, LinkLines_LineThickness);
+		}
+	}
+
 	void DrawWireCylinder(TArray<FDebugRenderSceneProxy::FDebugLine>& DebugLines, const FVector& Base, const FVector& X, const FVector& Y, const FVector& Z, FColor Color, float Radius, float HalfHeight, int32 NumSides, uint8 DepthPriority, float LineThickness = 0)
 	{
 		const float	AngleDelta = 2.0f * PI / NumSides;
@@ -166,6 +179,27 @@ namespace FNavMeshRenderingHelpers
 		MeshData.Indices.Add(V2);
 	}
 
+	void AddArea(TArray<FNavMeshSceneProxyData::FDebugMeshData>& MeshBuilders, const TArray<int32>& MeshIndices, const TArray<FVector>& MeshVerts, const FColor Color, const FVector DrawOffset)
+	{
+		if (MeshIndices.Num() == 0)
+		{
+			return;
+		}
+
+		FNavMeshSceneProxyData::FDebugMeshData DebugMeshData;
+		for (int32 VertIdx = 0; VertIdx < MeshVerts.Num(); ++VertIdx)
+		{
+			FNavMeshRenderingHelpers::AddVertex(DebugMeshData, MeshVerts[VertIdx] + DrawOffset, Color);
+		}
+		for (int32 TriIdx = 0; TriIdx < MeshIndices.Num(); TriIdx += 3)
+		{
+			FNavMeshRenderingHelpers::AddTriangle(DebugMeshData, MeshIndices[TriIdx], MeshIndices[TriIdx + 1], MeshIndices[TriIdx + 2]);
+		}
+
+		DebugMeshData.ClusterColor = Color;
+		MeshBuilders.Add(DebugMeshData);
+	}
+
 	void AddRecastGeometry(TArray<FVector>& OutVertexBuffer, TArray<uint32>& OutIndexBuffer, const float* Coords, int32 NumVerts, const int32* Faces, int32 NumFaces, const FTransform& Transform = FTransform::Identity)
 	{
 		const int32 VertIndexBase = OutVertexBuffer.Num();
@@ -331,7 +365,8 @@ int32 FNavMeshSceneProxyData::GetDetailFlags(const ARecastNavMesh* NavMesh) cons
 		(NavMesh->bDrawFailedNavLinks ? (1 << static_cast<int32>(ENavMeshDetailFlags::FailedNavLinks)) : 0) |
 		(NavMesh->bDrawClusters ? (1 << static_cast<int32>(ENavMeshDetailFlags::Clusters)) : 0) |
 		(NavMesh->bDrawOctree ? (1 << static_cast<int32>(ENavMeshDetailFlags::NavOctree)) : 0) | 
-		(NavMesh->bDrawOctreeDetails ? (1 << static_cast<int32>(ENavMeshDetailFlags::NavOctreeDetails)) : 0);
+		(NavMesh->bDrawOctreeDetails ? (1 << static_cast<int32>(ENavMeshDetailFlags::NavOctreeDetails)) : 0) | 
+		(NavMesh->bDrawMarkedForbiddenPolys ? (1 << static_cast<int32>(ENavMeshDetailFlags::MarkForbiddenPolys)) : 0);
 }
 
 void FNavMeshSceneProxyData::GatherData(const ARecastNavMesh* NavMesh, int32 InNavDetailFlags, const TArray<int32>& TileSet)
@@ -351,6 +386,7 @@ void FNavMeshSceneProxyData::GatherData(const ARecastNavMesh* NavMesh, int32 InN
 		FRecastDebugGeometry NavMeshGeometry;
 		NavMeshGeometry.bGatherPolyEdges = FNavMeshRenderingHelpers::HasFlag(NavDetailFlags, ENavMeshDetailFlags::PolyEdges);
 		NavMeshGeometry.bGatherNavMeshEdges = FNavMeshRenderingHelpers::HasFlag(NavDetailFlags, ENavMeshDetailFlags::BoundaryEdges);
+		NavMeshGeometry.bMarkForbiddenPolys = FNavMeshRenderingHelpers::HasFlag(NavDetailFlags, ENavMeshDetailFlags::MarkForbiddenPolys);
 
 		const FNavDataConfig& NavConfig = NavMesh->GetConfig();
 		TArray<FColor> NavMeshColors;
@@ -437,14 +473,7 @@ void FNavMeshSceneProxyData::GatherData(const ARecastNavMesh* NavMesh, int32 InN
 					const FVector V1 = Link.Right + NavMeshDrawOffset;
 					const FColor LinkColor = ((Link.Direction && Link.ValidEnds) || (Link.ValidEnds & FRecastDebugGeometry::OMLE_Left)) ? FNavMeshRenderingHelpers::DarkenColor(NavMeshColors[Link.AreaID]) : NavMeshRenderColor_OffMeshConnectionInvalid;
 
-					FNavMeshRenderingHelpers::CacheArc(NavLinkLines, V0, V1, 0.4f, 4, LinkColor, LinkLines_LineThickness);
-
-					const FVector VOffset(0, 0, FVector::Dist(V0, V1) * 1.333f);
-					FNavMeshRenderingHelpers::CacheArrowHead(NavLinkLines, V1, V0 + VOffset, 30.f, LinkColor, LinkLines_LineThickness);
-					if (Link.Direction)
-					{
-						FNavMeshRenderingHelpers::CacheArrowHead(NavLinkLines, V0, V1 + VOffset, 30.f, LinkColor, LinkLines_LineThickness);
-					}
+					FNavMeshRenderingHelpers::CacheLink(NavLinkLines, V0, V1, LinkColor, Link.Direction);
 
 					// if the connection as a whole is valid check if there are any of ends is invalid
 					if (LinkColor != NavMeshRenderColor_OffMeshConnectionInvalid)
@@ -461,6 +490,19 @@ void FNavMeshSceneProxyData::GatherData(const ARecastNavMesh* NavMesh, int32 InN
 					}
 				}
 			}
+			if (NavMeshGeometry.bMarkForbiddenPolys)
+			{
+				for (int32 OffMeshLineIndex = 0; OffMeshLineIndex < NavMeshGeometry.ForbiddenLinks.Num(); ++OffMeshLineIndex)
+				{
+					FRecastDebugGeometry::FOffMeshLink& Link = NavMeshGeometry.ForbiddenLinks[OffMeshLineIndex];
+					
+					const FVector V0 = Link.Left + NavMeshDrawOffset;
+					const FVector V1 = Link.Right + NavMeshDrawOffset;
+					const FColor LinkColor = NavMeshRenderColor_PolyForbidden;
+
+					FNavMeshRenderingHelpers::CacheLink(NavLinkLines, V0, V1, LinkColor, Link.Direction);
+				}
+			}
 		}
 
 		const bool bGatherTileLabels = FNavMeshRenderingHelpers::HasFlag(NavDetailFlags, ENavMeshDetailFlags::TileLabels);
@@ -590,25 +632,9 @@ void FNavMeshSceneProxyData::GatherData(const ARecastNavMesh* NavMesh, int32 InN
 			{
 				for (int32 AreaType = 0; AreaType < RECAST_MAX_AREAS; ++AreaType)
 				{
-					const TArray<int32>& MeshIndices = NavMeshGeometry.AreaIndices[AreaType];
-					if (MeshIndices.Num() == 0)
-					{
-						continue;
-					}
-
-					FDebugMeshData DebugMeshData;
-					for (int32 VertIdx = 0; VertIdx < MeshVerts.Num(); ++VertIdx)
-					{
-						FNavMeshRenderingHelpers::AddVertex(DebugMeshData, MeshVerts[VertIdx] + NavMeshDrawOffset, NavMeshColors[AreaType]);
-					}
-					for (int32 TriIdx = 0; TriIdx < MeshIndices.Num(); TriIdx += 3)
-					{
-						FNavMeshRenderingHelpers::AddTriangle(DebugMeshData, MeshIndices[TriIdx], MeshIndices[TriIdx + 1], MeshIndices[TriIdx + 2]);
-					}
-
-					DebugMeshData.ClusterColor = NavMeshColors[AreaType];
-					MeshBuilders.Add(DebugMeshData);
+					FNavMeshRenderingHelpers::AddArea(MeshBuilders, NavMeshGeometry.AreaIndices[AreaType], NavMeshGeometry.MeshVerts, NavMeshColors[AreaType], NavMeshDrawOffset);
 				}
+				FNavMeshRenderingHelpers::AddArea(MeshBuilders, NavMeshGeometry.ForbiddenIndices, NavMeshGeometry.MeshVerts, NavMeshRenderColor_PolyForbidden, NavMeshDrawOffset);
 			}
 		}
 
diff --git a/Engine/Source/Runtime/NavigationSystem/Private/NavMesh/PImplRecastNavMesh.cpp b/Engine/Source/Runtime/NavigationSystem/Private/NavMesh/PImplRecastNavMesh.cpp
--- a/Engine/Source/Runtime/NavigationSystem/Private/NavMesh/PImplRecastNavMesh.cpp
+++ b/Engine/Source/Runtime/NavigationSystem/Private/NavMesh/PImplRecastNavMesh.cpp
@@ -2169,6 +2169,10 @@ void FPImplRecastNavMesh::GetDebugGeometry(FRecastDebugGeometry& OutGeometry, in
 	int32 NumVertsToReserve = 0;
 	int32 NumIndicesToReserve = 0;
 
+	int32 ForbiddenFlags = OutGeometry.bMarkForbiddenPolys 
+		? GetFilterForbiddenFlags((const FRecastQueryFilter*)NavMeshOwner->GetDefaultQueryFilterImpl()) 
+		: 0;
+
 	const FRecastNavMeshGenerator* Generator = static_cast<const FRecastNavMeshGenerator*>(NavMeshOwner->GetGenerator());
 
 	if (Generator && Generator->IsBuildingRestrictedToActiveTiles()
@@ -2210,7 +2214,7 @@ void FPImplRecastNavMesh::GetDebugGeometry(FRecastDebugGeometry& OutGeometry, in
 				dtMeshTile const* const Tile = ConstNavMesh->getTileAt(TileLocation.X, TileLocation.Y, Layer);
 				if (Tile != nullptr && Tile->header != nullptr)
 				{
-					VertBase += GetTilesDebugGeometry(Generator, *Tile, VertBase, OutGeometry);
+					VertBase += GetTilesDebugGeometry(Generator, *Tile, VertBase, OutGeometry, INDEX_NONE, ForbiddenFlags);
 				}
 			}
 		}
@@ -2248,12 +2252,12 @@ void FPImplRecastNavMesh::GetDebugGeometry(FRecastDebugGeometry& OutGeometry, in
 				continue;
 			}
 
-			VertBase += GetTilesDebugGeometry(Generator, *Tile, VertBase, OutGeometry, TileIdx);
+			VertBase += GetTilesDebugGeometry(Generator, *Tile, VertBase, OutGeometry, TileIdx, ForbiddenFlags);
 		}
 	}
 }
 
-int32 FPImplRecastNavMesh::GetTilesDebugGeometry(const FRecastNavMeshGenerator* Generator, const dtMeshTile& Tile, int32 VertBase, FRecastDebugGeometry& OutGeometry, int32 TileIdx) const
+int32 FPImplRecastNavMesh::GetTilesDebugGeometry(const FRecastNavMeshGenerator* Generator, const dtMeshTile& Tile, int32 VertBase, FRecastDebugGeometry& OutGeometry, int32 TileIdx, uint16 ForbiddenFlags) const
 {
 	check(NavMeshOwner && DetourNavMesh);
 	dtMeshHeader const* const Header = Tile.header;
@@ -2289,7 +2293,10 @@ int32 FPImplRecastNavMesh::GetTilesDebugGeometry(const FRecastNavMeshGenerator*
 		{
 			dtPolyDetail const* const DetailPoly = &Tile.detailMeshes[PolyIdx];
 
-			TArray<int32>* Indices = bIsBeingBuilt ? &OutGeometry.BuiltMeshIndices : &OutGeometry.AreaIndices[Poly->getArea()];
+			TArray<int32>* Indices = bIsBeingBuilt ? &OutGeometry.BuiltMeshIndices 
+				: ((Poly->flags & ForbiddenFlags) != 0
+					? &OutGeometry.ForbiddenIndices
+					: &OutGeometry.AreaIndices[Poly->getArea()]);
 
 			// one triangle at a time
 			for (int32 TriIdx = 0; TriIdx < DetailPoly->triCount; ++TriIdx)
@@ -2354,7 +2361,9 @@ int32 FPImplRecastNavMesh::GetTilesDebugGeometry(const FRecastNavMeshGenerator*
 				, OffMeshConnection->rad
 			};
 
-			OutGeometry.OffMeshLinks.Add(Link);
+			(LinkPoly->flags & ForbiddenFlags) != 0
+				? OutGeometry.ForbiddenLinks.Add(Link)
+				: OutGeometry.OffMeshLinks.Add(Link);
 		}
 	}
 
@@ -2378,7 +2387,10 @@ int32 FPImplRecastNavMesh::GetTilesDebugGeometry(const FRecastNavMeshGenerator*
 				Link.ValidEnds = GetValidEnds(*DetourNavMesh, Tile, *LinkPoly);
 
 				const int LinkIdx = OutGeometry.OffMeshSegments.Add(Link);
-				OutGeometry.OffMeshSegmentAreas[Link.AreaID].Add(LinkIdx);
+				ensureMsgf((LinkPoly->flags & ForbiddenFlags) == 0, TEXT("Not implemented"));
+				(LinkPoly->flags & ForbiddenFlags) != 0
+					? OutGeometry.OffMeshSegmentAreas[Link.AreaID].Add(LinkIdx)
+					: OutGeometry.OffMeshSegmentAreas[Link.AreaID].Add(LinkIdx);
 			}
 		}
 	}
diff --git a/Engine/Source/Runtime/NavigationSystem/Private/NavMesh/RecastNavMesh.cpp b/Engine/Source/Runtime/NavigationSystem/Private/NavMesh/RecastNavMesh.cpp
--- a/Engine/Source/Runtime/NavigationSystem/Private/NavMesh/RecastNavMesh.cpp
+++ b/Engine/Source/Runtime/NavigationSystem/Private/NavMesh/RecastNavMesh.cpp
@@ -220,6 +220,7 @@ ARecastNavMesh::ARecastNavMesh(const FObjectInitializer& ObjectInitializer)
 	, bDrawNavMeshEdges(true)
 	, bDrawNavLinks(true)
 	, bDrawOctreeDetails(true)
+	, bDrawMarkedForbiddenPolys(false)
 	, bDistinctlyDrawTilesBeingBuilt(true)
 	, bDrawNavMesh(true)
 	, DrawOffset(10.f)
@@ -553,7 +554,7 @@ FColor ARecastNavMesh::GetAreaIDColor(uint8 AreaID) const
 {
 	const UClass* AreaClass = GetAreaClass(AreaID);
 	const UNavArea* DefArea = AreaClass ? ((UClass*)AreaClass)->GetDefaultObject<UNavArea>() : NULL;
-	return DefArea ? DefArea->DrawColor : NavDataConfig.Color;
+	return DefArea ? DefArea->DrawColor : FColor::Red;
 }
 
 void ARecastNavMesh::SortAreasForGenerator(TArray<FRecastAreaNavModifierElement>& Modifiers) const
@@ -1449,6 +1450,73 @@ void ARecastNavMesh::SetPolyArrayArea(const TArray<FNavPoly>& Polys, TSubclassOf
 	}
 }
 
+int32 ARecastNavMesh::ReplaceAreaInTileBounds(const FBox& Bounds, TSubclassOf<UNavArea> OldArea, TSubclassOf<UNavArea> NewArea, bool ReplaceLinks, TArray<NavNodeRef>* OutTouchedNodes)
+{
+	int32 PolysTouched = 0;
+
+	if (RecastNavMeshImpl && RecastNavMeshImpl->GetRecastMesh())
+	{
+		QUICK_SCOPE_CYCLE_COUNTER(STAT_RecastNavMesh_ReplaceAreaInTiles);
+
+		const int32 OldAreaID = GetAreaID(OldArea);
+		ensure(OldAreaID != INDEX_NONE);
+		const int32 NewAreaID = 34;// GetAreaID(NewArea);
+		ensure(NewAreaID != INDEX_NONE);
+		ensure(NewAreaID != OldAreaID);
+
+		// workaround for privacy issue in the recast API
+		dtNavMesh* DetourNavMesh = RecastNavMeshImpl->GetRecastMesh();
+		dtNavMesh const* const ConstDetourNavMesh = RecastNavMeshImpl->GetRecastMesh();
+
+		const FVector RcNavMeshOrigin = Unreal2RecastPoint(NavMeshOriginOffset);
+		const float RcTileSize = FMath::TruncToInt(TileSizeUU / CellSize);
+		const float TileSizeInWorldUnits = RcTileSize * CellSize;
+		const FRcTileBox TileBox(Bounds, RcNavMeshOrigin, TileSizeInWorldUnits);
+
+		for (int32 TileY = TileBox.YMin; TileY <= TileBox.YMax; ++TileY)
+		{
+			for (int32 TileX = TileBox.XMin; TileX <= TileBox.XMax; ++TileX)
+			{
+				const int32 MaxTiles = ConstDetourNavMesh->getTileCountAt(TileX, TileY);
+				if (MaxTiles == 0)
+				{
+					continue;
+				}
+
+				TArray<const dtMeshTile*> Tiles;
+				Tiles.AddZeroed(MaxTiles);
+				const int32 NumTiles = ConstDetourNavMesh->getTilesAt(TileX, TileY, Tiles.GetData(), MaxTiles);
+				for (int32 i = 0; i < NumTiles; i++)
+				{
+					dtTileRef TileRef = ConstDetourNavMesh->getTileRef(Tiles[i]);
+					if (TileRef)
+					{
+						const int32 TileIndex = (int32)ConstDetourNavMesh->decodePolyIdTile(TileRef);
+						const dtMeshTile* Tile = ((const dtNavMesh*)DetourNavMesh)->getTile(TileIndex);
+						//const int32 MaxPolys = Tile && Tile->header ? Tile->header->offMeshBase : 0;
+						const int32 MaxPolys = Tile && Tile->header
+							? (ReplaceLinks ? Tile->header->polyCount : Tile->header->offMeshBase)
+							: 0;
+						if (MaxPolys > 0)
+						{
+							dtPoly* Poly = Tile->polys;
+							for (int32 PolyIndex = 0; PolyIndex < MaxPolys; PolyIndex++, Poly++)
+							{
+								if (Poly->getArea() == OldAreaID)
+								{
+									Poly->setArea(NewAreaID);
+									++PolysTouched;
+								}
+							}
+						}
+					}
+				}
+			}
+		}
+	}
+	return PolysTouched;
+}
+
 bool ARecastNavMesh::GetPolyFlags(NavNodeRef PolyID, uint16& PolyFlags, uint16& AreaFlags) const
 {
 	bool bFound = false;
@@ -2615,5 +2683,17 @@ const dtNavMesh* ARecastNavMesh::GetRecastMesh() const
 {
 	return RecastNavMeshImpl ? RecastNavMeshImpl->GetRecastMesh() : nullptr;
 }
-
 #endif// WITH_RECAST
+
+//----------------------------------------------------------------------//
+// BP API
+//----------------------------------------------------------------------//
+bool ARecastNavMesh::K2_ReplaceAreaInTileBounds(FBox Bounds, TSubclassOf<UNavArea> OldArea, TSubclassOf<UNavArea> NewArea, bool ReplaceLinks)
+{
+	bool bReplaced = ReplaceAreaInTileBounds(Bounds, OldArea, NewArea, ReplaceLinks) > 0;
+	if (bReplaced)
+	{
+		RequestDrawingUpdate();
+	}
+	return bReplaced;
+}
diff --git a/Engine/Source/Runtime/NavigationSystem/Private/NavigationSystem.cpp b/Engine/Source/Runtime/NavigationSystem/Private/NavigationSystem.cpp
--- a/Engine/Source/Runtime/NavigationSystem/Private/NavigationSystem.cpp
+++ b/Engine/Source/Runtime/NavigationSystem/Private/NavigationSystem.cpp
@@ -47,6 +47,7 @@
 #include "NavigationPath.h"
 #include "AbstractNavData.h"
 #include "CrowdManagerBase.h"
+#include "AI/NavigationModifier.h"
 
 
 static const uint32 INITIAL_ASYNC_QUERIES_SIZE = 32;
@@ -2617,6 +2618,11 @@ const FNavigationRelevantData* UNavigationSystemV1::GetDataForObject(const UObje
 	return nullptr;
 }
 
+FNavigationRelevantData* UNavigationSystemV1::GetMutableDataForObject(const UObject& Object)
+{
+	return const_cast<FNavigationRelevantData*>(GetDataForObject(Object));
+}
+
 void UNavigationSystemV1::UpdateActorInNavOctree(AActor& Actor)
 {
 	if (IsNavigationSystemStatic())
@@ -2894,6 +2900,50 @@ bool UNavigationSystemV1::UpdateNavOctreeElementBounds(UActorComponent* Comp, co
 	return false;
 }
 
+bool UNavigationSystemV1::ReplaceAreaInOctreeData(const UObject& Object, TSubclassOf<UNavArea> OldArea, TSubclassOf<UNavArea> NewArea, bool bReplaceChildClasses)
+{
+	FNavigationRelevantData* Data = GetMutableDataForObject(Object);
+	
+	if (Data == nullptr || Data->HasModifiers() == false)
+	{
+		return false;
+	}
+
+	for (FAreaNavModifier& AreaModifier : Data->Modifiers.GetMutableAreas())
+	{ 
+		if (AreaModifier.GetAreaClass() == OldArea
+			|| (bReplaceChildClasses && AreaModifier.GetAreaClass()->IsChildOf(OldArea)))
+		{
+			AreaModifier.SetAreaClass(NewArea);
+		}
+	}
+	for (FSimpleLinkNavModifier& SimpleLink : Data->Modifiers.GetSimpleLinks())
+	{ 
+		for (FNavigationLink& Link : SimpleLink.Links)
+		{
+			if (Link.GetAreaClass() == OldArea
+				|| (bReplaceChildClasses && Link.GetAreaClass()->IsChildOf(OldArea)))
+			{
+				Link.SetAreaClass(NewArea);
+			}
+		}
+		for (FNavigationSegmentLink& Link : SimpleLink.SegmentLinks)
+		{
+			if (Link.GetAreaClass() == OldArea
+				|| (bReplaceChildClasses && Link.GetAreaClass()->IsChildOf(OldArea)))
+			{
+				Link.SetAreaClass(NewArea);
+			}
+		}
+	}
+	for (FCustomLinkNavModifier& CustomLink : Data->Modifiers.GetCustomLinks())
+	{ 
+		ensureMsgf(false, TEXT("Not implemented yet"));
+	}
+
+	return true;
+}
+
 void UNavigationSystemV1::OnComponentRegistered(UActorComponent* Comp)
 {
 	if (IsNavigationSystemStatic())
@@ -3983,6 +4033,11 @@ bool UNavigationSystemV1::IsNavigationBeingBuiltOrLocked(UObject* WorldContextOb
 	return false;
 }
 
+bool UNavigationSystemV1::K2_ReplaceAreaInOctreeData(const UObject* Object, TSubclassOf<UNavArea> OldArea, TSubclassOf<UNavArea> NewArea)
+{
+	return Object ? ReplaceAreaInOctreeData(*Object, OldArea, NewArea) : false;
+}
+
 //----------------------------------------------------------------------//
 // HACKS!!!
 //----------------------------------------------------------------------//
diff --git a/Engine/Source/Runtime/NavigationSystem/Public/NavMesh/NavMeshRenderingComponent.h b/Engine/Source/Runtime/NavigationSystem/Public/NavMesh/NavMeshRenderingComponent.h
--- a/Engine/Source/Runtime/NavigationSystem/Public/NavMesh/NavMeshRenderingComponent.h
+++ b/Engine/Source/Runtime/NavigationSystem/Public/NavMesh/NavMeshRenderingComponent.h
@@ -41,6 +41,7 @@ enum class ENavMeshDetailFlags : uint8
 	Clusters,
 	NavOctree,
 	NavOctreeDetails,
+	MarkForbiddenPolys,
 };
 
 // exported to API for GameplayDebugger module
diff --git a/Engine/Source/Runtime/NavigationSystem/Public/NavMesh/PImplRecastNavMesh.h b/Engine/Source/Runtime/NavigationSystem/Public/NavMesh/PImplRecastNavMesh.h
--- a/Engine/Source/Runtime/NavigationSystem/Public/NavMesh/PImplRecastNavMesh.h
+++ b/Engine/Source/Runtime/NavigationSystem/Public/NavMesh/PImplRecastNavMesh.h
@@ -238,7 +238,10 @@ public:
 	void GetEdgesForPathCorridorImpl(const TArray<NavNodeRef>* PathCorridor, TArray<FNavigationPortalEdge>* PathCorridorEdges, const dtNavMeshQuery& NavQuery) const;
 
 protected:
-	int32 GetTilesDebugGeometry(const FRecastNavMeshGenerator* Generator, const dtMeshTile& Tile, int32 VertBase, FRecastDebugGeometry& OutGeometry, int32 TileIdx = INDEX_NONE) const;
+	/** 
+	 *	@param ForbiddenFlags polys with flags matching the fillter will get added to 
+	 */
+	int32 GetTilesDebugGeometry(const FRecastNavMeshGenerator* Generator, const dtMeshTile& Tile, int32 VertBase, FRecastDebugGeometry& OutGeometry, int32 TileIdx = INDEX_NONE, uint16 ForbiddenFlags = 0) const;
 };
 
 #endif	// WITH_RECAST
diff --git a/Engine/Source/Runtime/NavigationSystem/Public/NavMesh/RecastNavMesh.h b/Engine/Source/Runtime/NavigationSystem/Public/NavMesh/RecastNavMesh.h
--- a/Engine/Source/Runtime/NavigationSystem/Public/NavMesh/RecastNavMesh.h
+++ b/Engine/Source/Runtime/NavigationSystem/Public/NavMesh/RecastNavMesh.h
@@ -166,10 +166,12 @@ struct FRecastDebugGeometry
 
 	TArray<FVector> MeshVerts;
 	TArray<int32> AreaIndices[RECAST_MAX_AREAS];
+	TArray<int32> ForbiddenIndices;
 	TArray<int32> BuiltMeshIndices;
 	TArray<FVector> PolyEdges;
 	TArray<FVector> NavMeshEdges;
 	TArray<FOffMeshLink> OffMeshLinks;
+	TArray<FOffMeshLink> ForbiddenLinks;
 	TArray<FCluster> Clusters;
 	TArray<FClusterLink> ClusterLinks;
 	TArray<FOffMeshSegment> OffMeshSegments;
@@ -177,8 +179,9 @@ struct FRecastDebugGeometry
 
 	int32 bGatherPolyEdges : 1;
 	int32 bGatherNavMeshEdges : 1;
-
-	FRecastDebugGeometry() : bGatherPolyEdges(false), bGatherNavMeshEdges(false)
+	int32 bMarkForbiddenPolys : 1;
+	
+	FRecastDebugGeometry() : bGatherPolyEdges(false), bGatherNavMeshEdges(false), bMarkForbiddenPolys(false)
 	{}
 
 	uint32 NAVIGATIONSYSTEM_API GetAllocatedSize() const;
@@ -333,6 +336,9 @@ class NAVIGATIONSYSTEM_API ARecastNavMesh : public ANavigationData
 	UPROPERTY(EditAnywhere, Category = Display, meta=(editcondition = "bDrawOctree"))
 	uint32 bDrawOctreeDetails : 1;
 
+	UPROPERTY(EditAnywhere, Category = Display)
+	uint32 bDrawMarkedForbiddenPolys : 1;
+
 	UPROPERTY(config)
 	uint32 bDistinctlyDrawTilesBeingBuilt:1;
 
@@ -793,6 +799,10 @@ public:
 	/** Sets area ID for the specified polygons */
 	void SetPolyArrayArea(const TArray<FNavPoly>& Polys, TSubclassOf<UNavArea> AreaClass);
 
+	/** In given Bounds find all areas of class OldArea and replace them with NewArea
+	 *	@return number of polys touched */
+	int32 ReplaceAreaInTileBounds(const FBox& Bounds, TSubclassOf<UNavArea> OldArea, TSubclassOf<UNavArea> NewArea, bool ReplaceLinks = true, TArray<NavNodeRef>* OutTouchedNodes = nullptr);
+	
 	/** Retrieves poly and area flags for specified polygon */
 	bool GetPolyFlags(NavNodeRef PolyID, uint16& PolyFlags, uint16& AreaFlags) const;
 	bool GetPolyFlags(NavNodeRef PolyID, FNavMeshNodeFlags& Flags) const;
@@ -944,6 +954,10 @@ private:
 private:
 	static const FRecastQueryFilter* NamedFilters[ERecastNamedFilter::NamedFiltersCount];
 #endif // WITH_RECAST
+
+	/** @return true if any polygon/link has been touched */
+	UFUNCTION(BlueprintCallable, Category = NavMesh, meta = (DisplayName = "ReplaceAreaInTileBounds"))
+	bool K2_ReplaceAreaInTileBounds(FBox Bounds, TSubclassOf<UNavArea> OldArea, TSubclassOf<UNavArea> NewArea, bool ReplaceLinks = true);
 };
 
 #if WITH_RECAST
diff --git a/Engine/Source/Runtime/NavigationSystem/Public/NavigationSystem.h b/Engine/Source/Runtime/NavigationSystem/Public/NavigationSystem.h
--- a/Engine/Source/Runtime/NavigationSystem/Public/NavigationSystem.h
+++ b/Engine/Source/Runtime/NavigationSystem/Public/NavigationSystem.h
@@ -302,6 +302,9 @@ public:
 	UFUNCTION(BlueprintCallable, Category = "AI|Navigation|Generation")
 	void SetGeometryGatheringMode(ENavDataGatheringModeConfig NewMode);
 
+	UFUNCTION(BlueprintCallable, Category = "AI|Navigation", meta=(DisplayName="ReplaceAreaInOctreeData"))
+	bool K2_ReplaceAreaInOctreeData(const UObject* Object, TSubclassOf<UNavArea> OldArea, TSubclassOf<UNavArea> NewArea);
+
 	FORCEINLINE bool IsActiveTilesGenerationEnabled() const{ return bGenerateNavigationOnlyAroundNavigationInvokers; }
 	
 	/** delegate type for events that dirty the navigation data ( Params: const FBox& DirtyBounds ) */
@@ -601,6 +604,7 @@ public:
 	void RemoveNavOctreeElementId(const FOctreeElementId& ElementId, int32 UpdateFlags);
 
 	const FNavigationRelevantData* GetDataForObject(const UObject& Object) const;
+	FNavigationRelevantData* GetMutableDataForObject(const UObject& Object);
 
 	/** find all elements in navigation octree within given box (intersection) */
 	void FindElementsInNavOctree(const FBox& QueryBox, const FNavigationOctreeFilter& Filter, TArray<FNavigationOctreeElement>& Elements);
@@ -614,6 +618,9 @@ public:
 	/** update component bounds in navigation octree and mark only specified area as dirty, doesn't re-export component geometry */
 	bool UpdateNavOctreeElementBounds(UActorComponent* Comp, const FBox& NewBounds, const FBox& DirtyArea);
 
+	/** fetched Object's data from the octree and replaces occurences of OldArea with NewArea */
+	bool ReplaceAreaInOctreeData(const UObject& Object, TSubclassOf<UNavArea> OldArea, TSubclassOf<UNavArea> NewArea, bool bReplaceChildClasses = false);
+
 	//----------------------------------------------------------------------//
 	// Custom navigation links
 	//----------------------------------------------------------------------//
