Remove temporary code to calm the static analyzer.

#rnx

diff --git a/Engine/Source/Runtime/NavigationSystem/Private/NavMesh/PImplRecastNavMesh.cpp b/Engine/Source/Runtime/NavigationSystem/Private/NavMesh/PImplRecastNavMesh.cpp
--- a/Engine/Source/Runtime/NavigationSystem/Private/NavMesh/PImplRecastNavMesh.cpp
+++ b/Engine/Source/Runtime/NavigationSystem/Private/NavMesh/PImplRecastNavMesh.cpp
@@ -2388,9 +2388,7 @@ int32 FPImplRecastNavMesh::GetTilesDebugGeometry(const FRecastNavMeshGenerator*
 
 				const int LinkIdx = OutGeometry.OffMeshSegments.Add(Link);
 				ensureMsgf((LinkPoly->flags & ForbiddenFlags) == 0, TEXT("Not implemented"));
-				(LinkPoly->flags & ForbiddenFlags) != 0
-					? OutGeometry.OffMeshSegmentAreas[Link.AreaID].Add(LinkIdx)
-					: OutGeometry.OffMeshSegmentAreas[Link.AreaID].Add(LinkIdx);
+				OutGeometry.OffMeshSegmentAreas[Link.AreaID].Add(LinkIdx);
 			}
 		}
 	}
