Deprecation warning fix.

diff --git a/Engine/Source/Runtime/NavigationSystem/Private/NavModifierComponent.cpp b/Engine/Source/Runtime/NavigationSystem/Private/NavModifierComponent.cpp
--- a/Engine/Source/Runtime/NavigationSystem/Private/NavModifierComponent.cpp
+++ b/Engine/Source/Runtime/NavigationSystem/Private/NavModifierComponent.cpp
@@ -28,7 +28,7 @@ void UNavModifierComponent::CalcAndCacheBounds() const
 			// mechanisms won't kick in) but we're binding without checking it since
 			// this property can change without re-running CalcAndCacheBounds.
 			// We're filtering for nav relevancy in OnTransformUpdated.
-			TransformUpdateHandle = MyOwner->GetRootComponent()->TransformUpdated.AddUObject(this, &UNavModifierComponent::OnTransformUpdated);
+			TransformUpdateHandle = MyOwner->GetRootComponent()->TransformUpdated.AddUObject(const_cast<UNavModifierComponent*>(this), &UNavModifierComponent::OnTransformUpdated);
 		}
 
 		Bounds = FBox(ForceInit);
