-Adding additional validation to ensure the element index is properly within the element bounds before accessing.

diff --git a/Engine/Source/Runtime/NavigationSystem/Private/NavigationSystem.cpp b/Engine/Source/Runtime/NavigationSystem/Private/NavigationSystem.cpp
--- a/Engine/Source/Runtime/NavigationSystem/Private/NavigationSystem.cpp
+++ b/Engine/Source/Runtime/NavigationSystem/Private/NavigationSystem.cpp
@@ -2592,7 +2592,7 @@ const FNavigationRelevantData* UNavigationSystemV1::GetDataForObject(const UObje
 
 	const FOctreeElementId* OctreeID = GetObjectsNavOctreeId(Object);
 
-	if (OctreeID != nullptr && OctreeID->IsValidId() == true)
+	if (OctreeID != nullptr && NavOctree->IsValidElementId(*OctreeID))
 	{
 		return NavOctree->GetDataForID(*OctreeID);
 	}
