NavigationOctree: tentative fix to prevent obsolete entries in NavigationSystem ObjectToOctreeId map. We now also consider pending kill objects since we only need them to retrieve the world and to get the map key.

diff --git a/Engine/Source/Runtime/NavigationSystem/Private/NavigationOctree.cpp b/Engine/Source/Runtime/NavigationSystem/Private/NavigationOctree.cpp
--- a/Engine/Source/Runtime/NavigationSystem/Private/NavigationOctree.cpp
+++ b/Engine/Source/Runtime/NavigationSystem/Private/NavigationOctree.cpp
@@ -204,7 +204,9 @@ FORCENOINLINE
 void FNavigationOctreeSemantics::SetElementId(const FNavigationOctreeElement& Element, FOctreeElementId Id)
 {
 	UWorld* World = NULL;
-	UObject* ElementOwner = Element.GetOwner();
+
+	const bool bEvenIfPendingKill = true;
+	UObject* ElementOwner = Element.GetOwner(bEvenIfPendingKill);
 	if (ElementOwner == nullptr)
 	{
 		return;
diff --git a/Engine/Source/Runtime/NavigationSystem/Public/NavigationOctree.h b/Engine/Source/Runtime/NavigationSystem/Public/NavigationOctree.h
--- a/Engine/Source/Runtime/NavigationSystem/Public/NavigationOctree.h
+++ b/Engine/Source/Runtime/NavigationSystem/Public/NavigationOctree.h
@@ -68,7 +68,7 @@ struct NAVIGATIONSYSTEM_API FNavigationOctreeElement
 		Data->ValidateAndShrink();
 	}
 
-	FORCEINLINE UObject* GetOwner() const { return Data->SourceObject.Get(); }
+	FORCEINLINE UObject* GetOwner(bool bEvenIfPendingKill = false) const { return Data->SourceObject.Get(bEvenIfPendingKill); }
 };
 
 struct FNavigationOctreeSemantics
