Fixed StoreCollisionCache in RecastNavMeshGenerator.cpp to properly reset the array GeomExport.Data->CollisionData array.

It used to Reserve rather than Empty which in edge cases left the array at undesided state/size.

Yoan.StAmant

diff --git a/Engine/Source/Runtime/NavigationSystem/Private/NavMesh/RecastNavMeshGenerator.cpp b/Engine/Source/Runtime/NavigationSystem/Private/NavMesh/RecastNavMeshGenerator.cpp
--- a/Engine/Source/Runtime/NavigationSystem/Private/NavMesh/RecastNavMeshGenerator.cpp
+++ b/Engine/Source/Runtime/NavigationSystem/Private/NavMesh/RecastNavMeshGenerator.cpp
@@ -298,8 +298,8 @@ static void StoreCollisionCache(FRecastGeometryExport& GeomExport)
 
 	HeaderInfo.Validation.DataSize = CacheSize;
 
-	// reserve + add combo to allocate exact amount (without any overhead/slack)
-	GeomExport.Data->CollisionData.Reserve(CacheSize);
+	// empty + add combo to allocate exact amount (without any overhead/slack)
+	GeomExport.Data->CollisionData.Empty(CacheSize);
 	GeomExport.Data->CollisionData.AddUninitialized(CacheSize);
 
 	// store collisions
