Extracted NavigationSystemV1's logic concerning NavOctree and navigation DirtyAreas into a separate structures for reusability

Notable changes:
- Added FNavigationOctreeController that wraps up what used to be NavigationSystemV1's navoctree-related logic
- Added FNavigationDirtyAreasController that wraps up what used to be NavigationSystemV1's DirtyAreas-related logic
- Added FNavigationDataHandler that is a helper struct that wraps up what used to be NavigationSystemV1's logic related to operation involving both navoctree and dirty areas
- Deprecated both FNavDataConfig.NavigationDataClass and FNavDataConfig.NavigationDataClassName and replaced them with a single NavDataClass property
- FNavigationOctree is not responsible for hashing element objects and storing their ElementId
- NavOctree elements how know about the octree they belong to (via a member property).

diff --git a/Engine/Config/BaseEngine.ini b/Engine/Config/BaseEngine.ini
--- a/Engine/Config/BaseEngine.ini
+++ b/Engine/Config/BaseEngine.ini
@@ -687,6 +687,9 @@ bUseStreamingPause=false
 +FunctionRedirects=(OldName="NativeUserListEntry.GetOwningListView", NewName="UserListEntryLibrary.GetOwningListView")
 +FunctionRedirects=(OldName="UserObjectListEntry.GetListItemObject", NewName="UserObjectListEntryLibrary.GetListItemObject")
 
+; 4.24
++PropertyRedirects=(OldName="NavDataConfig.NavigationDataClassName", NewName="NavDataConfig.NavDataClass")
+
 [CoreUObject.Metadata]
 ; Note: UnrealHeaderTool should be rerun after making changes to MetadataRedirects to catch any keys specified in class headers
 +MetadataRedirects=(OldKey="K2Protected", NewKey="BlueprintProtected")
diff --git a/Engine/Source/Runtime/Engine/Classes/AI/Navigation/NavigationTypes.h b/Engine/Source/Runtime/Engine/Classes/AI/Navigation/NavigationTypes.h
--- a/Engine/Source/Runtime/Engine/Classes/AI/Navigation/NavigationTypes.h
+++ b/Engine/Source/Runtime/Engine/Classes/AI/Navigation/NavigationTypes.h
@@ -24,7 +24,6 @@
 #define DEFAULT_NAV_QUERY_EXTENT_VERTICAL 250.f
 
 class AActor;
-class ANavigationData;
 class INavAgentInterface;
 class INavRelevantInterface;
 class ULevel;
@@ -49,6 +48,11 @@ namespace FNavigationSystem
 	{
 		return TestLocation != InvalidLocation;
 	}
+
+	FORCEINLINE bool BoxesAreSame(const FBox& A, const FBox& B)
+	{
+		return FVector::PointsAreSame(A.Min, B.Min) && FVector::PointsAreSame(A.Max, B.Max);
+	}
 }
 
 UENUM()
@@ -467,8 +471,6 @@ struct ENGINE_API FNavAgentProperties : public FMovementProperties
 	float NavWalkingSearchHeightScale;
 
 	/** Type of navigation data used by agent, null means "any" */
-	/*UPROPERTY(EditAnywhere, BlueprintReadWrite, Category = MovementProperties)
-	TSubclassOf<ANavigationData> PreferredNavData;*/
 	UPROPERTY(EditAnywhere, BlueprintReadWrite, Category=MovementProperties, meta=(MetaClass = "NavigationData"))
 	FSoftClassPath PreferredNavData;
 	
@@ -506,10 +508,7 @@ struct ENGINE_API FNavAgentProperties : public FMovementProperties
 			: INVALID_NAVEXTENT;
 	}
 
-	void SetPreferredNavData(TSubclassOf<AActor> NavDataClass)
-	{
-		PreferredNavData = FSoftClassPath(NavDataClass.Get());
-	}
+	void SetPreferredNavData(TSubclassOf<AActor> NavDataClass);
 
 	static const FNavAgentProperties DefaultProperties;
 };
@@ -533,15 +532,38 @@ struct ENGINE_API FNavDataConfig : public FNavAgentProperties
 	UPROPERTY(EditAnywhere, BlueprintReadOnly, Category = Querying)
 	FVector DefaultQueryExtent;
 
-	// mz@todo make sure we handle NavigationDataClass and NavigationDataClassName in PostEditChange
+	UE_DEPRECATED(4.24, "FNavDataConfig.NavigationDataClass is deprecated and setting it directly has no effect. Please use setter and getter functions instead.")
 	UPROPERTY(Transient)
-	TSubclassOf<AActor> NavigationDataClass;
+	mutable TSubclassOf<AActor> NavigationDataClass;
 
-	UPROPERTY(VisibleAnywhere, BlueprintReadOnly, Category=Navigation, meta=(MetaClass = "NavigationData"))
+#if WITH_EDITOR
+	// used to be a UPROPERTY, but had to remove it so that it doesn't interfere
+	// with property redirects
+	UE_DEPRECATED(4.24, "FNavDataConfig.NavigationDataClassName is deprecated. Please use setter and getter functions instead.")
 	FSoftClassPath NavigationDataClassName;
+#endif // WITH_EDITOR
+
+protected:
+	UPROPERTY(EditAnywhere, BlueprintReadWrite, Category = Navigation, meta = (MetaClass = "NavigationData"))
+	TSoftClassPtr<AActor> NavDataClass;
 
+public:	
 	FNavDataConfig(float Radius = FNavigationSystem::FallbackAgentRadius, float Height = FNavigationSystem::FallbackAgentHeight);
 	FNavDataConfig(const FNavDataConfig& Other);
+
+	void SetNavDataClass(UClass* InNavDataClass);
+	void SetNavDataClass(TSoftClassPtr<AActor> InNavDataClass);
+	
+	template<typename T>
+	TSubclassOf<T> GetNavDataClass() const { return TSubclassOf<T>(NavDataClass.Get()); }
+
+#if WITH_EDITOR
+	static FName GetNavigationDataClassPropertyName()
+	{
+		static const FName NAME_NavigationDataClass = GET_MEMBER_NAME_CHECKED(FNavDataConfig, NavDataClass);
+		return NAME_NavigationDataClass;
+	}
+#endif // WITH_EDITOR
 };
 
 struct FNavigationProjectionWork
@@ -707,4 +729,3 @@ struct ENGINE_API FNavHeightfieldSamples
 
 	FORCEINLINE bool IsEmpty() const { return Heights.Num() == 0; }
 };
-
diff --git a/Engine/Source/Runtime/Engine/Private/AI/Navigation/NavigationTypes.cpp b/Engine/Source/Runtime/Engine/Private/AI/Navigation/NavigationTypes.cpp
--- a/Engine/Source/Runtime/Engine/Private/AI/Navigation/NavigationTypes.cpp
+++ b/Engine/Source/Runtime/Engine/Private/AI/Navigation/NavigationTypes.cpp
@@ -40,8 +40,7 @@ FNavDataConfig::FNavDataConfig(float Radius, float Height)
 	, Name(TEXT("Default"))
 	, Color(140, 255, 0, 164)
 	, DefaultQueryExtent(DEFAULT_NAV_QUERY_EXTENT_HORIZONTAL, DEFAULT_NAV_QUERY_EXTENT_HORIZONTAL, DEFAULT_NAV_QUERY_EXTENT_VERTICAL)
-	, NavigationDataClass(FNavigationSystem::GetDefaultNavDataClass())
-	, NavigationDataClassName(NavigationDataClass)
+	, NavDataClass(FNavigationSystem::GetDefaultNavDataClass())
 {
 }
 
@@ -50,12 +49,19 @@ FNavDataConfig::FNavDataConfig(const FNavDataConfig& Other)
 	, Name(Other.Name)
 	, Color(Other.Color)
 	, DefaultQueryExtent(Other.DefaultQueryExtent)
-	, NavigationDataClass(Other.NavigationDataClass)
-	, NavigationDataClassName(NavigationDataClass)
+	, NavDataClass(Other.NavDataClass)
 {
+}
 
+void FNavDataConfig::SetNavDataClass(UClass* InNavDataClass)
+{
+	NavDataClass = InNavDataClass;
 }
 
+void FNavDataConfig::SetNavDataClass(TSoftClassPtr<AActor> InNavDataClass)
+{
+	NavDataClass = InNavDataClass;
+}
 //----------------------------------------------------------------------//
 // FNavigationRelevantData
 //----------------------------------------------------------------------//
@@ -248,6 +254,11 @@ bool FNavAgentProperties::IsNavDataMatching(const FNavAgentProperties& Other) co
 	return (PreferredNavData == Other.PreferredNavData || PreferredNavData == nullptr || Other.PreferredNavData == nullptr);
 }
 
+void FNavAgentProperties::SetPreferredNavData(TSubclassOf<AActor> NavDataClass)
+{
+	PreferredNavData = FSoftClassPath(NavDataClass.Get());
+}
+
 //----------------------------------------------------------------------//
 // UNavAreaBase
 //----------------------------------------------------------------------//
diff --git a/Engine/Source/Runtime/NavigationSystem/Private/NavigationData.cpp b/Engine/Source/Runtime/NavigationSystem/Private/NavigationData.cpp
--- a/Engine/Source/Runtime/NavigationSystem/Private/NavigationData.cpp
+++ b/Engine/Source/Runtime/NavigationSystem/Private/NavigationData.cpp
@@ -556,12 +556,12 @@ TArray<FBox> ANavigationData::GetNavigableBounds() const
 	TArray<FBox> Result;
 	const UNavigationSystemV1* NavSys = FNavigationSystem::GetCurrent<const UNavigationSystemV1>(GetWorld());
 	
-	const auto& NavigationBounds = NavSys->GetNavigationBounds();
-	Result.Reserve(NavigationBounds.Num());
-
-	for (const auto& Bounds : NavigationBounds)
+	if (NavSys)
 	{
-		Result.Add(Bounds.AreaBox);
+		// @note this has been switched over from calling GetNavigationBounds
+		// to get navigable bounds relevant to this one nav data instance
+		// This implements the original intension of the function
+		NavSys->GetNavigationBoundsForNavData(*this, Result);
 	}
 	
 	return Result;
@@ -572,20 +572,11 @@ TArray<FBox> ANavigationData::GetNavigableBoundsInLevel(ULevel* InLevel) const
 	TArray<FBox> Result;
 	const UNavigationSystemV1* NavSys = FNavigationSystem::GetCurrent<const UNavigationSystemV1>(GetWorld());
 
-	if (NavSys != nullptr)
+	if (NavSys)
 	{
-		const auto& NavigationBounds = NavSys->GetNavigationBounds();
-		Result.Reserve(NavigationBounds.Num());
-
-		for (const auto& Bounds : NavigationBounds)
-		{
-			if (Bounds.Level == InLevel)
-			{
-				Result.Add(Bounds.AreaBox);
-			}
-		}
+		NavSys->GetNavigationBoundsForNavData(*this, Result, InLevel);
 	}
-	
+
 	return Result;
 }
 
diff --git a/Engine/Source/Runtime/NavigationSystem/Private/NavigationDataHandler.cpp b/Engine/Source/Runtime/NavigationSystem/Private/NavigationDataHandler.cpp
new file mode 100644
--- /dev/null
+++ b/Engine/Source/Runtime/NavigationSystem/Private/NavigationDataHandler.cpp
@@ -0,0 +1,471 @@
+// Copyright 1998-2018 Epic Games, Inc. All Rights Reserved.
+
+#include "NavigationDataHandler.h"
+#include "NavMesh/RecastNavMeshGenerator.h"
+
+
+DEFINE_LOG_CATEGORY_STATIC(LogNavOctree, Warning, All);
+
+namespace
+{
+	int32 GetDirtyFlagHelper(int32 UpdateFlags, int32 DefaultValue)
+	{
+		return ((UpdateFlags & FNavigationOctreeController::OctreeUpdate_Geometry) != 0) ? ENavigationDirtyFlag::All :
+			((UpdateFlags & FNavigationOctreeController::OctreeUpdate_Modifiers) != 0) ? ENavigationDirtyFlag::DynamicModifier :
+			DefaultValue;
+	}
+}
+	
+FNavigationDataHandler::FNavigationDataHandler(FNavigationOctreeController& InOctreeController, FNavigationDirtyAreasController& InDirtyAreasController)
+		: OctreeController(InOctreeController), DirtyAreasController(InDirtyAreasController)
+{}
+
+void FNavigationDataHandler::RemoveNavOctreeElementId(const FOctreeElementId& ElementId, int32 UpdateFlags)
+{
+	if (OctreeController.NavOctree->IsValidElementId(ElementId))
+	{
+		const FNavigationOctreeElement& ElementData = OctreeController.NavOctree->GetElementById(ElementId);
+		const int32 DirtyFlag = GetDirtyFlagHelper(UpdateFlags, ElementData.Data->GetDirtyFlag());
+		// mark area occupied by given actor as dirty
+		DirtyAreasController.AddArea(ElementData.Bounds.GetBox(), DirtyFlag);
+		OctreeController.NavOctree->RemoveNode(ElementId);
+	}
+}
+
+FSetElementId FNavigationDataHandler::RegisterNavOctreeElement(UObject& ElementOwner, INavRelevantInterface& ElementInterface, int32 UpdateFlags)
+{
+	SCOPE_CYCLE_COUNTER(STAT_Navigation_RegisterNavOctreeElement);
+
+	FSetElementId SetId;
+
+	if (OctreeController.NavOctree.IsValid() == false)
+	{
+		return SetId;
+	}
+
+	if (OctreeController.IsNavigationOctreeLocked())
+	{
+		UE_LOG(LogNavOctree, Log, TEXT("IGNORE(RegisterNavOctreeElement) %s"), *GetPathNameSafe(&ElementOwner));
+		return SetId;
+	}
+
+	const bool bIsRelevant = ElementInterface.IsNavigationRelevant();
+	UE_LOG(LogNavOctree, Log, TEXT("REG %s %s"), *ElementOwner.GetName(), bIsRelevant ? TEXT("[relevant]") : TEXT("[skip]"));
+
+	if (bIsRelevant)
+	{
+		bool bCanAdd = false;
+
+		UObject* ParentNode = ElementInterface.GetNavigationParent();
+		if (ParentNode)
+		{
+			OctreeController.OctreeChildNodesMap.AddUnique(ParentNode, FWeakObjectPtr(&ElementOwner));
+			bCanAdd = true;
+		}
+		else
+		{
+			bCanAdd = (OctreeController.HasObjectsNavOctreeId(ElementOwner) == false);
+		}
+
+		if (bCanAdd)
+		{
+			FNavigationDirtyElement UpdateInfo(&ElementOwner, &ElementInterface, GetDirtyFlagHelper(UpdateFlags, 0));
+
+			SetId = OctreeController.PendingOctreeUpdates.FindId(UpdateInfo);
+			if (SetId.IsValidId())
+			{
+				// make sure this request stays, in case it has been invalidated already
+				OctreeController.PendingOctreeUpdates[SetId] = UpdateInfo;
+			}
+			else
+			{
+				SetId = OctreeController.PendingOctreeUpdates.Add(UpdateInfo);
+			}
+		}
+	}
+
+	return SetId;
+}
+
+void FNavigationDataHandler::AddElementToNavOctree(const FNavigationDirtyElement& DirtyElement)
+{
+	check(OctreeController.NavOctree.IsValid());
+
+	// handle invalidated requests first
+	if (DirtyElement.bInvalidRequest)
+	{
+		if (DirtyElement.bHasPrevData)
+		{
+			DirtyAreasController.AddArea(DirtyElement.PrevBounds, DirtyElement.PrevFlags);
+		}
+
+		return;
+	}
+
+	UObject* ElementOwner = DirtyElement.Owner.Get();
+	if (ElementOwner == nullptr || ElementOwner->IsPendingKill() || DirtyElement.NavInterface == nullptr)
+	{
+		return;
+	}
+
+	FNavigationOctreeElement GeneratedData(*OctreeController.NavOctree.Get(), *ElementOwner);
+	const FBox ElementBounds = DirtyElement.NavInterface->GetNavigationBounds();
+
+	UObject* NavigationParent = DirtyElement.NavInterface->GetNavigationParent();
+	if (NavigationParent)
+	{
+		// check if parent node is waiting in queue
+		const FSetElementId ParentRequestId = OctreeController.PendingOctreeUpdates.FindId(FNavigationDirtyElement(NavigationParent));
+		const FOctreeElementId* ParentId = OctreeController.GetObjectsNavOctreeId(*NavigationParent);
+		if (ParentRequestId.IsValidId() && ParentId == nullptr)
+		{
+			FNavigationDirtyElement& ParentNode = OctreeController.PendingOctreeUpdates[ParentRequestId];
+			AddElementToNavOctree(ParentNode);
+
+			// mark as invalid so it won't be processed twice
+			ParentNode.bInvalidRequest = true;
+		}
+
+		const FOctreeElementId* UseParentId = ParentId ? ParentId : OctreeController.GetObjectsNavOctreeId(*NavigationParent);
+		if (UseParentId && OctreeController.NavOctree->IsValidElementId(*UseParentId))
+		{
+			UE_LOG(LogNavOctree, Log, TEXT("ADD %s to %s"), *GetNameSafe(ElementOwner), *GetNameSafe(NavigationParent));
+			OctreeController.NavOctree->AppendToNode(*UseParentId, DirtyElement.NavInterface, ElementBounds, GeneratedData);
+		}
+		else
+		{
+			UE_LOG(LogNavOctree, Warning, TEXT("Can't add node [%s] - parent [%s] not found in octree!"), *GetNameSafe(ElementOwner), *GetNameSafe(NavigationParent));
+		}
+	}
+	else
+	{
+		UE_LOG(LogNavOctree, Log, TEXT("ADD %s"), *GetNameSafe(ElementOwner));
+		OctreeController.NavOctree->AddNode(ElementOwner, DirtyElement.NavInterface, ElementBounds, GeneratedData);
+	}
+
+	const FBox BBox = GeneratedData.Bounds.GetBox();
+	const bool bValidBBox = BBox.IsValid && !BBox.GetSize().IsNearlyZero();
+
+	if (bValidBBox && !GeneratedData.IsEmpty())
+	{
+		const int32 DirtyFlag = DirtyElement.FlagsOverride ? DirtyElement.FlagsOverride : GeneratedData.Data->GetDirtyFlag();
+		DirtyAreasController.AddArea(BBox, DirtyFlag);
+	}
+}
+
+void FNavigationDataHandler::UnregisterNavOctreeElement(UObject& ElementOwner, INavRelevantInterface& ElementInterface, int32 UpdateFlags)
+{
+	SCOPE_CYCLE_COUNTER(STAT_Navigation_UnregisterNavOctreeElement);
+
+	if (OctreeController.NavOctree.IsValid() == false)
+	{
+		return;
+	}
+
+	if (OctreeController.IsNavigationOctreeLocked())
+	{
+#if !WITH_EDITOR
+		UE_LOG(LogNavOctree, Log, TEXT("IGNORE(UnregisterNavOctreeElement) %s"), *GetPathNameSafe(&ElementOwner));
+#endif // WITH_EDITOR
+		return;
+	}
+
+	const FOctreeElementId* ElementId = OctreeController.GetObjectsNavOctreeId(ElementOwner);
+	UE_LOG(LogNavOctree, Log, TEXT("UNREG %s %s"), *ElementOwner.GetName(), ElementId ? TEXT("[exists]") : TEXT("[does\'t exist]"));
+
+	if (ElementId != nullptr)
+	{
+		RemoveNavOctreeElementId(*ElementId, UpdateFlags);
+		OctreeController.RemoveObjectsNavOctreeId(ElementOwner);
+	}
+	else
+	{
+		const bool bCanRemoveChildNode = (UpdateFlags & FNavigationOctreeController::OctreeUpdate_ParentChain) == 0;
+		UObject* ParentNode = ElementInterface.GetNavigationParent();
+		if (ParentNode && bCanRemoveChildNode)
+		{
+			// if node has navigation parent (= doesn't exists in octree on its own)
+			// and it's not part of parent chain update
+			// remove it from map and force update on parent to rebuild octree element
+
+			OctreeController.OctreeChildNodesMap.RemoveSingle(ParentNode, FWeakObjectPtr(&ElementOwner));
+			UpdateNavOctreeParentChain(*ParentNode);
+		}
+	}
+
+	// mark pending update as invalid, it will be dirtied according to currently active settings
+	const bool bCanInvalidateQueue = (UpdateFlags & FNavigationOctreeController::OctreeUpdate_Refresh) == 0;
+	if (bCanInvalidateQueue)
+	{
+		const FSetElementId RequestId = OctreeController.PendingOctreeUpdates.FindId(FNavigationDirtyElement(&ElementOwner));
+		if (RequestId.IsValidId())
+		{
+			OctreeController.PendingOctreeUpdates[RequestId].bInvalidRequest = true;
+		}
+	}
+}
+
+void FNavigationDataHandler::UpdateNavOctreeElement(UObject& ElementOwner, INavRelevantInterface& ElementInterface, int32 UpdateFlags)
+{
+	INC_DWORD_STAT(STAT_Navigation_UpdateNavOctree);
+
+	if (OctreeController.IsNavigationOctreeLocked())
+	{
+		UE_LOG(LogNavOctree, Log, TEXT("IGNORE(UpdateNavOctreeElement) %s"), *ElementOwner.GetPathName());
+		return;
+	}
+
+	// grab existing octree data
+	FBox CurrentBounds;
+	int32 CurrentFlags;
+	const bool bAlreadyExists = OctreeController.GetNavOctreeElementData(ElementOwner, CurrentFlags, CurrentBounds);
+
+	// don't invalidate pending requests
+	UpdateFlags |= FNavigationOctreeController::OctreeUpdate_Refresh;
+
+	// always try to unregister, even if element owner doesn't exists in octree (parent nodes)
+	UnregisterNavOctreeElement(ElementOwner, ElementInterface, UpdateFlags);
+
+	const FSetElementId RequestId = RegisterNavOctreeElement(ElementOwner, ElementInterface, UpdateFlags);
+
+	// add original data to pending registration request
+	// so it could be dirtied properly when system receive unregister request while actor is still queued
+	if (RequestId.IsValidId())
+	{
+		FNavigationDirtyElement& UpdateInfo = OctreeController.PendingOctreeUpdates[RequestId];
+		UpdateInfo.PrevFlags = CurrentFlags;
+		if (UpdateInfo.PrevBounds.IsValid)
+		{
+			// Is we have something stored already we want to 
+			// sum it up, since we care about the whole bounding
+			// box of changes that potentially took place
+			UpdateInfo.PrevBounds += CurrentBounds;
+		}
+		else
+		{
+			UpdateInfo.PrevBounds = CurrentBounds;
+		}
+		UpdateInfo.bHasPrevData = bAlreadyExists;
+	}
+
+	UpdateNavOctreeParentChain(ElementOwner, /*bSkipElementOwnerUpdate=*/ true);
+}
+
+void FNavigationDataHandler::UpdateNavOctreeParentChain(UObject& ElementOwner, bool bSkipElementOwnerUpdate)
+{
+	const int32 UpdateFlags = FNavigationOctreeController::OctreeUpdate_ParentChain | FNavigationOctreeController::OctreeUpdate_Refresh;
+
+	TArray<FWeakObjectPtr> ChildNodes;
+	OctreeController.OctreeChildNodesMap.MultiFind(&ElementOwner, ChildNodes);
+
+	if (ChildNodes.Num() == 0)
+	{
+		if (bSkipElementOwnerUpdate == false)
+		{
+			INavRelevantInterface* ElementInterface = Cast<INavRelevantInterface>(&ElementOwner);
+			if (ElementInterface)
+			{
+				UpdateNavOctreeElement(ElementOwner, *ElementInterface, UpdateFlags);
+			}
+		}
+		return;
+	}
+
+	INavRelevantInterface* ElementInterface = Cast<INavRelevantInterface>(&ElementOwner);
+	TArray<INavRelevantInterface*> ChildNavInterfaces;
+	ChildNavInterfaces.AddZeroed(ChildNodes.Num());
+
+	for (int32 Idx = 0; Idx < ChildNodes.Num(); Idx++)
+	{
+		if (ChildNodes[Idx].IsValid())
+		{
+			UObject* ChildNodeOb = ChildNodes[Idx].Get();
+			ChildNavInterfaces[Idx] = Cast<INavRelevantInterface>(ChildNodeOb);
+			if (ChildNodeOb && ChildNavInterfaces[Idx])
+			{
+				UnregisterNavOctreeElement(*ChildNodeOb, *ChildNavInterfaces[Idx], UpdateFlags);
+			}
+		}
+	}
+
+	if (bSkipElementOwnerUpdate == false && ElementInterface)
+	{
+		UnregisterNavOctreeElement(ElementOwner, *ElementInterface, UpdateFlags);
+		RegisterNavOctreeElement(ElementOwner, *ElementInterface, UpdateFlags);
+	}
+
+	for (int32 Idx = 0; Idx < ChildNodes.Num(); Idx++)
+	{
+		UObject* ChildElement = ChildNodes[Idx].Get();
+		if (ChildElement && ChildNavInterfaces[Idx])
+		{
+			RegisterNavOctreeElement(*ChildElement, *ChildNavInterfaces[Idx], UpdateFlags);
+		}
+	}
+}
+
+bool FNavigationDataHandler::UpdateNavOctreeElementBounds(UActorComponent& Comp, const FBox& NewBounds, const FBox& DirtyArea)
+{
+	const FOctreeElementId* ElementId = OctreeController.GetObjectsNavOctreeId(Comp);
+	if (ElementId && ElementId->IsValidId())
+	{
+		OctreeController.NavOctree->UpdateNode(*ElementId, NewBounds);
+
+		// Add dirty area
+		if (DirtyArea.IsValid)
+		{
+			ElementId = OctreeController.GetObjectsNavOctreeId(Comp);
+			if (ElementId && ElementId->IsValidId())
+			{
+				FNavigationOctreeElement& ElementData = OctreeController.NavOctree->GetElementById(*ElementId);
+				DirtyAreasController.AddArea(DirtyArea, ElementData.Data->GetDirtyFlag());
+			}
+		}
+
+		return true;
+	}
+
+	return false;
+}
+	
+void FNavigationDataHandler::FindElementsInNavOctree(const FBox& QueryBox, const FNavigationOctreeFilter& Filter, TArray<FNavigationOctreeElement>& Elements)
+{
+	if (OctreeController.NavOctree.IsValid() == false)
+	{
+		UE_LOG(LogNavOctree, Warning, TEXT("FNavigationDataHandler::FindElementsInNavOctree gets called while NavOctree is null"));
+		return;
+	}
+
+	for (FNavigationOctree::TConstElementBoxIterator<> It(*OctreeController.NavOctree, QueryBox); It.HasPendingElements(); It.Advance())
+	{
+		const FNavigationOctreeElement& Element = It.GetCurrentElement();
+		if (Element.IsMatchingFilter(Filter))
+		{
+			Elements.Add(Element);
+		}
+	}
+}
+
+bool FNavigationDataHandler::ReplaceAreaInOctreeData(const UObject& Object, TSubclassOf<UNavArea> OldArea, TSubclassOf<UNavArea> NewArea, bool bReplaceChildClasses)
+{
+	FNavigationRelevantData* Data = OctreeController.GetMutableDataForObject(Object);
+
+	if (Data == nullptr || Data->HasModifiers() == false)
+	{
+		return false;
+	}
+
+	for (FAreaNavModifier& AreaModifier : Data->Modifiers.GetMutableAreas())
+	{
+		if (AreaModifier.GetAreaClass() == OldArea
+			|| (bReplaceChildClasses && AreaModifier.GetAreaClass()->IsChildOf(OldArea)))
+		{
+			AreaModifier.SetAreaClass(NewArea);
+		}
+	}
+	for (FSimpleLinkNavModifier& SimpleLink : Data->Modifiers.GetSimpleLinks())
+	{
+		for (FNavigationLink& Link : SimpleLink.Links)
+		{
+			if (Link.GetAreaClass() == OldArea
+				|| (bReplaceChildClasses && Link.GetAreaClass()->IsChildOf(OldArea)))
+			{
+				Link.SetAreaClass(NewArea);
+			}
+		}
+		for (FNavigationSegmentLink& Link : SimpleLink.SegmentLinks)
+		{
+			if (Link.GetAreaClass() == OldArea
+				|| (bReplaceChildClasses && Link.GetAreaClass()->IsChildOf(OldArea)))
+			{
+				Link.SetAreaClass(NewArea);
+			}
+		}
+	}
+	for (FCustomLinkNavModifier& CustomLink : Data->Modifiers.GetCustomLinks())
+	{
+		ensureMsgf(false, TEXT("Not implemented yet"));
+	}
+
+	return true;
+}
+
+void FNavigationDataHandler::AddLevelCollisionToOctree(ULevel& Level)
+{
+#if WITH_RECAST
+	if (OctreeController.NavOctree.IsValid() &&
+		OctreeController.NavOctree->GetNavGeometryStoringMode() == FNavigationOctree::StoreNavGeometry)
+	{
+		const TArray<FVector>* LevelGeom = Level.GetStaticNavigableGeometry();
+		const FOctreeElementId* ElementId = OctreeController.GetObjectsNavOctreeId(Level);
+
+		if (!ElementId && LevelGeom && LevelGeom->Num() > 0)
+		{
+			FNavigationOctreeElement BSPElem(*OctreeController.NavOctree.Get(), Level);
+			FRecastNavMeshGenerator::ExportVertexSoupGeometry(*LevelGeom, *BSPElem.Data);
+
+			const FBox& Bounds = BSPElem.Data->Bounds;
+			if (!Bounds.GetExtent().IsNearlyZero())
+			{
+				OctreeController.NavOctree->AddNode(&Level, nullptr, Bounds, BSPElem);
+				DirtyAreasController.AddArea(Bounds, ENavigationDirtyFlag::All);
+
+				UE_LOG(LogNavOctree, Log, TEXT("ADD %s"), *Level.GetName());
+			}
+		}
+	}
+#endif// WITH_RECAST
+}
+
+void FNavigationDataHandler::RemoveLevelCollisionFromOctree(ULevel& Level)
+{
+	if (OctreeController.NavOctree.IsValid())
+	{
+		const FOctreeElementId* ElementId = OctreeController.GetObjectsNavOctreeId(Level);
+		UE_LOG(LogNavOctree, Log, TEXT("UNREG %s %s"), *Level.GetName(), ElementId ? TEXT("[exists]") : TEXT(""));
+
+		if (ElementId != nullptr)
+		{
+			if (OctreeController.NavOctree->IsValidElementId(*ElementId))
+			{
+				// mark area occupied by given actor as dirty
+				FNavigationOctreeElement& ElementData = OctreeController.NavOctree->GetElementById(*ElementId);
+				DirtyAreasController.AddArea(ElementData.Bounds.GetBox(), ENavigationDirtyFlag::All);
+			}
+
+			OctreeController.NavOctree->RemoveNode(*ElementId);
+			OctreeController.RemoveObjectsNavOctreeId(Level);
+		}
+	}
+}
+
+void FNavigationDataHandler::UpdateActorAndComponentsInNavOctree(AActor& Actor)
+{
+	INavRelevantInterface* NavElement = Cast<INavRelevantInterface>(&Actor);
+	if (NavElement)
+	{
+		UpdateNavOctreeElement(Actor, *NavElement, FNavigationOctreeController::OctreeUpdate_Default);
+	}
+
+	for (UActorComponent* Component : Actor.GetComponents())
+	{
+		INavRelevantInterface* CompNavElement = Cast<INavRelevantInterface>(Component);
+		if (CompNavElement)
+		{
+			// Component != null is implied by successful INavRelevantInterface cast 
+			UpdateNavOctreeElement(*Component, *CompNavElement, FNavigationOctreeController::OctreeUpdate_Default);
+		}
+	}
+}
+
+void FNavigationDataHandler::ProcessPendingOctreeUpdates()
+{
+	if (OctreeController.PendingOctreeUpdates.Num() && OctreeController.NavOctree)
+	{
+		for (TSet<FNavigationDirtyElement>::TIterator It(OctreeController.PendingOctreeUpdates); It; ++It)
+		{
+			AddElementToNavOctree(*It);
+		}
+	}
+	OctreeController.PendingOctreeUpdates.Empty(32);
+}
diff --git a/Engine/Source/Runtime/NavigationSystem/Private/NavigationDirtyAreasController.cpp b/Engine/Source/Runtime/NavigationSystem/Private/NavigationDirtyAreasController.cpp
new file mode 100644
--- /dev/null
+++ b/Engine/Source/Runtime/NavigationSystem/Private/NavigationDirtyAreasController.cpp
@@ -0,0 +1,54 @@
+// Copyright 1998-2019 Epic Games, Inc. All Rights Reserved.
+
+#include "NavigationDirtyAreasController.h"
+#include "NavigationData.h"
+
+//----------------------------------------------------------------------//
+// FNavigationDirtyAreasController
+//----------------------------------------------------------------------//
+FNavigationDirtyAreasController::FNavigationDirtyAreasController()
+	: bCanAccumulateDirtyAreas(true)
+	, bDirtyAreasReportedWhileAccumulationLocked(false)
+{
+
+}
+
+void FNavigationDirtyAreasController::Tick(const float DeltaSeconds, const TArray<ANavigationData*>& NavDataSet, bool bForceRebuilding)
+{
+	DirtyAreasUpdateTime += DeltaSeconds;
+	const bool bCanRebuildNow = bForceRebuilding || (DirtyAreasUpdateFreq != 0.f && DirtyAreasUpdateTime >= (1.0f / DirtyAreasUpdateFreq));
+
+	if (DirtyAreas.Num() > 0 && bCanRebuildNow)
+	{
+		for (ANavigationData* NavData : NavDataSet)
+		{
+			if (NavData)
+			{
+				NavData->RebuildDirtyAreas(DirtyAreas);
+			}
+		}
+
+		DirtyAreasUpdateTime = 0.f;
+		DirtyAreas.Reset();
+	}
+}
+
+void FNavigationDirtyAreasController::AddArea(const FBox& NewArea, int32 Flags)
+{
+	if (Flags > 0 && bCanAccumulateDirtyAreas && NewArea.IsValid)
+	{
+		DirtyAreas.Add(FNavigationDirtyArea(NewArea, Flags));
+	}
+#if !UE_BUILD_SHIPPING
+	bDirtyAreasReportedWhileAccumulationLocked = bDirtyAreasReportedWhileAccumulationLocked || (Flags > 0 && !bCanAccumulateDirtyAreas);
+#endif // !UE_BUILD_SHIPPING
+}
+
+void FNavigationDirtyAreasController::Reset()
+{
+	// discard all pending dirty areas, we are going to rebuild navmesh anyway 
+	DirtyAreas.Reset();
+#if !UE_BUILD_SHIPPING
+	bDirtyAreasReportedWhileAccumulationLocked = false;
+#endif // !UE_BUILD_SHIPPING
+}
diff --git a/Engine/Source/Runtime/NavigationSystem/Private/NavigationOctree.cpp b/Engine/Source/Runtime/NavigationSystem/Private/NavigationOctree.cpp
--- a/Engine/Source/Runtime/NavigationSystem/Private/NavigationOctree.cpp
+++ b/Engine/Source/Runtime/NavigationSystem/Private/NavigationOctree.cpp
@@ -4,6 +4,11 @@
 #include "AI/Navigation/NavRelevantInterface.h"
 #include "NavigationSystem.h"
 
+namespace
+{
+	FNavigationOctree DummyNavOctree(FVector(0), 0);
+}
+
 //----------------------------------------------------------------------//
 // FNavigationOctree
 //----------------------------------------------------------------------//
@@ -20,6 +25,8 @@ FNavigationOctree::~FNavigationOctree()
 {
 	DEC_DWORD_STAT_BY( STAT_NavigationMemory, sizeof(*this) );
 	DEC_MEMORY_STAT_BY(STAT_Navigation_CollisionTreeMemory, NodesMemory);
+	
+	ObjectToOctreeId.Empty();
 }
 
 void FNavigationOctree::SetDataGatheringMode(ENavDataGatheringModeConfig Mode)
@@ -195,6 +202,11 @@ const FNavigationRelevantData* FNavigationOctree::GetDataForID(const FOctreeElem
 	return &*OctreeElement.Data;
 }
 
+void FNavigationOctree::SetElementId(const UObject& Object, FOctreeElementId Id)
+{
+	ObjectToOctreeId.Add(HashObject(Object), Id);
+}
+
 //----------------------------------------------------------------------//
 // FNavigationOctreeSemantics
 //----------------------------------------------------------------------//
@@ -203,31 +215,21 @@ FORCENOINLINE
 #endif // NAVSYS_DEBUG
 void FNavigationOctreeSemantics::SetElementId(const FNavigationOctreeElement& Element, FOctreeElementId Id)
 {
-	UWorld* World = NULL;
 
 	const bool bEvenIfPendingKill = true;
 	UObject* ElementOwner = Element.GetOwner(bEvenIfPendingKill);
-	if (ElementOwner == nullptr)
+	if (ElementOwner)
 	{
-		return;
+		Element.GetOwnerOctree().SetElementId(*ElementOwner, Id);
 	}
+}
 
-	if (AActor* Actor = Cast<AActor>(ElementOwner))
-	{
-		World = Actor->GetWorld();
-	}
-	else if (UActorComponent* AC = Cast<UActorComponent>(ElementOwner))
-	{
-		World = AC->GetWorld();
-	}
-	else if (ULevel* Level = Cast<ULevel>(ElementOwner))
-	{
-		World = Level->OwningWorld;
-	}
+//----------------------------------------------------------------------//
+//  FNavigationOctreeElement
+//----------------------------------------------------------------------//
+FNavigationOctreeElement::FNavigationOctreeElement(UObject& SourceObject)
+	: Data(new FNavigationRelevantData(SourceObject))
+	, OwnerOctree(DummyNavOctree)
+{
 
-	UNavigationSystemV1* NavSys = FNavigationSystem::GetCurrent<UNavigationSystemV1>(World);
-	if (NavSys)
-	{
-		NavSys->SetObjectsNavOctreeId(*ElementOwner, Id);
-	}
-}
+}
\ No newline at end of file
diff --git a/Engine/Source/Runtime/NavigationSystem/Private/NavigationOctreeController.cpp b/Engine/Source/Runtime/NavigationSystem/Private/NavigationOctreeController.cpp
new file mode 100644
--- /dev/null
+++ b/Engine/Source/Runtime/NavigationSystem/Private/NavigationOctreeController.cpp
@@ -0,0 +1,60 @@
+// Copyright 1998-2019 Epic Games, Inc. All Rights Reserved.
+
+#include "NavigationOctreeController.h"
+#include "AI/Navigation/NavRelevantInterface.h"
+#include "NavigationSystem.h"
+
+
+//----------------------------------------------------------------------//
+// FNavigationOctreeController
+//----------------------------------------------------------------------//
+void FNavigationOctreeController::Reset()
+{
+	if (NavOctree.IsValid())
+	{
+		NavOctree->Destroy();
+		NavOctree = NULL;
+	}
+	PendingOctreeUpdates.Empty(32);
+}
+
+void FNavigationOctreeController::SetNavigableGeometryStoringMode(FNavigationOctree::ENavGeometryStoringMode NavGeometryMode)
+{
+	check(NavOctree.IsValid());
+	NavOctree->SetNavigableGeometryStoringMode(NavGeometryMode);
+}
+
+bool FNavigationOctreeController::GetNavOctreeElementData(const UObject& NodeOwner, int32& DirtyFlags, FBox& DirtyBounds)
+{
+	const FOctreeElementId* ElementId = GetObjectsNavOctreeId(NodeOwner);
+	if (ElementId != NULL)
+	{
+		if (NavOctree->IsValidElementId(*ElementId))
+		{
+			// mark area occupied by given actor as dirty
+			FNavigationOctreeElement& ElementData = NavOctree->GetElementById(*ElementId);
+			DirtyFlags = ElementData.Data->GetDirtyFlag();
+			DirtyBounds = ElementData.Bounds.GetBox();
+			return true;
+		}
+	}
+
+	return false;
+}
+
+const FNavigationRelevantData* FNavigationOctreeController::GetDataForObject(const UObject& Object) const
+{
+	const FOctreeElementId* OctreeID = GetObjectsNavOctreeId(Object);
+
+	if (OctreeID != nullptr && ensure(NavOctree.IsValid()) && NavOctree->IsValidElementId(*OctreeID))
+	{
+		return NavOctree->GetDataForID(*OctreeID);
+	}
+
+	return nullptr;
+}
+
+FNavigationRelevantData* FNavigationOctreeController::GetMutableDataForObject(const UObject& Object)
+{
+	return const_cast<FNavigationRelevantData*>(GetDataForObject(Object));
+}
\ No newline at end of file
diff --git a/Engine/Source/Runtime/NavigationSystem/Private/NavigationSystem.cpp b/Engine/Source/Runtime/NavigationSystem/Private/NavigationSystem.cpp
--- a/Engine/Source/Runtime/NavigationSystem/Private/NavigationSystem.cpp
+++ b/Engine/Source/Runtime/NavigationSystem/Private/NavigationSystem.cpp
@@ -1,6 +1,7 @@
 // Copyright 1998-2019 Epic Games, Inc. All Rights Reserved.
 
 #include "NavigationSystem.h"
+#include "NavigationDataHandler.h"
 #include "Misc/ScopeLock.h"
 #include "Stats/StatsMisc.h"
 #include "Modules/ModuleManager.h"
@@ -33,6 +34,7 @@
 #include "EditorModes.h"
 #include "Editor/GeometryMode/Public/GeometryEdMode.h"
 #include "Editor/GeometryMode/Public/EditorGeometry.h"
+#include "Editor/LevelEditor/Public/LevelEditor.h"
 #endif
 
 #if WITH_HOT_RELOAD
@@ -55,9 +57,6 @@ static const uint32 REGISTRATION_QUEUE_SIZE = 16;	// and we'll not reallocate
 
 #define LOCTEXT_NAMESPACE "Navigation"
 
-DEFINE_LOG_CATEGORY_STATIC(LogNavOctree, Warning, All);
-
-
 DECLARE_CYCLE_STAT(TEXT("Nav Tick: mark dirty"), STAT_Navigation_TickMarkDirty, STATGROUP_Navigation);
 DECLARE_CYCLE_STAT(TEXT("Nav Tick: async build"), STAT_Navigation_TickAsyncBuild, STATGROUP_Navigation);
 DECLARE_CYCLE_STAT(TEXT("Nav Tick: async pathfinding"), STAT_Navigation_TickAsyncPathfinding, STATGROUP_Navigation);
@@ -248,17 +247,16 @@ UNavigationSystemV1::UNavigationSystemV1(const FObjectInitializer& ObjectInitial
 	, bTickWhilePaused(false)
 	, bWholeWorldNavigable(false)
 	, bSkipAgentHeightCheckWhenPickingNavData(false)
-	, DirtyAreasUpdateFreq(60)
 	, OperationMode(FNavigationSystemRunMode::InvalidMode)
-	, NavOctree(NULL)
 	, NavBuildingLockFlags(0)
 	, InitialNavBuildingLockFlags(0)
-	, bNavOctreeLock(false)
 	, bInitialSetupHasBeenPerformed(false)
 	, bInitialLevelsAdded(false)
 	, bWorldInitDone(false)
+	, bCleanUpDone(false)
 	, CurrentlyDrawnNavDataIndex(0)
-	, DirtyAreasUpdateTime(0)
+	, NavOctree(nullptr)
+	, bNavOctreeLock(false)
 {
 #if WITH_EDITOR
 	NavUpdateLockFlags = 0;
@@ -394,11 +392,7 @@ UNavigationSystemV1::UNavigationSystemV1(const FObjectInitializer& ObjectInitial
 	ActiveTilesUpdateInterval = 1.f;
 	bGenerateNavigationOnlyAroundNavigationInvokers = false;
 	DataGatheringMode = ENavDataGatheringModeConfig::Instant;
-	bCanAccumulateDirtyAreas = true;
 	bShouldDiscardSubLevelNavData = true;
-#if !UE_BUILD_SHIPPING
-	bDirtyAreasReportedWhileAccumulationLocked = false;
-#endif // !UE_BUILD_SHIPPING
 
 	if (HasAnyFlags(RF_ClassDefaultObject) == false)
 	{
@@ -430,6 +424,7 @@ UNavigationSystemV1::UNavigationSystemV1(const FObjectInitializer& ObjectInitial
 #endif // WITH_EDITOR
 }
 
+PRAGMA_DISABLE_DEPRECATION_WARNINGS
 UNavigationSystemV1::~UNavigationSystemV1()
 {
 	CleanUp(FNavigationSystem::ECleanupMode::CleanupUnsafe);
@@ -446,6 +441,7 @@ UNavigationSystemV1::~UNavigationSystemV1()
 	FCoreDelegates::OnGetOnScreenMessages.RemoveAll(this);
 #endif // !UE_BUILD_SHIPPING
 }
+PRAGMA_ENABLE_DEPRECATION_WARNINGS
 
 void UNavigationSystemV1::ConfigureAsStatic()
 {
@@ -494,8 +490,8 @@ void UNavigationSystemV1::UpdateAbstractNavData()
 	if (AbstractNavData == NULL)
 	{
 		FNavDataConfig DummyConfig;
-		DummyConfig.NavigationDataClass = AAbstractNavData::StaticClass();
-		AbstractNavData = CreateNavigationDataInstance(DummyConfig);
+		DummyConfig.SetNavDataClass(AAbstractNavData::StaticClass());
+		AbstractNavData = CreateNavigationDataInstanceInLevel(DummyConfig, nullptr);
 		if (AbstractNavData)
 		{
 			AbstractNavData->SetFlags(RF_Transient);
@@ -506,7 +502,7 @@ void UNavigationSystemV1::UpdateAbstractNavData()
 void UNavigationSystemV1::SetSupportedAgentsNavigationClass(int32 AgentIndex, TSubclassOf<ANavigationData> NavigationDataClass)
 {
 	check(SupportedAgents.IsValidIndex(AgentIndex));
-	SupportedAgents[AgentIndex].NavigationDataClass = NavigationDataClass;
+	SupportedAgents[AgentIndex].SetNavDataClass(NavigationDataClass);
 
 	// keep preferred navigation data class in sync with actual class
 	// this will be passed to navigation data actor and will be required
@@ -515,15 +511,6 @@ void UNavigationSystemV1::SetSupportedAgentsNavigationClass(int32 AgentIndex, TS
 	// "Any" navigation data preference is valid only for instanced agents
 	SupportedAgents[AgentIndex].SetPreferredNavData(NavigationDataClass);
 
-	if (NavigationDataClass != nullptr)
-	{
-		SupportedAgents[AgentIndex].NavigationDataClassName = FSoftClassPath::GetOrCreateIDForClass(NavigationDataClass);
-	}
-	else
-	{
-		SupportedAgents[AgentIndex].NavigationDataClassName.Reset();
-	}
-
 #if WITH_EDITOR
 	if (GIsEditor)
 	{
@@ -564,20 +551,7 @@ void UNavigationSystemV1::PostInitProperties()
 			for (int32 AgentIndex = 0; AgentIndex < SupportedAgents.Num(); ++AgentIndex)
 			{
 				FNavDataConfig& SupportedAgentConfig = SupportedAgents[AgentIndex];
-				// a piece of legacy maintanance 
-				if (SupportedAgentConfig.NavigationDataClass != nullptr && SupportedAgentConfig.NavigationDataClassName.IsValid() == false)
-				{
-					// fill NavigationDataClassName
-					SupportedAgentConfig.NavigationDataClassName = FSoftClassPath(SupportedAgentConfig.NavigationDataClass);
-				}
-				else
-				{
-					TSubclassOf<ANavigationData> NavigationDataClass = SupportedAgentConfig.NavigationDataClassName.IsValid()
-						? LoadClass<ANavigationData>(NULL, *SupportedAgentConfig.NavigationDataClassName.ToString(), NULL, LOAD_None, NULL)
-						: nullptr;
-
-					SetSupportedAgentsNavigationClass(AgentIndex, NavigationDataClass);
-				}
+				SetSupportedAgentsNavigationClass(AgentIndex, SupportedAgentConfig.GetNavDataClass<ANavigationData>());
 			}
 		}
 	
@@ -607,69 +581,70 @@ void UNavigationSystemV1::PostInitProperties()
 	}
 }
 
+void UNavigationSystemV1::ConstructNavOctree()
+{
+	DefaultOctreeController.Reset();
+	DefaultOctreeController.NavOctree = MakeShareable(new FNavigationOctree(FVector(0, 0, 0), 64000));
+	DefaultOctreeController.NavOctree->SetDataGatheringMode(DataGatheringMode);
+}
+
 bool UNavigationSystemV1::ConditionalPopulateNavOctree()
 {
 	// Discard all navigation updates caused by octree construction
-	TGuardValue<TArray<FNavigationDirtyArea>> DirtyGuard(DirtyAreas, TArray<FNavigationDirtyArea>());
+	TGuardValue<TArray<FNavigationDirtyArea>> DirtyGuard(DefaultDirtyAreasController.DirtyAreas, TArray<FNavigationDirtyArea>());
 
-	// We are going to fully re-populate NavOctree so all pending update request are outdated
-	PendingOctreeUpdates.Empty(32);
-	
-	// Discard current octree
-	DestroyNavOctree();
-	
 	// See if any of registered navigation data need navoctree
 	bSupportRebuilding = RequiresNavOctree();
 
 	if (bSupportRebuilding)
 	{
-		NavOctree = MakeShareable(new FNavigationOctree(FVector(0,0,0), 64000));
-		NavOctree->SetDataGatheringMode(DataGatheringMode);
-		
-		const ERuntimeGenerationType RuntimeGenerationType = GetRuntimeGenerationType();
-		const bool bStoreNavGeometry = (RuntimeGenerationType == ERuntimeGenerationType::Dynamic);
-		NavOctree->SetNavigableGeometryStoringMode(bStoreNavGeometry ? FNavigationOctree::StoreNavGeometry : FNavigationOctree::SkipNavGeometry);
-		if (bStoreNavGeometry)
+		ConstructNavOctree();
+		if (DefaultOctreeController.IsValid())
 		{
+			const ERuntimeGenerationType RuntimeGenerationType = GetRuntimeGenerationType();
+			const bool bStoreNavGeometry = (RuntimeGenerationType == ERuntimeGenerationType::Dynamic);
+			DefaultOctreeController.SetNavigableGeometryStoringMode(bStoreNavGeometry ? FNavigationOctree::StoreNavGeometry : FNavigationOctree::SkipNavGeometry);
+			if (bStoreNavGeometry)
+			{
 #if WITH_RECAST
-			NavOctree->ComponentExportDelegate = FNavigationOctree::FNavigableGeometryComponentExportDelegate::CreateStatic(&FRecastNavMeshGenerator::ExportComponentGeometry);
+				DefaultOctreeController.NavOctree->ComponentExportDelegate = FNavigationOctree::FNavigableGeometryComponentExportDelegate::CreateStatic(&FRecastNavMeshGenerator::ExportComponentGeometry);
 #endif // WITH_RECAST
-		}
-
-		if (!IsNavigationOctreeLocked())
-		{
-			UWorld* World = GetWorld();
-			check(World);
+			}
 
-			// now process all actors on all levels
-			for (int32 LevelIndex = 0; LevelIndex < World->GetNumLevels(); ++LevelIndex)
+			if (!DefaultOctreeController.IsNavigationOctreeLocked())
 			{
-				ULevel* Level = World->GetLevel(LevelIndex);
-				if (ensure(Level))
+				UWorld* World = GetWorld();
+				check(World);
+
+				// now process all actors on all levels
+				for (int32 LevelIndex = 0; LevelIndex < World->GetNumLevels(); ++LevelIndex)
 				{
-					AddLevelToOctree(*Level);
+					ULevel* Level = World->GetLevel(LevelIndex);
+					if (ensure(Level))
+					{
+						AddLevelToOctree(*Level);
+					}
 				}
 			}
 		}
 	}
-	
-	// Add all found elements to octree, this will not add new dirty areas to navigation
-	if (PendingOctreeUpdates.Num())
+	else
 	{
-		for (TSet<FNavigationDirtyElement>::TIterator It(PendingOctreeUpdates); It; ++It)
-		{
-			AddElementToNavOctree(*It);
-		}
-		PendingOctreeUpdates.Empty(32);
+		// Discard current octree along with pending updates
+		DestroyNavOctree();
 	}
 
+	// Add all found elements to octree, this will not add new dirty areas to navigation
+	FNavigationDataHandler NavHandler(DefaultOctreeController, DefaultDirtyAreasController);
+	NavHandler.ProcessPendingOctreeUpdates();
+
 	return bSupportRebuilding;
 }
 
 #if WITH_EDITOR
 void UNavigationSystemV1::PostEditChangeChainProperty(FPropertyChangedChainEvent& PropertyChangedEvent)
 {
-	static const FName NAME_NavigationDataClass = GET_MEMBER_NAME_CHECKED(FNavDataConfig, NavigationDataClass);
+	static const FName NAME_NavDataClass = FNavDataConfig::GetNavigationDataClassPropertyName();
 	static const FName NAME_SupportedAgents = GET_MEMBER_NAME_CHECKED(UNavigationSystemV1, SupportedAgents);
 	static const FName NAME_AllowClientSideNavigation = GET_MEMBER_NAME_CHECKED(UNavigationSystemV1, bAllowClientSideNavigation);
 
@@ -678,13 +653,13 @@ void UNavigationSystemV1::PostEditChangeChainProperty(FPropertyChangedChainEvent
 	if (PropertyChangedEvent.Property)
 	{
 		FName PropName = PropertyChangedEvent.Property->GetFName();
-		if (PropName == NAME_NavigationDataClass)
+		if (PropName == NAME_NavDataClass)
 		{
 			int32 SupportedAgentIndex = PropertyChangedEvent.GetArrayIndex(NAME_SupportedAgents.ToString());
 			if (SupportedAgents.IsValidIndex(SupportedAgentIndex))
 			{
 				// reflect the change to SupportedAgent's 
-				TSubclassOf<ANavigationData> NavClass = SupportedAgents[SupportedAgentIndex].NavigationDataClass.Get();
+				TSubclassOf<ANavigationData> NavClass = SupportedAgents[SupportedAgentIndex].GetNavDataClass<ANavigationData>();
 				SetSupportedAgentsNavigationClass(SupportedAgentIndex, NavClass);
 				SaveConfig();
 			}
@@ -710,9 +685,9 @@ void UNavigationSystemV1::PostEditChangeProperty(FPropertyChangedEvent& Property
 		FName PropName = PropertyChangedEvent.Property->GetFName();
 		if (PropName == NAME_EnableActiveTiles)
 		{
-			if (NavOctree.IsValid())
+			if (DefaultOctreeController.NavOctree.IsValid())
 			{
-				NavOctree->SetDataGatheringMode(DataGatheringMode);
+				DefaultOctreeController.NavOctree->SetDataGatheringMode(DataGatheringMode);
 			}
 
 			for (auto NavData : NavDataSet)
@@ -732,6 +707,11 @@ void UNavigationSystemV1::OnInitializeActors()
 	
 }
 
+void UNavigationSystemV1::OnBeginTearingDown()
+{
+	DestroyNavOctree();
+}
+
 void UNavigationSystemV1::OnWorldInitDone(FNavigationSystemRunMode Mode)
 {
 	static const bool bSkipRebuildInEditor = true;
@@ -739,6 +719,9 @@ void UNavigationSystemV1::OnWorldInitDone(FNavigationSystemRunMode Mode)
 	DoInitialSetup();
 	
 	UWorld* World = GetWorld();
+	check(World);
+
+	World->OnBeginTearingDown().AddUObject(this, &UNavigationSystemV1::OnBeginTearingDown);
 
 	if (IsThereAnywhereToBuildNavigation() == false
 		// Simulation mode is a special case - better not do it in this case
@@ -794,7 +777,7 @@ void UNavigationSystemV1::OnWorldInitDone(FNavigationSystemRunMode Mode)
 					ANavigationData* NavData = (*It);
 					if (NavData != NULL)
 					{
-						ERegistrationResult Result = RegisterNavData(NavData);
+						const ERegistrationResult Result = RegisterNavData(NavData);
 
 						if (Result == RegistrationSuccessful)
 						{
@@ -869,9 +852,9 @@ void UNavigationSystemV1::OnWorldInitDone(FNavigationSystemRunMode Mode)
 	}
 #endif
 
-	if (!bCanAccumulateDirtyAreas)
+	if (!DefaultDirtyAreasController.bCanAccumulateDirtyAreas)
 	{
-		DirtyAreas.Empty();
+		DefaultDirtyAreasController.DirtyAreas.Empty();
 	}
 
 	bWorldInitDone = true;
@@ -943,8 +926,6 @@ void UNavigationSystemV1::Tick(float DeltaSeconds)
 	{
 		return;
 	}
-
-	const bool bIsGame = World->IsGameWorld();
 	
 	if (PendingCustomLinkRegistration.Num())
 	{
@@ -957,7 +938,7 @@ void UNavigationSystemV1::Tick(float DeltaSeconds)
 		PendingNavBoundsUpdates.Reset();
 	}
 
-	if (PendingOctreeUpdates.Num() > 0)
+	if (DefaultOctreeController.PendingOctreeUpdates.Num() > 0)
 	{
 		SCOPE_CYCLE_COUNTER(STAT_Navigation_AddingActorsToNavOctree);
 
@@ -965,17 +946,12 @@ void UNavigationSystemV1::Tick(float DeltaSeconds)
 		STAT(double ThisTime = 0);
 		{
 			SCOPE_SECONDS_COUNTER(ThisTime);
-			for (TSet<FNavigationDirtyElement>::TIterator It(PendingOctreeUpdates); It; ++It)
-			{
-				AddElementToNavOctree(*It);
-			}
-			PendingOctreeUpdates.Empty(32);
+			FNavigationDataHandler NavHandler(DefaultOctreeController, DefaultDirtyAreasController);
+			NavHandler.ProcessPendingOctreeUpdates();
 		}
 		INC_FLOAT_STAT_BY(STAT_Navigation_CumulativeBuildTime,(float)ThisTime*1000);
 	}
 		
-	DirtyAreasUpdateTime += DeltaSeconds;
-
 	if (IsNavigationBuildingLocked() == false)
 	{
 		if (bGenerateNavigationOnlyAroundNavigationInvokers)
@@ -983,18 +959,7 @@ void UNavigationSystemV1::Tick(float DeltaSeconds)
 			UpdateInvokers();
 		}
 
-		{
-			SCOPE_CYCLE_COUNTER(STAT_Navigation_TickMarkDirty);
-
-			const float DirtyAreasUpdateDeltaTime = 1.0f / DirtyAreasUpdateFreq;
-			const bool bCanRebuildNow = (DirtyAreasUpdateTime >= DirtyAreasUpdateDeltaTime) || !bIsGame;
-
-			if (DirtyAreas.Num() > 0 && bCanRebuildNow)
-			{
-				RebuildDirtyAreas();
-				DirtyAreasUpdateTime = 0;
-			}
-		}
+		RebuildDirtyAreas(DeltaSeconds);
 
 		// Tick navigation mesh async builders
 		if (bAsyncBuildPaused == false)
@@ -1048,7 +1013,7 @@ void UNavigationSystemV1::SetNavigationAutoUpdateEnabled(bool bNewEnable, UNavig
 		if (NavSystem)
 		{
 			const bool bCurrentIsEnabled = NavSystem->GetIsAutoUpdateEnabled();
-			NavSystem->bCanAccumulateDirtyAreas = bCurrentIsEnabled
+			NavSystem->DefaultDirtyAreasController.bCanAccumulateDirtyAreas = bCurrentIsEnabled
 				|| (NavSystem->OperationMode != FNavigationSystemRunMode::EditorMode && NavSystem->OperationMode != FNavigationSystemRunMode::InvalidMode);
 
 			if (bCurrentIsEnabled)
@@ -1059,7 +1024,7 @@ void UNavigationSystemV1::SetNavigationAutoUpdateEnabled(bool bNewEnable, UNavig
 			else
 			{
 #if !UE_BUILD_SHIPPING
-				NavSystem->bDirtyAreasReportedWhileAccumulationLocked = false;
+				NavSystem->DefaultDirtyAreasController.bDirtyAreasReportedWhileAccumulationLocked = false;
 #endif // !UE_BUILD_SHIPPING
 				NavSystem->AddNavigationBuildLock(ENavigationBuildLock::NoUpdateInEditor);
 			}
@@ -1077,7 +1042,7 @@ FPathFindingResult UNavigationSystemV1::FindPathSync(const FNavAgentProperties&
 
 	if (Query.NavData.IsValid() == false)
 	{
-		Query.NavData = GetNavDataForProps(AgentProperties);
+		Query.NavData = GetNavDataForProps(AgentProperties, Query.StartLocation);
 	}
 
 	FPathFindingResult Result(ENavigationQueryResult::Error);
@@ -1158,7 +1123,7 @@ uint32 UNavigationSystemV1::FindPathAsync(const FNavAgentProperties& AgentProper
 
 	if (Query.NavData.IsValid() == false)
 	{
-		Query.NavData = GetNavDataForProps(AgentProperties);
+		Query.NavData = GetNavDataForProps(AgentProperties, Query.StartLocation);
 	}
 
 	if (Query.NavData.IsValid())
@@ -1398,7 +1363,7 @@ UNavigationPath* UNavigationSystemV1::FindPathToLocationSynchronously(UObject* W
 			if (NavAgent != NULL)
 			{
 				const FNavAgentProperties& AgentProps = NavAgent->GetNavAgentPropertiesRef();
-				NavigationData = NavSys->GetNavDataForProps(AgentProps);
+				NavigationData = NavSys->GetNavDataForProps(AgentProps, PathStart);
 				bValidPathContext = true;
 			}
 			else if (Cast<ANavigationData>(PathfindingContext))
@@ -1453,7 +1418,7 @@ bool UNavigationSystemV1::NavigationRaycast(UObject* WorldContextObject, const F
 		if (MyNavAgent)
 		{
 			const FNavAgentProperties& AgentProps = MyNavAgent->GetNavAgentPropertiesRef();
-			NavData = NavSys->GetNavDataForProps(AgentProps);
+			NavData = NavSys->GetNavDataForProps(AgentProps, RayStart);
 		}
 		if (NavData == NULL)
 		{
@@ -1474,6 +1439,11 @@ void UNavigationSystemV1::GetNavAgentPropertiesArray(TArray<FNavAgentProperties>
 	AgentToNavDataMap.GetKeys(OutNavAgentProperties);
 }
 
+ANavigationData* UNavigationSystemV1::GetNavDataForProps(const FNavAgentProperties& AgentProperties, const FVector& AgentLocation, const FVector& Extent) const
+{
+	return const_cast<ANavigationData*>(GetNavDataForProps(AgentProperties));
+}
+
 ANavigationData* UNavigationSystemV1::GetNavDataForProps(const FNavAgentProperties& AgentProperties)
 {
 	const UNavigationSystemV1* ConstThis = AsConst(this);
@@ -1917,7 +1887,8 @@ UNavigationSystemV1::ERegistrationResult UNavigationSystemV1::RegisterNavData(AN
 
 				for (int32 AgentIndex = 0; AgentIndex < SupportedAgents.Num(); ++AgentIndex)
 				{
-					if (NavData->GetClass() == SupportedAgents[AgentIndex].NavigationDataClass && SupportedAgents[AgentIndex].IsEquivalent(NavConfig) == true)
+					if (NavData->GetClass() == SupportedAgents[AgentIndex].GetNavDataClass<ANavigationData>()
+						&& SupportedAgents[AgentIndex].IsEquivalent(NavConfig) == true)
 					{
 						// it's supported, then just in case it's not a precise match (IsEquivalent succeeds with some precision) 
 						// update NavData with supported Agent
@@ -2015,7 +1986,7 @@ void UNavigationSystemV1::RegisterCustomLink(INavLinkCustomInterface& CustomLink
 			LinkBounds += WorldPtA;
 			LinkBounds += WorldPtB;
 
-			AddDirtyArea(LinkBounds, OctreeUpdate_Modifiers);
+			AddDirtyArea(LinkBounds, FNavigationOctreeController::OctreeUpdate_Modifiers);
 		}
 	}
 
@@ -2343,17 +2314,16 @@ ANavigationData* UNavigationSystemV1::GetNavDataWithID(const uint16 NavDataID) c
 
 void UNavigationSystemV1::AddDirtyArea(const FBox& NewArea, int32 Flags)
 {
-	if (Flags > 0 && bCanAccumulateDirtyAreas && NewArea.IsValid)
-	{
-		DirtyAreas.Add(FNavigationDirtyArea(NewArea, Flags));
-	}
-#if !UE_BUILD_SHIPPING
-	bDirtyAreasReportedWhileAccumulationLocked = bDirtyAreasReportedWhileAccumulationLocked || (Flags > 0 && !bCanAccumulateDirtyAreas);
-#endif // !UE_BUILD_SHIPPING
+	DefaultDirtyAreasController.AddArea(NewArea, Flags);
 }
 
 void UNavigationSystemV1::AddDirtyAreas(const TArray<FBox>& NewAreas, int32 Flags)
 { 
+	if (Flags == 0)
+	{
+		return;
+	}
+
 	for (int32 NewAreaIndex = 0; NewAreaIndex < NewAreas.Num(); NewAreaIndex++)
 	{
 		AddDirtyArea(NewAreas[NewAreaIndex], Flags);
@@ -2362,247 +2332,47 @@ void UNavigationSystemV1::AddDirtyAreas(const TArray<FBox>& NewAreas, int32 Flag
 
 bool UNavigationSystemV1::HasDirtyAreasQueued() const
 {
-	return DirtyAreas.Num() > 0;
-}
-
-int32 GetDirtyFlagHelper(int32 UpdateFlags, int32 DefaultValue)
-{
-	return ((UpdateFlags & UNavigationSystemV1::OctreeUpdate_Geometry) != 0) ? ENavigationDirtyFlag::All :
-		((UpdateFlags & UNavigationSystemV1::OctreeUpdate_Modifiers) != 0) ? ENavigationDirtyFlag::DynamicModifier :		
-		DefaultValue;
+	return DefaultDirtyAreasController.IsDirty();
 }
 
 FSetElementId UNavigationSystemV1::RegisterNavOctreeElement(UObject* ElementOwner, INavRelevantInterface* ElementInterface, int32 UpdateFlags)
 {
-	SCOPE_CYCLE_COUNTER(STAT_Navigation_RegisterNavOctreeElement);
-
-	FSetElementId SetId;
-
-#if WITH_EDITOR
-	if (IsNavigationRegisterLocked())
-	{
-		return SetId;
-	}
-#endif
-
-	if (NavOctree.IsValid() == false || ElementOwner == NULL || ElementInterface == NULL)
-	{
-		return SetId;
-	}
-
-	if (IsNavigationOctreeLocked())
-	{
-		UE_LOG(LogNavOctree, Log, TEXT("IGNORE(RegisterNavOctreeElement) %s"), *GetPathNameSafe(ElementOwner));
-		return SetId;
-	}
-
-	const bool bIsRelevant = ElementInterface->IsNavigationRelevant();
-	UE_LOG(LogNavOctree, Log, TEXT("REG %s %s"), *GetNameSafe(ElementOwner), bIsRelevant ? TEXT("[relevant]") : TEXT(""));
-
-	if (bIsRelevant)
-	{
-		bool bCanAdd = false;
-
-		UObject* ParentNode = ElementInterface->GetNavigationParent();
-		if (ParentNode)
-		{
-			OctreeChildNodesMap.AddUnique(ParentNode, FWeakObjectPtr(ElementOwner));
-			bCanAdd = true;
-		}
-		else
-		{
-			const FOctreeElementId* ElementId = GetObjectsNavOctreeId(*ElementOwner);
-			bCanAdd = (ElementId == NULL);
-		}
-
-		if (bCanAdd)
-		{
-			FNavigationDirtyElement UpdateInfo(ElementOwner, ElementInterface, GetDirtyFlagHelper(UpdateFlags, 0));
-
-			SetId = PendingOctreeUpdates.FindId(UpdateInfo);
-			if (SetId.IsValidId())
-			{
-				// make sure this request stays, in case it has been invalidated already
-				PendingOctreeUpdates[SetId] = UpdateInfo;
-			}
-			else
-			{
-				SetId = PendingOctreeUpdates.Add(UpdateInfo);
-			}
-		}
-	}
-
-	return SetId;
+	return (ElementOwner && ElementInterface) 
+		? FNavigationDataHandler(DefaultOctreeController, DefaultDirtyAreasController).RegisterNavOctreeElement(*ElementOwner, *ElementInterface, UpdateFlags)
+		: FSetElementId();
 }
 
 void UNavigationSystemV1::AddElementToNavOctree(const FNavigationDirtyElement& DirtyElement)
 {
-	// handle invalidated requests first
-	if (DirtyElement.bInvalidRequest)
-	{
-		if (DirtyElement.bHasPrevData)
-		{
-			AddDirtyArea(DirtyElement.PrevBounds, DirtyElement.PrevFlags);
-		}
-		
-		return;
-	}
-
-	UObject* ElementOwner = DirtyElement.Owner.Get();
-	if (ElementOwner == NULL || ElementOwner->IsPendingKill() || DirtyElement.NavInterface == NULL)
-	{
-		return;
-	}
-
-	FNavigationOctreeElement GeneratedData(*ElementOwner);
-	const FBox ElementBounds = DirtyElement.NavInterface->GetNavigationBounds();
-
-	UObject* NavigationParent = DirtyElement.NavInterface->GetNavigationParent();
-	if (NavigationParent)
-	{
-		// check if parent node is waiting in queue
-		const FSetElementId ParentRequestId = PendingOctreeUpdates.FindId(FNavigationDirtyElement(NavigationParent));
-		const FOctreeElementId* ParentId = GetObjectsNavOctreeId(*NavigationParent);
-		if (ParentRequestId.IsValidId() && ParentId == NULL)
-		{
-			FNavigationDirtyElement& ParentNode = PendingOctreeUpdates[ParentRequestId];
-			AddElementToNavOctree(ParentNode);
-
-			// mark as invalid so it won't be processed twice
-			ParentNode.bInvalidRequest = true;
-		}
-
-		const FOctreeElementId* UseParentId = ParentId ? ParentId : GetObjectsNavOctreeId(*NavigationParent);
-		if (UseParentId && NavOctree->IsValidElementId(*UseParentId))
-		{
-			UE_LOG(LogNavOctree, Log, TEXT("ADD %s to %s"), *GetNameSafe(ElementOwner), *GetNameSafe(NavigationParent));
-			NavOctree->AppendToNode(*UseParentId, DirtyElement.NavInterface, ElementBounds, GeneratedData);
-		}
-		else 
-		{
-			UE_LOG(LogNavOctree, Warning, TEXT("Can't add node [%s] - parent [%s] not found in octree!"), *GetNameSafe(ElementOwner), *GetNameSafe(NavigationParent));
-		}
-	}
-	else
-	{
-		UE_LOG(LogNavOctree, Log, TEXT("ADD %s"), *GetNameSafe(ElementOwner));
-		NavOctree->AddNode(ElementOwner, DirtyElement.NavInterface, ElementBounds, GeneratedData);
-	}
-
-	const FBox BBox = GeneratedData.Bounds.GetBox();
-	const bool bValidBBox = BBox.IsValid && !BBox.GetSize().IsNearlyZero();
-
-	if (bValidBBox && !GeneratedData.IsEmpty())
-	{
-		const int32 DirtyFlag = DirtyElement.FlagsOverride ? DirtyElement.FlagsOverride : GeneratedData.Data->GetDirtyFlag();
-		AddDirtyArea(BBox, DirtyFlag);
-	}
+	FNavigationDataHandler(DefaultOctreeController, DefaultDirtyAreasController).AddElementToNavOctree(DirtyElement);
 }
 
 bool UNavigationSystemV1::GetNavOctreeElementData(const UObject& NodeOwner, int32& DirtyFlags, FBox& DirtyBounds)
 {
-	const FOctreeElementId* ElementId = GetObjectsNavOctreeId(NodeOwner);
-	if (ElementId != NULL)
-	{
-		if (NavOctree->IsValidElementId(*ElementId))
-		{
-			// mark area occupied by given actor as dirty
-			FNavigationOctreeElement& ElementData = NavOctree->GetElementById(*ElementId);
-			DirtyFlags = ElementData.Data->GetDirtyFlag();
-			DirtyBounds = ElementData.Bounds.GetBox();
-			return true;
-		}
-	}
-
-	return false;
+	return DefaultOctreeController.GetNavOctreeElementData(NodeOwner, DirtyFlags, DirtyBounds);
 }
 
 void UNavigationSystemV1::UnregisterNavOctreeElement(UObject* ElementOwner, INavRelevantInterface* ElementInterface, int32 UpdateFlags)
 {
-	SCOPE_CYCLE_COUNTER(STAT_Navigation_UnregisterNavOctreeElement);
-
-#if WITH_EDITOR
-	if (IsNavigationUnregisterLocked())
-	{
-		return;
-	}
-#endif
-
-	if (NavOctree.IsValid() == false || ElementOwner == NULL || ElementInterface == NULL)
-	{
-		return;
-	}
-
-	if (IsNavigationOctreeLocked())
+	if (ElementOwner && ElementInterface)
 	{
-		UE_LOG(LogNavOctree, Log, TEXT("IGNORE(UnregisterNavOctreeElement) %s"), *GetPathNameSafe(ElementOwner));
-		return;
-	}
-
-	const FOctreeElementId* ElementId = GetObjectsNavOctreeId(*ElementOwner);
-	UE_LOG(LogNavOctree, Log, TEXT("UNREG %s %s"), *GetNameSafe(ElementOwner), ElementId ? TEXT("[exists]") : TEXT(""));
-
-	if (ElementId != NULL)
-	{
-		RemoveNavOctreeElementId(*ElementId, UpdateFlags);
-		RemoveObjectsNavOctreeId(*ElementOwner);
-	}
-	else
-	{
-		const bool bCanRemoveChildNode = (UpdateFlags & OctreeUpdate_ParentChain) == 0;
-		UObject* ParentNode = ElementInterface->GetNavigationParent();
-		if (ParentNode && bCanRemoveChildNode)
-		{
-			// if node has navigation parent (= doesn't exists in octree on its own)
-			// and it's not part of parent chain update
-			// remove it from map and force update on parent to rebuild octree element
-
-			OctreeChildNodesMap.RemoveSingle(ParentNode, FWeakObjectPtr(ElementOwner));
-			UpdateNavOctreeParentChain(ParentNode);
-		}
-	}
-
-	// mark pending update as invalid, it will be dirtied according to currently active settings
-	const bool bCanInvalidateQueue = (UpdateFlags & OctreeUpdate_Refresh) == 0;
-	if (bCanInvalidateQueue)
-	{
-		const FSetElementId RequestId = PendingOctreeUpdates.FindId(FNavigationDirtyElement(ElementOwner));
-		if (RequestId.IsValidId())
-		{
-			PendingOctreeUpdates[RequestId].bInvalidRequest = true;
-		}
+		FNavigationDataHandler(DefaultOctreeController, DefaultDirtyAreasController).UnregisterNavOctreeElement(*ElementOwner, *ElementInterface, UpdateFlags);
 	}
 }
 
 void UNavigationSystemV1::RemoveNavOctreeElementId(const FOctreeElementId& ElementId, int32 UpdateFlags)
 {
-	if (NavOctree->IsValidElementId(ElementId))
-	{
-		const FNavigationOctreeElement& ElementData = NavOctree->GetElementById(ElementId);
-		const int32 DirtyFlag = GetDirtyFlagHelper(UpdateFlags, ElementData.Data->GetDirtyFlag());
-		// mark area occupied by given actor as dirty
-		AddDirtyArea(ElementData.Bounds.GetBox(), DirtyFlag);
-		NavOctree->RemoveNode(ElementId);
-	}
+	FNavigationDataHandler(DefaultOctreeController, DefaultDirtyAreasController).RemoveNavOctreeElementId(ElementId, UpdateFlags);
 }
 
 const FNavigationRelevantData* UNavigationSystemV1::GetDataForObject(const UObject& Object) const
 {
-	check(NavOctree.IsValid());
-
-	const FOctreeElementId* OctreeID = GetObjectsNavOctreeId(Object);
-
-	if (OctreeID != nullptr && NavOctree->IsValidElementId(*OctreeID))
-	{
-		return NavOctree->GetDataForID(*OctreeID);
-	}
-
-	return nullptr;
+	return DefaultOctreeController.GetDataForObject(Object);
 }
 
 FNavigationRelevantData* UNavigationSystemV1::GetMutableDataForObject(const UObject& Object)
 {
-	return const_cast<FNavigationRelevantData*>(GetDataForObject(Object));
+	return DefaultOctreeController.GetMutableDataForObject(Object);
 }
 
 void UNavigationSystemV1::UpdateActorInNavOctree(AActor& Actor)
@@ -2619,7 +2389,7 @@ void UNavigationSystemV1::UpdateActorInNavOctree(AActor& Actor)
 		UNavigationSystemV1* NavSys = FNavigationSystem::GetCurrent<UNavigationSystemV1>(Actor.GetWorld());
 		if (NavSys)
 		{
-			NavSys->UpdateNavOctreeElement(&Actor, NavElement, OctreeUpdate_Default);
+			NavSys->UpdateNavOctreeElement(&Actor, NavElement, FNavigationOctreeController::OctreeUpdate_Default);
 		}
 	}
 }
@@ -2636,22 +2406,22 @@ void UNavigationSystemV1::UpdateComponentInNavOctree(UActorComponent& Comp)
 	// special case for early out: use cached nav relevancy
 	if (Comp.bNavigationRelevant == true)
 	{
+		AActor* OwnerActor = Comp.GetOwner();
+		if (OwnerActor)
+		{
 		INavRelevantInterface* NavElement = Cast<INavRelevantInterface>(&Comp);
 		if (NavElement)
 		{
-			AActor* OwnerActor = Comp.GetOwner();
-			if (OwnerActor)
-			{
 				UNavigationSystemV1* NavSys = FNavigationSystem::GetCurrent<UNavigationSystemV1>(OwnerActor->GetWorld());
 				if (NavSys)
 				{
 					if (OwnerActor->IsComponentRelevantForNavigation(&Comp))
 					{
-						NavSys->UpdateNavOctreeElement(&Comp, NavElement, OctreeUpdate_Default);
+						NavSys->UpdateNavOctreeElement(&Comp, NavElement, FNavigationOctreeController::OctreeUpdate_Default);
 					}
 					else
 					{
-						NavSys->UnregisterNavOctreeElement(&Comp, NavElement, OctreeUpdate_Default);
+						NavSys->UnregisterNavOctreeElement(&Comp, NavElement, FNavigationOctreeController::OctreeUpdate_Default);
 					}
 				}
 			}
@@ -2669,7 +2439,7 @@ void UNavigationSystemV1::UpdateComponentInNavOctree(UActorComponent& Comp)
 				UNavigationSystemV1* NavSys = FNavigationSystem::GetCurrent<UNavigationSystemV1>(OwnerActor->GetWorld());
 				if (NavSys)
 				{
-					NavSys->UnregisterNavOctreeElement(&Comp, NavElement, OctreeUpdate_Default);
+					NavSys->UnregisterNavOctreeElement(&Comp, NavElement, FNavigationOctreeController::OctreeUpdate_Default);
 				}
 			}
 		}
@@ -2678,19 +2448,42 @@ void UNavigationSystemV1::UpdateComponentInNavOctree(UActorComponent& Comp)
 
 void UNavigationSystemV1::UpdateActorAndComponentsInNavOctree(AActor& Actor, bool bUpdateAttachedActors)
 {
-	UpdateActorInNavOctree(Actor);
-		
-	for (UActorComponent* Component : Actor.GetComponents())
+	if (IsNavigationSystemStatic())
+	{
+		return;
+	}
+
+	UNavigationSystemV1* NavSys = FNavigationSystem::GetCurrent<UNavigationSystemV1>(Actor.GetWorld());
+	if (NavSys)
 	{
-		if (Component)
+		FNavigationDataHandler DataHandler(NavSys->DefaultOctreeController, NavSys->DefaultDirtyAreasController);
+
+		if (ShouldUpdateNavOctreeOnComponentChange())
 		{
-			UpdateComponentInNavOctree(*Component);
+			// do the full update
+			DataHandler.UpdateActorAndComponentsInNavOctree(Actor);
+		}
+		else
+	{
+			// update just the actor itself 
+			INavRelevantInterface* NavElement = Cast<INavRelevantInterface>(&Actor);
+			if (NavElement)
+		{
+				DataHandler.UpdateNavOctreeElement(Actor, *NavElement, FNavigationOctreeController::OctreeUpdate_Default);
 		}
 	}
 
 	if (bUpdateAttachedActors)
 	{
-		UpdateAttachedActorsInNavOctree(Actor);
+			TArray<AActor*> UniqueAttachedActors;
+			if (GetAllAttachedActors(Actor, UniqueAttachedActors) > 0)
+			{
+				for (AActor* AttachedActor : UniqueAttachedActors)
+				{
+					DataHandler.UpdateActorAndComponentsInNavOctree(Actor);
+				}
+			}
+		}
 	}
 }
 
@@ -2703,29 +2496,37 @@ void UNavigationSystemV1::UpdateNavOctreeAfterMove(USceneComponent* Comp)
 	}
 }
 
-void UNavigationSystemV1::UpdateAttachedActorsInNavOctree(AActor& RootActor)
+int32 UNavigationSystemV1::GetAllAttachedActors(const AActor& RootActor, TArray<AActor*>& OutAttachedActors)
 {
-	TArray<AActor*> UniqueAttachedActors;
-	UniqueAttachedActors.Add(&RootActor);
+	OutAttachedActors.Reset();
+	RootActor.GetAttachedActors(OutAttachedActors);
 
 	TArray<AActor*> TempAttachedActors;
-	for (int32 ActorIndex = 0; ActorIndex < UniqueAttachedActors.Num(); ++ActorIndex)
+	for (int32 ActorIndex = 0; ActorIndex < OutAttachedActors.Num(); ++ActorIndex)
 	{
-		check(UniqueAttachedActors[ActorIndex]);
+		check(OutAttachedActors[ActorIndex]);
 		// find all attached actors
-		UniqueAttachedActors[ActorIndex]->GetAttachedActors(TempAttachedActors);
-		
+		OutAttachedActors[ActorIndex]->GetAttachedActors(TempAttachedActors);
+
 		for (int32 AttachmentIndex = 0; AttachmentIndex < TempAttachedActors.Num(); ++AttachmentIndex)
 		{
 			// and store the ones we don't know about yet
-			UniqueAttachedActors.AddUnique(TempAttachedActors[AttachmentIndex]);
+			OutAttachedActors.AddUnique(TempAttachedActors[AttachmentIndex]);
 		}
 	}
-	
-	// skipping the first item since that's the root, and we just care about the attached actors
-	for (int32 ActorIndex = 1; ActorIndex < UniqueAttachedActors.Num(); ++ActorIndex)
+
+	return OutAttachedActors.Num();
+}
+
+void UNavigationSystemV1::UpdateAttachedActorsInNavOctree(AActor& RootActor)
+{
+	TArray<AActor*> UniqueAttachedActors;
+	if (GetAllAttachedActors(RootActor, UniqueAttachedActors) > 0)
 	{
-		UpdateActorAndComponentsInNavOctree(*UniqueAttachedActors[ActorIndex], /*bUpdateAttachedActors = */false);
+		for (AActor* AttachedActor : UniqueAttachedActors)
+	{
+			UpdateActorAndComponentsInNavOctree(*AttachedActor, /*bUpdateAttachedActors=*/false);
+		}
 	}
 }
 
@@ -2759,176 +2560,35 @@ void UNavigationSystemV1::ClearNavOctreeAll(AActor* Actor)
 
 void UNavigationSystemV1::UpdateNavOctreeElement(UObject* ElementOwner, INavRelevantInterface* ElementInterface, int32 UpdateFlags)
 {
-	INC_DWORD_STAT(STAT_Navigation_UpdateNavOctree);
-
-	if (IsNavigationOctreeLocked())
+	if (ElementOwner && ElementInterface)
 	{
-		UE_LOG(LogNavOctree, Log, TEXT("IGNORE(UpdateNavOctreeElement) %s"), *GetPathNameSafe(ElementOwner));
-		return;
-	}
-	else if (ElementOwner == nullptr)
-	{
-		return;
-	}
-
-	// grab existing octree data
-	FBox CurrentBounds;
-	int32 CurrentFlags;
-	const bool bAlreadyExists = GetNavOctreeElementData(*ElementOwner, CurrentFlags, CurrentBounds);
-
-	// don't invalidate pending requests
-	UpdateFlags |= OctreeUpdate_Refresh;
-
-	// always try to unregister, even if element owner doesn't exists in octree (parent nodes)
-	UnregisterNavOctreeElement(ElementOwner, ElementInterface, UpdateFlags);
-
-	const FSetElementId RequestId = RegisterNavOctreeElement(ElementOwner, ElementInterface, UpdateFlags);
-
-	// add original data to pending registration request
-	// so it could be dirtied properly when system receive unregister request while actor is still queued
-	if (RequestId.IsValidId())
-	{
-		FNavigationDirtyElement& UpdateInfo = PendingOctreeUpdates[RequestId];
-		UpdateInfo.PrevFlags = CurrentFlags;
-		if (UpdateInfo.PrevBounds.IsValid)
-		{
-			// Is we have something stored already we want to 
-			// sum it up, since we care about the whole bounding
-			// box of changes that potentially took place
-			UpdateInfo.PrevBounds += CurrentBounds;
-		}
-		else
-		{
-			UpdateInfo.PrevBounds = CurrentBounds;
+		FNavigationDataHandler(DefaultOctreeController, DefaultDirtyAreasController).UpdateNavOctreeElement(*ElementOwner, *ElementInterface, UpdateFlags);
 		}
-		UpdateInfo.bHasPrevData = bAlreadyExists;
-	}
-
-	UpdateNavOctreeParentChain(ElementOwner, /*bSkipElementOwnerUpdate=*/ true);
 }
 
 void UNavigationSystemV1::UpdateNavOctreeParentChain(UObject* ElementOwner, bool bSkipElementOwnerUpdate)
 {
-	const int32 UpdateFlags = OctreeUpdate_ParentChain | OctreeUpdate_Refresh;
-
-	TArray<FWeakObjectPtr> ChildNodes;
-	OctreeChildNodesMap.MultiFind(ElementOwner, ChildNodes);
-
-	if (ChildNodes.Num() == 0)
-	{
-		if (bSkipElementOwnerUpdate == false)
-		{
-			INavRelevantInterface* ElementInterface = Cast<INavRelevantInterface>(ElementOwner);
-			UpdateNavOctreeElement(ElementOwner, ElementInterface, UpdateFlags);
-		}
-		return;
-	}
-
-	INavRelevantInterface* ElementInterface = Cast<INavRelevantInterface>(ElementOwner);
-	TArray<INavRelevantInterface*> ChildNavInterfaces;
-	ChildNavInterfaces.AddZeroed(ChildNodes.Num());
-	
-	for (int32 Idx = 0; Idx < ChildNodes.Num(); Idx++)
-	{
-		if (ChildNodes[Idx].IsValid())
+	if (ElementOwner)
 		{
-			UObject* ChildNodeOb = ChildNodes[Idx].Get();
-			ChildNavInterfaces[Idx] = Cast<INavRelevantInterface>(ChildNodeOb);
-			UnregisterNavOctreeElement(ChildNodeOb, ChildNavInterfaces[Idx], UpdateFlags);
-		}
-	}
-
-	if (bSkipElementOwnerUpdate == false)
-	{
-		UnregisterNavOctreeElement(ElementOwner, ElementInterface, UpdateFlags);
-		RegisterNavOctreeElement(ElementOwner, ElementInterface, UpdateFlags);
-	}
-
-	for (int32 Idx = 0; Idx < ChildNodes.Num(); Idx++)
-	{
-		if (ChildNodes[Idx].IsValid())
-		{
-			RegisterNavOctreeElement(ChildNodes[Idx].Get(), ChildNavInterfaces[Idx], UpdateFlags);
-		}
+		FNavigationDataHandler(DefaultOctreeController, DefaultDirtyAreasController).UpdateNavOctreeParentChain(*ElementOwner, bSkipElementOwnerUpdate);
 	}
 }
 
 bool UNavigationSystemV1::UpdateNavOctreeElementBounds(UActorComponent* Comp, const FBox& NewBounds, const FBox& DirtyArea)
 {
-	if (Comp == nullptr)
-	{
-		return false;
-	}
-
-	const FOctreeElementId* ElementId = GetObjectsNavOctreeId(*Comp);
-	if (ElementId && ElementId->IsValidId())
-	{
-		NavOctree->UpdateNode(*ElementId, NewBounds);
-		
-		// Add dirty area
-		if (DirtyArea.IsValid)
-		{
-			ElementId = GetObjectsNavOctreeId(*Comp);
-			if (ElementId && ElementId->IsValidId())
-			{
-				FNavigationOctreeElement& ElementData = NavOctree->GetElementById(*ElementId);
-				AddDirtyArea(DirtyArea, ElementData.Data->GetDirtyFlag());
-			}
-		}
-
-		return true;
-	}
-	
-	return false;
+	return Comp
+		? FNavigationDataHandler(DefaultOctreeController, DefaultDirtyAreasController).UpdateNavOctreeElementBounds(*Comp, NewBounds, DirtyArea)
+		: false;
 }
 
 bool UNavigationSystemV1::ReplaceAreaInOctreeData(const UObject& Object, TSubclassOf<UNavArea> OldArea, TSubclassOf<UNavArea> NewArea, bool bReplaceChildClasses)
 {
-	FNavigationRelevantData* Data = GetMutableDataForObject(Object);
-	
-	if (Data == nullptr || Data->HasModifiers() == false)
-	{
-		return false;
-	}
-
-	for (FAreaNavModifier& AreaModifier : Data->Modifiers.GetMutableAreas())
-	{ 
-		if (AreaModifier.GetAreaClass() == OldArea
-			|| (bReplaceChildClasses && AreaModifier.GetAreaClass()->IsChildOf(OldArea)))
-		{
-			AreaModifier.SetAreaClass(NewArea);
-		}
-	}
-	for (FSimpleLinkNavModifier& SimpleLink : Data->Modifiers.GetSimpleLinks())
-	{ 
-		for (FNavigationLink& Link : SimpleLink.Links)
-		{
-			if (Link.GetAreaClass() == OldArea
-				|| (bReplaceChildClasses && Link.GetAreaClass()->IsChildOf(OldArea)))
-			{
-				Link.SetAreaClass(NewArea);
-			}
-		}
-		for (FNavigationSegmentLink& Link : SimpleLink.SegmentLinks)
-		{
-			if (Link.GetAreaClass() == OldArea
-				|| (bReplaceChildClasses && Link.GetAreaClass()->IsChildOf(OldArea)))
-			{
-				Link.SetAreaClass(NewArea);
-			}
-		}
-	}
-	for (FCustomLinkNavModifier& CustomLink : Data->Modifiers.GetCustomLinks())
-	{ 
-		ensureMsgf(false, TEXT("Not implemented yet"));
-	}
-
-	return true;
+	return FNavigationDataHandler(DefaultOctreeController, DefaultDirtyAreasController).ReplaceAreaInOctreeData(Object, OldArea, NewArea, bReplaceChildClasses);
 }
 
 void UNavigationSystemV1::OnComponentRegistered(UActorComponent* Comp)
 {
-	if (IsNavigationSystemStatic())
+	if ((Comp == nullptr) || IsNavigationSystemStatic())
 	{
 		return;
 	}
@@ -2937,13 +2597,13 @@ void UNavigationSystemV1::OnComponentRegistered(UActorComponent* Comp)
 	INavRelevantInterface* NavInterface = Cast<INavRelevantInterface>(Comp);
 	if (NavInterface)
 	{
-		AActor* OwnerActor = Comp ? Comp->GetOwner() : NULL;
+		AActor* OwnerActor = Comp->GetOwner();
 		if (OwnerActor && OwnerActor->IsComponentRelevantForNavigation(Comp))
 		{
 			UNavigationSystemV1* NavSys = FNavigationSystem::GetCurrent<UNavigationSystemV1>(OwnerActor->GetWorld());
 			if (NavSys)
 			{
-				NavSys->RegisterNavOctreeElement(Comp, NavInterface, OctreeUpdate_Default);
+				NavSys->RegisterNavOctreeElement(Comp, NavInterface, FNavigationOctreeController::OctreeUpdate_Default);
 			}
 		}
 	}
@@ -2951,7 +2611,7 @@ void UNavigationSystemV1::OnComponentRegistered(UActorComponent* Comp)
 
 void UNavigationSystemV1::OnComponentUnregistered(UActorComponent* Comp)
 {
-	if (IsNavigationSystemStatic())
+	if ((Comp == nullptr) || IsNavigationSystemStatic())
 	{
 		return;
 	}
@@ -2960,7 +2620,7 @@ void UNavigationSystemV1::OnComponentUnregistered(UActorComponent* Comp)
 	INavRelevantInterface* NavInterface = Cast<INavRelevantInterface>(Comp);
 	if (NavInterface)
 	{
-		AActor* OwnerActor = Comp ? Comp->GetOwner() : NULL;
+		AActor* OwnerActor = Comp->GetOwner();
 		if (OwnerActor)
 		{
 			// skip IsComponentRelevantForNavigation check, it's only for adding new stuff
@@ -2968,7 +2628,7 @@ void UNavigationSystemV1::OnComponentUnregistered(UActorComponent* Comp)
 			UNavigationSystemV1* NavSys = FNavigationSystem::GetCurrent<UNavigationSystemV1>(OwnerActor->GetWorld());
 			if (NavSys)
 			{
-				NavSys->UnregisterNavOctreeElement(Comp, NavInterface, OctreeUpdate_Default);
+				NavSys->UnregisterNavOctreeElement(Comp, NavInterface, FNavigationOctreeController::OctreeUpdate_Default);
 			}
 		}
 	}
@@ -2976,7 +2636,7 @@ void UNavigationSystemV1::OnComponentUnregistered(UActorComponent* Comp)
 
 void UNavigationSystemV1::OnActorRegistered(AActor* Actor)
 {
-	if (IsNavigationSystemStatic())
+	if ((Actor == nullptr) || IsNavigationSystemStatic())
 	{
 		return;
 	}
@@ -2988,14 +2648,14 @@ void UNavigationSystemV1::OnActorRegistered(AActor* Actor)
 		UNavigationSystemV1* NavSys = FNavigationSystem::GetCurrent<UNavigationSystemV1>(Actor->GetWorld());
 		if (NavSys)
 		{
-			NavSys->RegisterNavOctreeElement(Actor, NavInterface, OctreeUpdate_Default);
+			NavSys->RegisterNavOctreeElement(Actor, NavInterface, FNavigationOctreeController::OctreeUpdate_Default);
 		}
 	}
 }
 
 void UNavigationSystemV1::OnActorUnregistered(AActor* Actor)
 {
-	if (IsNavigationSystemStatic())
+	if ((Actor == nullptr) || IsNavigationSystemStatic())
 	{
 		return;
 	}
@@ -3007,27 +2667,14 @@ void UNavigationSystemV1::OnActorUnregistered(AActor* Actor)
 		UNavigationSystemV1* NavSys = FNavigationSystem::GetCurrent<UNavigationSystemV1>(Actor->GetWorld());
 		if (NavSys)
 		{
-			NavSys->UnregisterNavOctreeElement(Actor, NavInterface, OctreeUpdate_Default);
+			NavSys->UnregisterNavOctreeElement(Actor, NavInterface, FNavigationOctreeController::OctreeUpdate_Default);
 		}
 	}
 }
 
 void UNavigationSystemV1::FindElementsInNavOctree(const FBox& QueryBox, const FNavigationOctreeFilter& Filter, TArray<FNavigationOctreeElement>& Elements)
-{
-	if (NavOctree.IsValid() == false)
-	{
-		UE_LOG(LogNavigation, Warning, TEXT("UNavigationSystemV1::FindElementsInNavOctree gets called while NavOctree is null"));
-		return;
-	}
-	
-	for (FNavigationOctree::TConstElementBoxIterator<> It(*NavOctree, QueryBox); It.HasPendingElements(); It.Advance())
-	{
-		const FNavigationOctreeElement& Element = It.GetCurrentElement();
-		if (Element.IsMatchingFilter(Filter))
-		{
-			Elements.Add(Element);
-		}		
-	}
+{	
+	FNavigationDataHandler(DefaultOctreeController, DefaultDirtyAreasController).FindElementsInNavOctree(QueryBox, Filter, Elements);
 }
 
 void UNavigationSystemV1::ReleaseInitialBuildingLock()
@@ -3352,7 +2999,7 @@ void UNavigationSystemV1::Build()
 
 #if !UE_BUILD_SHIPPING
 	// no longer report that navmesh needs to be rebuild
-	bDirtyAreasReportedWhileAccumulationLocked = false;
+	DefaultDirtyAreasController.bDirtyAreasReportedWhileAccumulationLocked = false;
 #endif // !UE_BUILD_SHIPPING
 
 	UE_LOG(LogNavigation, Display, TEXT("UNavigationSystemV1::Build total execution time: %.5f"), float(FPlatformTime::Seconds() - BuildStartTime));
@@ -3395,7 +3042,7 @@ void UNavigationSystemV1::SpawnMissingNavigationData()
 			for (int32 AgentIndex = 0; AgentIndex < SupportedAgentsCount; ++AgentIndex)
 			{
 				if (AlreadyInstantiated[AgentIndex] == false
-					&& Nav->GetClass() == SupportedAgents[AgentIndex].NavigationDataClass 
+					&& Nav->GetClass() == SupportedAgents[AgentIndex].GetNavDataClass<ANavigationData>()
 					&& Nav->DoesSupportAgent(SupportedAgents[AgentIndex]) == true)
 				{
 					AlreadyInstantiated[AgentIndex] = true;
@@ -3412,11 +3059,11 @@ void UNavigationSystemV1::SpawnMissingNavigationData()
 		for (int32 AgentIndex = 0; AgentIndex < SupportedAgentsCount; ++AgentIndex)
 		{
 			const FNavDataConfig& NavConfig = SupportedAgents[AgentIndex];
-			if (AlreadyInstantiated[AgentIndex] == false && NavConfig.NavigationDataClass != nullptr)
+			if (AlreadyInstantiated[AgentIndex] == false && NavConfig.GetNavDataClass<ANavigationData>() != nullptr)
 			{
 				bool bHandled = false;
 
-				const ANavigationData* NavDataCDO = NavConfig.NavigationDataClass->GetDefaultObject<ANavigationData>();
+				const ANavigationData* NavDataCDO = NavConfig.GetNavDataClass<ANavigationData>()->GetDefaultObject<ANavigationData>();
 				if (NavDataCDO == nullptr || !NavDataCDO->CanSpawnOnRebuild())
 				{
 					continue;
@@ -3431,7 +3078,7 @@ void UNavigationSystemV1::SpawnMissingNavigationData()
 					continue;
 				}
 
-				ANavigationData* Instance = CreateNavigationDataInstance(NavConfig);
+				ANavigationData* Instance = CreateNavigationDataInstanceInLevel(NavConfig, nullptr);
 				if (Instance)
 				{
 					RequestRegistration(Instance);
@@ -3451,13 +3098,14 @@ void UNavigationSystemV1::SpawnMissingNavigationData()
 	}
 }
 
-ANavigationData* UNavigationSystemV1::CreateNavigationDataInstance(const FNavDataConfig& NavConfig)
+ANavigationData* UNavigationSystemV1::CreateNavigationDataInstanceInLevel(const FNavDataConfig& NavConfig, ULevel* SpawnLevel)
 {
 	UWorld* World = GetWorld();
 	check(World);
 
 	FActorSpawnParameters SpawnInfo;
-	if (bSpawnNavDataInNavBoundsLevel && RegisteredNavBounds.Num() > 0)
+	SpawnInfo.OverrideLevel = SpawnLevel;
+	if (bSpawnNavDataInNavBoundsLevel && SpawnLevel == nullptr && RegisteredNavBounds.Num() > 0)
 	{
 		// pick the first valid level
 		for (const FNavigationBounds& Bounds : RegisteredNavBounds)
@@ -3473,7 +3121,8 @@ ANavigationData* UNavigationSystemV1::CreateNavigationDataInstance(const FNavDat
 	{
 		SpawnInfo.OverrideLevel = World->PersistentLevel;
 	}
-	ANavigationData* Instance = World->SpawnActor<ANavigationData>(*NavConfig.NavigationDataClass, SpawnInfo);
+
+	ANavigationData* Instance = World->SpawnActor<ANavigationData>(*NavConfig.GetNavDataClass<ANavigationData>(), SpawnInfo);
 
 	if (Instance != NULL)
 	{
@@ -3556,18 +3205,12 @@ void UNavigationSystemV1::RebuildAll(bool bIsLoadTime)
 	GatherNavigationBounds();
 
 	// make sure that octree is up to date
-	for (TSet<FNavigationDirtyElement>::TIterator It(PendingOctreeUpdates); It; ++It)
-	{
-		AddElementToNavOctree(*It);
-	}
-	PendingOctreeUpdates.Empty(32);
-
-	// discard all pending dirty areas, we are going to rebuild navmesh anyway 
-	DirtyAreas.Reset();
+	FNavigationDataHandler NavHandler(DefaultOctreeController, DefaultDirtyAreasController);
+	NavHandler.ProcessPendingOctreeUpdates();
+	
 	PendingNavBoundsUpdates.Reset();
-#if !UE_BUILD_SHIPPING
-	bDirtyAreasReportedWhileAccumulationLocked = false;
-#endif // !UE_BUILD_SHIPPING
+
+	DefaultDirtyAreasController.Reset();
 
 	// 
 	for (int32 NavDataIndex = 0; NavDataIndex < NavDataSet.Num(); ++NavDataIndex)
@@ -3581,18 +3224,12 @@ void UNavigationSystemV1::RebuildAll(bool bIsLoadTime)
 	}
 }
 
-void UNavigationSystemV1::RebuildDirtyAreas()
+void UNavigationSystemV1::RebuildDirtyAreas(float DeltaSeconds)
 {
-	for (int32 NavDataIndex = 0; NavDataIndex < NavDataSet.Num(); ++NavDataIndex)
-	{
-		ANavigationData* NavData = NavDataSet[NavDataIndex];
-		if (NavData)
-		{
-			NavData->RebuildDirtyAreas(DirtyAreas);
-		}
-	}
-
-	DirtyAreas.Reset();
+	SCOPE_CYCLE_COUNTER(STAT_Navigation_TickMarkDirty);
+	UWorld* World = GetWorld();
+	const bool bForceRebuilding = (World != nullptr) && (World->IsGameWorld() == false);
+	DefaultDirtyAreasController.Tick(DeltaSeconds, NavDataSet, bForceRebuilding);
 }
 
 bool UNavigationSystemV1::IsNavigationBuildInProgress(bool bCheckDirtyToo)
@@ -3657,7 +3294,7 @@ int32 UNavigationSystemV1::GetNumRunningBuildTasks() const
 
 void UNavigationSystemV1::OnLevelAddedToWorld(ULevel* InLevel, UWorld* InWorld)
 {
-	if ((IsNavigationSystemStatic() == false) && (InWorld == GetWorld()))
+	if ((IsNavigationSystemStatic() == false) && (InWorld == GetWorld()) && (InLevel != nullptr))
 	{
 		AddLevelCollisionToOctree(InLevel);
 
@@ -3676,7 +3313,7 @@ void UNavigationSystemV1::OnLevelAddedToWorld(ULevel* InLevel, UWorld* InWorld)
 
 void UNavigationSystemV1::OnLevelRemovedFromWorld(ULevel* InLevel, UWorld* InWorld)
 {
-	if ((IsNavigationSystemStatic() == false) && (InWorld == GetWorld()))
+	if ((IsNavigationSystemStatic() == false) && (InWorld == GetWorld()) && (InLevel != nullptr))
 	{
 		RemoveLevelCollisionFromOctree(InLevel);
 
@@ -3719,50 +3356,17 @@ void UNavigationSystemV1::AddLevelToOctree(ULevel& Level)
 
 void UNavigationSystemV1::AddLevelCollisionToOctree(ULevel* Level)
 {
-#if WITH_RECAST
-	if (Level && NavOctree.IsValid() &&
-		NavOctree->GetNavGeometryStoringMode() == FNavigationOctree::StoreNavGeometry)
+	if (Level)
 	{
-		const TArray<FVector>* LevelGeom = Level->GetStaticNavigableGeometry();
-		const FOctreeElementId* ElementId = GetObjectsNavOctreeId(*Level);
-
-		if (!ElementId && LevelGeom && LevelGeom->Num() > 0)
-		{
-			FNavigationOctreeElement BSPElem(*Level);
-			FRecastNavMeshGenerator::ExportVertexSoupGeometry(*LevelGeom, *BSPElem.Data);
-
-			const auto& Bounds = BSPElem.Data->Bounds;
-			if (!Bounds.GetExtent().IsNearlyZero())
-			{
-				NavOctree->AddNode(Level, NULL, Bounds, BSPElem);
-				AddDirtyArea(Bounds, ENavigationDirtyFlag::All);
-
-				UE_LOG(LogNavOctree, Log, TEXT("ADD %s"), *GetNameSafe(Level));
-			}
-		}
+		FNavigationDataHandler(DefaultOctreeController, DefaultDirtyAreasController).AddLevelCollisionToOctree(*Level);
 	}
-#endif// WITH_RECAST
 }
 
 void UNavigationSystemV1::RemoveLevelCollisionFromOctree(ULevel* Level)
 {
-	if (Level && NavOctree.IsValid())
+	if (Level)
 	{
-		const FOctreeElementId* ElementId = GetObjectsNavOctreeId(*Level);
-		UE_LOG(LogNavOctree, Log, TEXT("UNREG %s %s"), *GetNameSafe(Level), ElementId ? TEXT("[exists]") : TEXT(""));
-
-		if (ElementId != NULL)
-		{
-			if (NavOctree->IsValidElementId(*ElementId))
-			{
-				// mark area occupied by given actor as dirty
-				FNavigationOctreeElement& ElementData = NavOctree->GetElementById(*ElementId);
-				AddDirtyArea(ElementData.Bounds.GetBox(), ENavigationDirtyFlag::All);
-			}
-
-			NavOctree->RemoveNode(*ElementId);
-			RemoveObjectsNavOctreeId(*Level);
-		}
+		FNavigationDataHandler(DefaultOctreeController, DefaultDirtyAreasController).RemoveLevelCollisionFromOctree(*Level);
 	}
 }
 
@@ -3803,7 +3407,7 @@ void UNavigationSystemV1::OnNavigationDirtied(const FBox& Bounds)
 #if WITH_HOT_RELOAD
 void UNavigationSystemV1::OnHotReload(bool bWasTriggeredAutomatically)
 {
-	if (RequiresNavOctree() && NavOctree.IsValid() == false)
+	if (RequiresNavOctree() && DefaultOctreeController.NavOctree.IsValid() == false)
 	{
 		ConditionalPopulateNavOctree();
 
@@ -3817,6 +3421,11 @@ void UNavigationSystemV1::OnHotReload(bool bWasTriggeredAutomatically)
 
 void UNavigationSystemV1::CleanUp(FNavigationSystem::ECleanupMode Mode)
 {
+	if (bCleanUpDone)
+	{
+		return;
+	}
+
 	UE_LOG(LogNavigation, Log, TEXT("UNavigationSystemV1::CleanUp"));
 
 #if WITH_EDITOR
@@ -3845,22 +3454,23 @@ void UNavigationSystemV1::CleanUp(FNavigationSystem::ECleanupMode Mode)
 	NavDataSet.Reset();
 
 	// reset unique link Id for new map
-	const UWorld* MyWorld = (Mode == FNavigationSystem::ECleanupMode::CleanupWithWorld) ? GetWorld() : NULL;
-	if (MyWorld && (MyWorld->WorldType == EWorldType::Game || MyWorld->WorldType == EWorldType::Editor))
+	UWorld* MyWorld = (Mode == FNavigationSystem::ECleanupMode::CleanupWithWorld) ? GetWorld() : NULL;
+	if (MyWorld)
+	{
+		MyWorld->OnBeginTearingDown().RemoveAll(this);
+
+		if (MyWorld->WorldType == EWorldType::Game || MyWorld->WorldType == EWorldType::Editor)
 	{
 		INavLinkCustomInterface::NextUniqueId = 1;
 	}
+	}
+
+	bCleanUpDone = true;
 }
 
 void UNavigationSystemV1::DestroyNavOctree()
 {
-	if (NavOctree.IsValid())
-	{
-		NavOctree->Destroy();
-		NavOctree = NULL;
-	}
-
-	ObjectToOctreeId.Empty();
+	DefaultOctreeController.Reset();
 }
 
 bool UNavigationSystemV1::RequiresNavOctree() const
@@ -4080,9 +3690,9 @@ bool UNavigationSystemV1::HandleCountNavMemCommand()
 {
 	UE_LOG(LogNavigation, Warning, TEXT("Logging NavigationSystem memory usage:"));
 
-	if (NavOctree.IsValid())
+	if (DefaultOctreeController.NavOctree.IsValid())
 	{
-		UE_LOG(LogNavigation, Warning, TEXT("NavOctree memory: %d"), NavOctree->GetSizeBytes());
+		UE_LOG(LogNavigation, Warning, TEXT("NavOctree memory: %d"), DefaultOctreeController.NavOctree->GetSizeBytes());
 	}
 
 	for (int32 NavDataIndex = 0; NavDataIndex < NavDataSet.Num(); ++NavDataIndex)
@@ -4148,7 +3758,7 @@ void UNavigationSystemV1::CycleNavigationDataDrawn()
 bool UNavigationSystemV1::IsNavigationDirty() const
 {
 #if !UE_BUILD_SHIPPING
-	if (bCanAccumulateDirtyAreas == false && bDirtyAreasReportedWhileAccumulationLocked)
+	if (DefaultDirtyAreasController.HadDirtyAreasReportedWhileAccumulationLocked())
 	{
 		return true;
 	}
@@ -4250,9 +3860,9 @@ void UNavigationSystemV1::UnregisterNavigationInvoker(AActor& Invoker)
 void UNavigationSystemV1::SetGeometryGatheringMode(ENavDataGatheringModeConfig NewMode)
 {
 	DataGatheringMode = NewMode;
-	if (NavOctree.IsValid())
+	if (DefaultOctreeController.NavOctree.IsValid())
 	{
-		NavOctree->SetDataGatheringMode(DataGatheringMode);
+		DefaultOctreeController.NavOctree->SetDataGatheringMode(DataGatheringMode);
 	}
 }
 
@@ -4431,6 +4041,11 @@ void UNavigationSystemV1::SimpleMoveToLocation(AController* Controller, const FV
 	UE_LOG(LogNavigation, Error, TEXT("SimpleMoveToLocation is deprecated. Use UAIBlueprintHelperLibrary::SimpleMoveToLocation instead"));
 }
 
+ANavigationData* UNavigationSystemV1::CreateNavigationDataInstance(const FNavDataConfig& NavConfig)
+{
+	return CreateNavigationDataInstanceInLevel(NavConfig, nullptr);
+}
+
 //----------------------------------------------------------------------//
 // NEW STUFF!
 //----------------------------------------------------------------------//
@@ -4482,7 +4097,7 @@ INavigationDataInterface* UNavigationSystemV1::GetNavDataForActor(const AActor&
 	if (AsNavAgent)
 	{
 		const FNavAgentProperties& AgentProps = AsNavAgent->GetNavAgentPropertiesRef();
-		NavData = NavSys->GetNavDataForProps(AgentProps);
+		NavData = NavSys->GetNavDataForProps(AgentProps, AsNavAgent->GetNavAgentLocation());
 	}
 	if (NavData == nullptr)
 	{
@@ -4493,7 +4108,7 @@ INavigationDataInterface* UNavigationSystemV1::GetNavDataForActor(const AActor&
 	return (INavigationDataInterface*)(Cast<ARecastNavMesh>(NavData));
 }
 
-int UNavigationSystemV1::GetNavigationBoundsForNavData(const ANavigationData& NavData, TArray<FBox>& OutBounds) const
+int UNavigationSystemV1::GetNavigationBoundsForNavData(const ANavigationData& NavData, TArray<FBox>& OutBounds, ULevel* InLevel) const
 {
 	const int InitialBoundsCount = OutBounds.Num();
 	OutBounds.Reserve(InitialBoundsCount + RegisteredNavBounds.Num());
@@ -4501,7 +4116,8 @@ int UNavigationSystemV1::GetNavigationBoundsForNavData(const ANavigationData& Na
 
 	for (const FNavigationBounds& NavigationBounds : RegisteredNavBounds)
 	{
-		if (NavigationBounds.SupportedAgents.Contains(AgentIndex))
+		if ((InLevel == nullptr || NavigationBounds.Level == InLevel)
+			&& NavigationBounds.SupportedAgents.Contains(AgentIndex))
 		{
 			OutBounds.Add(NavigationBounds.AreaBox);
 		}
diff --git a/Engine/Source/Runtime/NavigationSystem/Public/NavigationDataHandler.h b/Engine/Source/Runtime/NavigationSystem/Public/NavigationDataHandler.h
new file mode 100644
--- /dev/null
+++ b/Engine/Source/Runtime/NavigationSystem/Public/NavigationDataHandler.h
@@ -0,0 +1,28 @@
+// Copyright 1998-2019 Epic Games, Inc. All Rights Reserved.
+
+#pragma once
+#include "NavigationOctreeController.h"
+#include "NavigationDirtyAreasController.h"
+
+
+struct NAVIGATIONSYSTEM_API FNavigationDataHandler
+{
+	FNavigationOctreeController& OctreeController;
+	FNavigationDirtyAreasController& DirtyAreasController;
+
+	FNavigationDataHandler(FNavigationOctreeController& InOctreeController, FNavigationDirtyAreasController& InDirtyAreasController);
+
+	void RemoveNavOctreeElementId(const FOctreeElementId& ElementId, int32 UpdateFlags);
+	FSetElementId RegisterNavOctreeElement(UObject& ElementOwner, INavRelevantInterface& ElementInterface, int32 UpdateFlags);
+	void AddElementToNavOctree(const FNavigationDirtyElement& DirtyElement);
+	void UnregisterNavOctreeElement(UObject& ElementOwner, INavRelevantInterface& ElementInterface, int32 UpdateFlags);
+	void UpdateNavOctreeElement(UObject& ElementOwner, INavRelevantInterface& ElementInterface, int32 UpdateFlags);
+	void UpdateNavOctreeParentChain(UObject& ElementOwner, bool bSkipElementOwnerUpdate = false);
+	bool UpdateNavOctreeElementBounds(UActorComponent& Comp, const FBox& NewBounds, const FBox& DirtyArea);	
+	void FindElementsInNavOctree(const FBox& QueryBox, const FNavigationOctreeFilter& Filter, TArray<FNavigationOctreeElement>& Elements);
+	bool ReplaceAreaInOctreeData(const UObject& Object, TSubclassOf<UNavArea> OldArea, TSubclassOf<UNavArea> NewArea, bool bReplaceChildClasses);
+	void AddLevelCollisionToOctree(ULevel& Level);
+	void RemoveLevelCollisionFromOctree(ULevel& Level);
+	void UpdateActorAndComponentsInNavOctree(AActor& Actor);
+	void ProcessPendingOctreeUpdates();
+};
diff --git a/Engine/Source/Runtime/NavigationSystem/Public/NavigationDirtyAreasController.h b/Engine/Source/Runtime/NavigationSystem/Public/NavigationDirtyAreasController.h
new file mode 100644
--- /dev/null
+++ b/Engine/Source/Runtime/NavigationSystem/Public/NavigationDirtyAreasController.h
@@ -0,0 +1,38 @@
+// Copyright 1998-2019 Epic Games, Inc. All Rights Reserved.
+
+#pragma once
+#include "AI/Navigation/NavigationTypes.h"
+
+
+class ANavigationData;
+
+struct NAVIGATIONSYSTEM_API FNavigationDirtyAreasController
+{
+	/** update frequency for dirty areas on navmesh */
+	float DirtyAreasUpdateFreq = 60.f;
+
+	/** temporary cumulative time to calculate when we need to update dirty areas */
+	float DirtyAreasUpdateTime = 0.f;
+
+	/** stores areas marked as dirty throughout the frame, processes them
+ *	once a frame in Tick function */
+	TArray<FNavigationDirtyArea> DirtyAreas;
+
+	uint8 bCanAccumulateDirtyAreas : 1;
+#if !UE_BUILD_SHIPPING
+	uint8 bDirtyAreasReportedWhileAccumulationLocked : 1;
+#endif // !UE_BUILD_SHIPPING
+
+	FNavigationDirtyAreasController();
+
+	void Reset();
+	
+	void Tick(float DeltaSeconds, const TArray<ANavigationData*>& NavDataSet, bool bForceRebuilding = false);
+	void AddArea(const FBox& NewArea, int32 Flags);
+	
+	bool IsDirty() const { return DirtyAreas.Num() > 0; }
+
+#if !UE_BUILD_SHIPPING
+	bool HadDirtyAreasReportedWhileAccumulationLocked() const { return bCanAccumulateDirtyAreas == false && bDirtyAreasReportedWhileAccumulationLocked; }
+#endif UE_BUILD_SHIPPING
+};
\ No newline at end of file
diff --git a/Engine/Source/Runtime/NavigationSystem/Public/NavigationOctree.h b/Engine/Source/Runtime/NavigationSystem/Public/NavigationOctree.h
--- a/Engine/Source/Runtime/NavigationSystem/Public/NavigationOctree.h
+++ b/Engine/Source/Runtime/NavigationSystem/Public/NavigationOctree.h
@@ -13,18 +13,35 @@
 #include "Math/GenericOctree.h"
 
 class INavRelevantInterface;
-
+class FNavigationOctree;
 typedef FNavigationRelevantDataFilter FNavigationOctreeFilter;
 
 struct NAVIGATIONSYSTEM_API FNavigationOctreeElement
 {
 	FBoxSphereBounds Bounds;
 	TSharedRef<FNavigationRelevantData, ESPMode::ThreadSafe> Data;
+private:
+	FNavigationOctree& OwnerOctree;
+
+public:
+	UE_DEPRECATED(4.24, "Single-paramater FNavigationOctreeElement ctor is deprecated. Use the one explicitly passing the FNavigationOctree reference in")
+	explicit FNavigationOctreeElement(UObject& SourceObject);
 
-	FNavigationOctreeElement(UObject& SourceObject)
+	FNavigationOctreeElement(FNavigationOctree& InOwnerOctree, UObject& SourceObject)
 		: Data(new FNavigationRelevantData(SourceObject))
-	{
+		, OwnerOctree(InOwnerOctree)
+	{}
 
+	FNavigationOctreeElement(const FNavigationOctreeElement& Other)
+		: Bounds(Other.Bounds)
+		, Data(Other.Data)
+		, OwnerOctree(Other.OwnerOctree)
+	{}
+
+	FNavigationOctreeElement& operator=(const FNavigationOctreeElement& Other)
+	{
+		new(this) FNavigationOctreeElement(Other);
+		return *this;
 	}
 
 	FORCEINLINE bool IsEmpty() const
@@ -69,6 +86,8 @@ struct NAVIGATIONSYSTEM_API FNavigationOctreeElement
 	}
 
 	FORCEINLINE UObject* GetOwner(bool bEvenIfPendingKill = false) const { return Data->SourceObject.Get(bEvenIfPendingKill); }
+
+	FORCEINLINE FNavigationOctree& GetOwnerOctree() const { return OwnerOctree; }
 };
 
 struct FNavigationOctreeSemantics
@@ -90,13 +109,18 @@ struct FNavigationOctreeSemantics
 		return A.Data->SourceObject == B.Data->SourceObject;
 	}
 
+	FORCEINLINE static void ApplyOffset(FNavigationOctreeElement& Element, const FVector& InOffset)
+	{
+		ensureMsgf(false, TEXT("Not implemented yet"));
+	}
+
 #if NAVSYS_DEBUG
 	FORCENOINLINE 
 #endif // NAVSYS_DEBUG
 	static void SetElementId(const FNavigationOctreeElement& Element, FOctreeElementId Id);
 };
 
-class FNavigationOctree : public TOctree<FNavigationOctreeElement, FNavigationOctreeSemantics>, public TSharedFromThis<FNavigationOctree, ESPMode::ThreadSafe>
+class NAVIGATIONSYSTEM_API FNavigationOctree : public TOctree<FNavigationOctreeElement, FNavigationOctreeSemantics>, public TSharedFromThis<FNavigationOctree, ESPMode::ThreadSafe>
 {
 public:
 	DECLARE_DELEGATE_TwoParams(FNavigableGeometryComponentExportDelegate, UActorComponent*, FNavigationRelevantData&);
@@ -137,7 +161,17 @@ public:
 	void DemandLazyDataGathering(const FNavigationOctreeElement& Element);
 	void DemandLazyDataGathering(FNavigationRelevantData& ElementData);
 
+	void SetElementId(const UObject& Object, FOctreeElementId Id);
+
+	FORCEINLINE static uint32 HashObject(const UObject& Object)
+	{
+		return Object.GetUniqueID();
+	}
+
 protected:
+	friend struct FNavigationOctreeController;
+
+	TMap<uint32, FOctreeElementId> ObjectToOctreeId;
 	ENavDataGatheringMode DefaultGeometryGatheringMode;
 	uint32 bGatherGeometry : 1;
 	uint32 NodesMemory;
diff --git a/Engine/Source/Runtime/NavigationSystem/Public/NavigationOctreeController.h b/Engine/Source/Runtime/NavigationSystem/Public/NavigationOctreeController.h
new file mode 100644
--- /dev/null
+++ b/Engine/Source/Runtime/NavigationSystem/Public/NavigationOctreeController.h
@@ -0,0 +1,98 @@
+// Copyright 1998-2019 Epic Games, Inc. All Rights Reserved.
+
+#pragma once
+#include "NavigationOctree.h"
+
+
+struct NAVIGATIONSYSTEM_API FNavigationOctreeController
+{
+	enum EOctreeUpdateMode
+	{
+		OctreeUpdate_Default = 0,						// regular update, mark dirty areas depending on exported content
+		OctreeUpdate_Geometry = 1,						// full update, mark dirty areas for geometry rebuild
+		OctreeUpdate_Modifiers = 2,						// quick update, mark dirty areas for modifier rebuild
+		OctreeUpdate_Refresh = 4,						// update is used for refresh, don't invalidate pending queue
+		OctreeUpdate_ParentChain = 8,					// update child nodes, don't remove anything
+	};
+
+	TSet<FNavigationDirtyElement> PendingOctreeUpdates;
+	TSharedPtr<FNavigationOctree, ESPMode::ThreadSafe> NavOctree;
+	/** Map of all objects that are tied to indexed navigation parent */
+	TMultiMap<UObject*, FWeakObjectPtr> OctreeChildNodesMap;
+	/** if set, navoctree updates are ignored, use with caution! */
+	uint8 bNavOctreeLock : 1;
+
+	void SetNavigationOctreeLock(bool bLock);
+	bool HasPendingObjectNavOctreeId(UObject& Object) const;
+	void RemoveObjectsNavOctreeId(const UObject& Object);
+	void SetNavigableGeometryStoringMode(FNavigationOctree::ENavGeometryStoringMode NavGeometryMode);
+
+	void Reset();
+
+	const FNavigationOctree* GetOctree() const;
+	FNavigationOctree* GetMutableOctree();
+	const FOctreeElementId* GetObjectsNavOctreeId(const UObject& Object) const;
+	bool GetNavOctreeElementData(const UObject& NodeOwner, int32& DirtyFlags, FBox& DirtyBounds);
+	const FNavigationRelevantData* GetDataForObject(const UObject& Object) const;
+	FNavigationRelevantData* GetMutableDataForObject(const UObject& Object);
+	bool HasObjectsNavOctreeId(const UObject& Object) const;
+	bool IsNavigationOctreeLocked() const;
+	/** basically says if navoctree has been created already */
+	bool IsValid() const { return NavOctree.IsValid(); }
+
+private:
+	static uint32 HashObject(const UObject& Object);
+};
+
+//----------------------------------------------------------------------//
+// inlines
+//----------------------------------------------------------------------//
+FORCEINLINE uint32 FNavigationOctreeController::HashObject(const UObject& Object) 
+{ 
+	return FNavigationOctree::HashObject(Object);
+}
+
+FORCEINLINE_DEBUGGABLE const FOctreeElementId* FNavigationOctreeController::GetObjectsNavOctreeId(const UObject& Object) const
+{ 
+	return NavOctree.IsValid()
+		? NavOctree->ObjectToOctreeId.Find(HashObject(Object))
+		: nullptr;
+}
+
+FORCEINLINE_DEBUGGABLE bool FNavigationOctreeController::HasObjectsNavOctreeId(const UObject& Object) const
+{
+	return NavOctree.IsValid() && (NavOctree->ObjectToOctreeId.Find(HashObject(Object)) != nullptr);
+}
+
+FORCEINLINE bool FNavigationOctreeController::HasPendingObjectNavOctreeId(UObject& Object) const 
+{ 
+	return PendingOctreeUpdates.Contains(FNavigationDirtyElement(&Object)); 
+}
+
+FORCEINLINE_DEBUGGABLE void FNavigationOctreeController::RemoveObjectsNavOctreeId(const UObject& Object)
+{ 
+	if (NavOctree.IsValid())
+	{
+		NavOctree->ObjectToOctreeId.Remove(HashObject(Object));
+	}
+}
+
+FORCEINLINE const FNavigationOctree* FNavigationOctreeController::GetOctree() const
+{ 
+	return NavOctree.Get(); 
+}
+
+FORCEINLINE FNavigationOctree* FNavigationOctreeController::GetMutableOctree()
+{ 
+	return NavOctree.Get(); 
+}
+
+FORCEINLINE bool FNavigationOctreeController::IsNavigationOctreeLocked() const
+{ 
+	return bNavOctreeLock; 
+}
+
+FORCEINLINE void FNavigationOctreeController::SetNavigationOctreeLock(bool bLock) 
+{ 
+	bNavOctreeLock = bLock; 
+}
diff --git a/Engine/Source/Runtime/NavigationSystem/Public/NavigationSystem.h b/Engine/Source/Runtime/NavigationSystem/Public/NavigationSystem.h
--- a/Engine/Source/Runtime/NavigationSystem/Public/NavigationSystem.h
+++ b/Engine/Source/Runtime/NavigationSystem/Public/NavigationSystem.h
@@ -16,6 +16,11 @@
 #include "Kismet/BlueprintFunctionLibrary.h"
 #include "NavigationOctree.h"
 #include "AI/NavigationSystemConfig.h"
+#include "NavigationOctreeController.h"
+#include "NavigationDirtyAreasController.h"
+#if WITH_EDITOR
+#include "UnrealEdMisc.h"
+#endif // WITH_EDITOR
 #include "NavigationSystem.generated.h"
 
 
@@ -24,7 +29,6 @@ class ANavMeshBoundsVolume;
 class AWorldSettings;
 class FEdMode;
 class FNavDataGenerator;
-class FNavigationOctree;
 class INavLinkCustomInterface;
 class INavRelevantInterface;
 class UCrowdManagerBase;
@@ -182,9 +186,6 @@ protected:
 	TArray<FNavDataConfig> SupportedAgents;
 
 public:
-	/** update frequency for dirty areas on navmesh */
-	UPROPERTY(config, EditAnywhere, Category=NavigationSystem)
-	float DirtyAreasUpdateFreq;
 
 	UPROPERTY()
 	TArray<ANavigationData*> NavDataSet;
@@ -192,8 +193,6 @@ public:
 	UPROPERTY(transient)
 	TArray<ANavigationData*> NavDataRegistrationQueue;
 
-	TSet<FNavigationDirtyElement> PendingOctreeUpdates;
-
 	// List of pending navigation bounds update requests (add, remove, update size)
 	TArray<FNavigationBoundsUpdateRequest> PendingNavBoundsUpdates;
 
@@ -322,6 +321,7 @@ public:
 		RegistrationSuccessful,
 	};
 
+	// EOctreeUpdateMode is deprecated. Use FNavigationOctreeController::EOctreeUpdateMode instead
 	enum EOctreeUpdateMode
 	{
 		OctreeUpdate_Default = 0,						// regular update, mark dirty areas depending on exported content
@@ -432,7 +432,7 @@ public:
 	// @todo document
 	bool ProjectPointToNavigation(const FVector& Point, FNavLocation& OutLocation, const FVector& Extent = INVALID_NAVEXTENT, const FNavAgentProperties* AgentProperties = NULL, FSharedConstNavQueryFilter QueryFilter = NULL)
 	{
-		return ProjectPointToNavigation(Point, OutLocation, Extent, AgentProperties != NULL ? GetNavDataForProps(*AgentProperties) : GetDefaultNavDataInstance(FNavigationSystem::DontCreate), QueryFilter);
+		return ProjectPointToNavigation(Point, OutLocation, Extent, AgentProperties != NULL ? GetNavDataForProps(*AgentProperties, Point) : GetDefaultNavDataInstance(FNavigationSystem::DontCreate), QueryFilter);
 	}
 
 	// @todo document
@@ -441,12 +441,20 @@ public:
 	/** 
 	 * Looks for NavData generated for specified movement properties and returns it. NULL if not found;
 	 */
-	ANavigationData* GetNavDataForProps(const FNavAgentProperties& AgentProperties);
+	virtual ANavigationData* GetNavDataForProps(const FNavAgentProperties& AgentProperties);
 
 	/** 
 	 * Looks for NavData generated for specified movement properties and returns it. NULL if not found; Const version.
 	 */
-	const ANavigationData* GetNavDataForProps(const FNavAgentProperties& AgentProperties) const;
+	virtual const ANavigationData* GetNavDataForProps(const FNavAgentProperties& AgentProperties) const;
+
+	/**
+	 * Looks up NavData appropriate for specified movement properties and returns it. NULL if not found;
+	 * This is the encouraged way of querying for the appropriate NavData. It makes no difference for NavigationSystemV1
+	 *	(AgentLocation and Extent parameters not being used) but NaV2 will take advantage of it and all engine-level
+	 * AI navigation code is going to use call this flavor.
+	 */
+	virtual ANavigationData* GetNavDataForProps(const FNavAgentProperties& AgentProperties, const FVector& AgentLocation, const FVector& Extent = INVALID_NAVEXTENT) const;
 
 	/** Returns the world default navigation data instance. Creates one if it doesn't exist. */
 	ANavigationData* GetDefaultNavDataInstance(FNavigationSystem::ECreateIfMissing CreateNewIfNoneFound);
@@ -458,9 +466,11 @@ public:
 
 	ANavigationData* GetAbstractNavData() const { return AbstractNavData; }
 
-	/** constructs a navigation data instance of specified NavDataClass, in passed World
-	*	for supplied NavConfig */
-	virtual ANavigationData* CreateNavigationDataInstance(const FNavDataConfig& NavConfig);
+	/** constructs a navigation data instance of specified NavDataClass, in passed Level
+	 *	for supplied NavConfig. If Level == null and bSpawnNavDataInNavBoundsLevel == true
+	 *	then the first volume actor in RegisteredNavBounds will be used to source the level. 
+	 *	Otherwise the navdata instance will be spawned in NavigationSystem's world */
+	virtual ANavigationData* CreateNavigationDataInstanceInLevel(const FNavDataConfig& NavConfig, ULevel* SpawnLevel);
 
 	FSharedNavQueryFilter CreateDefaultQueryFilterCopy() const;
 
@@ -536,7 +546,10 @@ protected:
 	/** called in places where we need to spawn the NavOctree, but is checking additional conditions if we really want to do that
 	 *	depending on navigation data setup among others 
 	 *	@return true if NavOctree instance has been created, or if one is already present */
-	bool ConditionalPopulateNavOctree();
+	virtual bool ConditionalPopulateNavOctree();
+
+	/** called to instantiate NavigationSystem's NavOctree instance */
+	virtual void ConstructNavOctree();
 
 	/** Processes registration of candidates queues via RequestRegistration and stored in NavDataRegistrationQueue */
 	virtual void ProcessRegistrationCandidates();
@@ -576,6 +589,14 @@ public:
 	static void UpdateNavOctreeAfterMove(USceneComponent* Comp);
 
 protected:
+	/** A helper function that gathers all actors attached to RootActor and fetches 
+	 *	them back. The function does consider multi-level attachments. 
+	 *	@param AttachedActors is getting reset at the beginning of the function. 
+	 *		When done it is guaranteed to contain unique, non-null ptrs. It 
+	 *		will not include RootActor.
+	 *	@return the number of unique attached actors */
+	static int32 GetAllAttachedActors(const AActor& RootActor, TArray<AActor*>& OutAttachedActors);
+
 	/** updates navoctree information on actors attached to RootActor */
 	static void UpdateAttachedActorsInNavOctree(AActor& RootActor);
 
@@ -589,17 +610,16 @@ public:
 	void AddDirtyAreas(const TArray<FBox>& NewAreas, int32 Flags);
 	bool HasDirtyAreasQueued() const;
 
-	const FNavigationOctree* GetNavOctree() const { return NavOctree.Get(); }
-	FNavigationOctree* GetMutableNavOctree() { return NavOctree.Get(); }
+	const FNavigationOctree* GetNavOctree() const { return DefaultOctreeController.GetOctree(); }
+	FNavigationOctree* GetMutableNavOctree() { return DefaultOctreeController.GetMutableOctree(); }
 
 	FORCEINLINE static uint32 HashObject(const UObject& Object)
 	{
-		return Object.GetUniqueID();
+		return FNavigationOctree::HashObject(Object);
 	}
-	FORCEINLINE void SetObjectsNavOctreeId(const UObject& Object, FOctreeElementId Id) { ObjectToOctreeId.Add(HashObject(Object), Id); }
-	FORCEINLINE const FOctreeElementId* GetObjectsNavOctreeId(const UObject& Object) const { return ObjectToOctreeId.Find(HashObject(Object)); }
-	FORCEINLINE bool HasPendingObjectNavOctreeId(UObject* Object) const { return PendingOctreeUpdates.Contains(FNavigationDirtyElement(Object)); }
-	FORCEINLINE void RemoveObjectsNavOctreeId(const UObject& Object) { ObjectToOctreeId.Remove(HashObject(Object)); }
+	FORCEINLINE const FOctreeElementId* GetObjectsNavOctreeId(const UObject& Object) const { return DefaultOctreeController.GetObjectsNavOctreeId(Object); }
+	FORCEINLINE bool HasPendingObjectNavOctreeId(UObject* Object) const { return Object && DefaultOctreeController.HasPendingObjectNavOctreeId(*Object); }
+	FORCEINLINE void RemoveObjectsNavOctreeId(const UObject& Object) { DefaultOctreeController.RemoveObjectsNavOctreeId(Object); }
 
 	void RemoveNavOctreeElementId(const FOctreeElementId& ElementId, int32 UpdateFlags);
 
@@ -686,13 +706,13 @@ public:
 	}
 
 	/** check if navigation octree updates are currently ignored */
-	FORCEINLINE bool IsNavigationOctreeLocked() const { return bNavOctreeLock; }
+	FORCEINLINE bool IsNavigationOctreeLocked() const { return DefaultOctreeController.IsNavigationOctreeLocked(); }
 
 	// @todo document
 	UFUNCTION(BlueprintCallable, Category = "AI|Navigation")
 	void OnNavigationBoundsUpdated(ANavMeshBoundsVolume* NavVolume);
 	virtual void OnNavigationBoundsAdded(ANavMeshBoundsVolume* NavVolume);
-	void OnNavigationBoundsRemoved(ANavMeshBoundsVolume* NavVolume);
+	virtual void OnNavigationBoundsRemoved(ANavMeshBoundsVolume* NavVolume);
 
 	/** Used to display "navigation building in progress" notify */
 	bool IsNavigationBuildInProgress(bool bCheckDirtyToo = true);
@@ -711,10 +731,13 @@ protected:
 
 	/** spawn new crowd manager */
 	virtual void UpdateAbstractNavData();
-
-	/** Called during ConditionalPopulateNavOctree and gives subclassess a chance 
+	
+	/** Called during ConditionalPopulateNavOctree and gives subclassess a chance
 	 *	to influence what gets added */
 	virtual void AddLevelToOctree(ULevel& Level);
+
+	/** Called as part of UWorld::BeginTearingDown */
+	virtual void OnBeginTearingDown();
 	
 public:
 	/** Called upon UWorld destruction to release what needs to be released */
@@ -729,6 +752,7 @@ public:
 	virtual void OnWorldInitDone(FNavigationSystemRunMode Mode);
 
 	FORCEINLINE bool IsInitialized() const { return bWorldInitDone; }
+	FORCEINLINE FNavigationSystemRunMode GetRunMode() const { return OperationMode; }
 
 	/** adds BSP collisions of currently streamed in levels to octree */
 	void InitializeLevelCollisions();
@@ -736,7 +760,7 @@ public:
 	FORCEINLINE void AddNavigationBuildLock(uint8 Flags) { NavBuildingLockFlags |= Flags; }
 	void RemoveNavigationBuildLock(uint8 Flags, bool bSkipRebuildInEditor = false);
 
-	void SetNavigationOctreeLock(bool bLock) { bNavOctreeLock = bLock; }
+	void SetNavigationOctreeLock(bool bLock) { DefaultOctreeController.SetNavigationOctreeLock(bLock); }
 
 	/** checks if auto-rebuilding navigation data is enabled. Defaults to bNavigationAutoUpdateEnabled
 	*	value, but can be overridden per nav sys instance */
@@ -812,26 +836,19 @@ protected:
 
 	UPROPERTY()
 	FNavigationSystemRunMode OperationMode;
-
-	TSharedPtr<FNavigationOctree, ESPMode::ThreadSafe> NavOctree;
-
+	
 	TArray<FAsyncPathFindingQuery> AsyncPathFindingQueries;
 
 	FCriticalSection NavDataRegistration;
 
 	TMap<FNavAgentProperties, TWeakObjectPtr<ANavigationData> > AgentToNavDataMap;
 	
-	TMap<uint32, FOctreeElementId> ObjectToOctreeId;
-
-	/** Map of all objects that are tied to indexed navigation parent */
-	TMultiMap<UObject*, FWeakObjectPtr> OctreeChildNodesMap;
+	FNavigationOctreeController DefaultOctreeController;
 
 	/** Map of all custom navigation links, that are relevant for path following */
 	TMap<uint32, FNavigationSystem::FCustomLinkOwnerInfo> CustomLinksMap;
 
-	/** stores areas marked as dirty throughout the frame, processes them 
-	 *	once a frame in Tick function */
-	TArray<FNavigationDirtyArea> DirtyAreas;
+	FNavigationDirtyAreasController DefaultDirtyAreasController;
 
 	// async queries
 	FCriticalSection NavDataRegistrationSection;
@@ -844,27 +861,20 @@ protected:
 
 	/** set of locking flags applied on startup of navigation system */
 	uint8 InitialNavBuildingLockFlags;
-
-	/** if set, navoctree updates are ignored, use with caution! */
-	uint8 bNavOctreeLock : 1;
-
+	
 	uint8 bInitialSetupHasBeenPerformed : 1;
 	uint8 bInitialLevelsAdded : 1;
 	uint8 bWorldInitDone : 1;
+	/** set when the NavSys instance has been cleaned up. This is an irreversible state */
+	uint8 bCleanUpDone : 1;
 	uint8 bAsyncBuildPaused : 1; // mz@todo remove, replaced by bIsPIEActive and IsGameWorld
-	uint8 bCanAccumulateDirtyAreas : 1;
-#if !UE_BUILD_SHIPPING
-	uint8 bDirtyAreasReportedWhileAccumulationLocked : 1;
-#endif // !UE_BUILD_SHIPPING
+protected:
 
 	/** cached navigable world bounding box*/
 	mutable FBox NavigableWorldBounds;
 
 	/** indicates which of multiple navigation data instances to draw*/
 	int32 CurrentlyDrawnNavDataIndex;
-
-	/** temporary cumulative time to calculate when we need to update dirty areas */
-	float DirtyAreasUpdateTime;
 	
 #if !UE_BUILD_SHIPPING
 	/** self-registering exec command to handle nav sys console commands */
@@ -933,14 +943,13 @@ protected:
 	virtual void SpawnMissingNavigationData();
 
 protected:
-	virtual void RebuildDirtyAreas();
+	virtual void RebuildDirtyAreas(float DeltaSeconds);
 
-private:
 	// adds navigation bounds update request to a pending list
 	void AddNavigationBoundsUpdateRequest(const FNavigationBoundsUpdateRequest& UpdateRequest);
 
 	/** Triggers navigation building on all eligible navigation data. */
-	void RebuildAll(bool bIsLoadTime = false);
+	virtual void RebuildAll(bool bIsLoadTime = false);
 		 
 	/** Handler for FWorldDelegates::LevelAddedToWorld event */
 	 void OnLevelAddedToWorld(ULevel* InLevel, UWorld* InWorld);
@@ -959,7 +968,7 @@ private:
 	void PerformAsyncQueries(TArray<FAsyncPathFindingQuery> PathFindingQueries);
 
 	/** */
-	void DestroyNavOctree();
+	virtual void DestroyNavOctree();
 
 	/** Whether Navigation system needs to populate nav octree. 
 	 *  Depends on runtime generation settings of each navigation data, always true in the editor
@@ -976,7 +985,8 @@ private:
 	//----------------------------------------------------------------------//
 public:
 	void VerifyNavigationRenderingComponents(const bool bShow);
-	virtual int GetNavigationBoundsForNavData(const ANavigationData& NavData, TArray<FBox>& OutBounds) const;
+	/** @param if InLevel is given then only navigation bounds from that level will be considered*/
+	virtual int GetNavigationBoundsForNavData(const ANavigationData& NavData, TArray<FBox>& OutBounds, ULevel* InLevel = nullptr) const;
 	static INavigationDataInterface* GetNavDataForActor(const AActor& Actor);
 	virtual void Configure(const UNavigationSystemConfig& Config);
 	
@@ -1010,6 +1020,36 @@ public:
 	UE_DEPRECATED(4.22, "This version is deprecated.  Please use GetRandomLocationInNavigableRadius instead")
 	UFUNCTION(BlueprintPure, Category = "AI|Navigation", meta = (WorldContext = "WorldContextObject", DisplayName = "GetRandomPointInNavigableRadius", ScriptName = "GetRandomPointInNavigableRadius"))
 	static bool K2_GetRandomPointInNavigableRadius(UObject* WorldContextObject, const FVector& Origin, FVector& RandomLocation, float Radius, ANavigationData* NavData = NULL, TSubclassOf<UNavigationQueryFilter> FilterClass = NULL);
+	UE_DEPRECATED(4.23, "This function is deprecated.  Please use CreateNavigationDataInstanceInLevel instead")
+	virtual ANavigationData* CreateNavigationDataInstance(const FNavDataConfig& NavConfig);
+	
+	UE_DEPRECATED(4.24, "This function is deprecated and no longer used. NavigationSystem is no longer involved in storing navoctree element IDs. See FNavigationOctree for more details.")
+	void SetObjectsNavOctreeId(const UObject& Object, FOctreeElementId Id) {}
+
+	UE_DEPRECATED(4.24, "This member is deprecated and no longer used.  Please access OctreeController instead")
+	TMap<uint32, FOctreeElementId> ObjectToOctreeId;
+	UE_DEPRECATED(4.24, "This member is deprecated and no longer used.  Please access OctreeController instead")
+	TSet<FNavigationDirtyElement> PendingOctreeUpdates;
+	UE_DEPRECATED(4.24, "This member is deprecated and no longer used.  Please access OctreeController instead")
+	TSharedPtr<FNavigationOctree, ESPMode::ThreadSafe> NavOctree;
+	UE_DEPRECATED(4.24, "This member is deprecated and no longer used.  Please access OctreeController instead")
+	TMultiMap<UObject*, FWeakObjectPtr> OctreeChildNodesMap;
+	UE_DEPRECATED(4.24, "This member is deprecated and no longer used.  Please access OctreeController instead")
+	uint8 bNavOctreeLock : 1;
+
+	UE_DEPRECATED(4.24, "This member is deprecated and no longer used.  Please access DirtyAreasController instead")
+	UPROPERTY(config, EditAnywhere, Category = NavigationSystem)
+	float DirtyAreasUpdateFreq = 60;
+	UE_DEPRECATED(4.24, "This member is deprecated and no longer used.  Please access DirtyAreasController instead")
+	float DirtyAreasUpdateTime = 0;
+	UE_DEPRECATED(4.24, "This member is deprecated and no longer used.  Please access DirtyAreasController instead")
+	TArray<FNavigationDirtyArea> DirtyAreas;
+	UE_DEPRECATED(4.24, "This member is deprecated and no longer used.  Please access DirtyAreasController instead")
+	uint8 bCanAccumulateDirtyAreas : 1;
+#if !UE_BUILD_SHIPPING
+	UE_DEPRECATED(4.24, "This member is deprecated and no longer used.  Please access DirtyAreasController instead")
+	uint8 bDirtyAreasReportedWhileAccumulationLocked : 1;
+#endif // !UE_BUILD_SHIPPING
 };
 
 //----------------------------------------------------------------------//
