Added functionality to RecastNavMesh that allows one to procedurally bulk-change navmesh generation properties

As part of the process I've deprecated essentially unused ARecastNavMesh.AgentMaxHeight and removed the code accessing it.

diff --git a/Engine/Source/Runtime/NavigationSystem/Private/NavMesh/RecastNavMesh.cpp b/Engine/Source/Runtime/NavigationSystem/Private/NavMesh/RecastNavMesh.cpp
--- a/Engine/Source/Runtime/NavigationSystem/Private/NavMesh/RecastNavMesh.cpp
+++ b/Engine/Source/Runtime/NavigationSystem/Private/NavMesh/RecastNavMesh.cpp
@@ -213,6 +213,58 @@ namespace FNavMeshConfig
 	}
 }
 
+FRecastNavMeshGenerationProperties::FRecastNavMeshGenerationProperties()
+{
+	TilePoolSize = 1024;
+	TileSizeUU = 1000.f;
+	CellSize = 19;
+	CellHeight = 10;
+	AgentRadius = 34.f;
+	AgentHeight = 144.f;
+	AgentMaxSlope = 44.f;
+	AgentMaxStepHeight = 35.f;
+	MinRegionArea = 0.f;
+	MergeRegionSize = 400.f;
+	MaxSimplificationError = 1.3f;	// from RecastDemo
+	TileNumberHardLimit = 1 << 20;
+	RegionPartitioning = ERecastPartitioning::Watershed;
+	LayerPartitioning = ERecastPartitioning::Watershed;
+	RegionChunkSplits = 2;
+	LayerChunkSplits = 2;
+	bSortNavigationAreasByCost = false;
+	bPerformVoxelFiltering = true;
+	bMarkLowHeightAreas = false;
+	bFilterLowSpanSequences = false;
+	bFilterLowSpanFromTileCache = false;
+	bFixedTilePoolSize = false;
+}
+
+FRecastNavMeshGenerationProperties::FRecastNavMeshGenerationProperties(const ARecastNavMesh& RecastNavMesh)
+{
+	TilePoolSize = RecastNavMesh.TilePoolSize;
+	TileSizeUU = RecastNavMesh.TileSizeUU;
+	CellSize = RecastNavMesh.CellSize;
+	CellHeight = RecastNavMesh.CellHeight;
+	AgentRadius = RecastNavMesh.AgentRadius;
+	AgentHeight = RecastNavMesh.AgentHeight;
+	AgentMaxSlope = RecastNavMesh.AgentMaxSlope;
+	AgentMaxStepHeight = RecastNavMesh.AgentMaxStepHeight;
+	MinRegionArea = RecastNavMesh.MinRegionArea;
+	MergeRegionSize = RecastNavMesh.MergeRegionSize;
+	MaxSimplificationError = RecastNavMesh.MaxSimplificationError;
+	TileNumberHardLimit = RecastNavMesh.TileNumberHardLimit;
+	RegionPartitioning = RecastNavMesh.RegionPartitioning;
+	LayerPartitioning = RecastNavMesh.LayerPartitioning;
+	RegionChunkSplits = RecastNavMesh.RegionChunkSplits;
+	LayerChunkSplits = RecastNavMesh.LayerChunkSplits;
+	bSortNavigationAreasByCost = RecastNavMesh.bSortNavigationAreasByCost;
+	bPerformVoxelFiltering = RecastNavMesh.bPerformVoxelFiltering;
+	bMarkLowHeightAreas = RecastNavMesh.bMarkLowHeightAreas;
+	bFilterLowSpanSequences = RecastNavMesh.bFilterLowSpanSequences;
+	bFilterLowSpanFromTileCache = RecastNavMesh.bFilterLowSpanFromTileCache;
+	bFixedTilePoolSize = RecastNavMesh.bFixedTilePoolSize;
+}
+
 ARecastNavMesh::FNavPolyFlags ARecastNavMesh::NavLinkFlag = ARecastNavMesh::FNavPolyFlags(0);
 
 ARecastNavMesh::ARecastNavMesh(const FObjectInitializer& ObjectInitializer)
@@ -676,7 +728,7 @@ void ARecastNavMesh::Serialize( FArchive& Ar )
 void ARecastNavMesh::SetConfig(const FNavDataConfig& Src) 
 { 
 	NavDataConfig = Src; 
-	AgentMaxHeight = AgentHeight = Src.AgentHeight;
+	AgentHeight = Src.AgentHeight;
 	AgentRadius = Src.AgentRadius;
 
 	if (Src.HasStepHeightOverride())
@@ -1520,6 +1572,7 @@ bool ARecastNavMesh::GetPolyFlags(NavNodeRef PolyID, FNavMeshNodeFlags& Flags) c
 			const UClass* AreaClass = GetAreaClass(Flags.Area);
 			const UNavArea* DefArea = AreaClass ? ((UClass*)AreaClass)->GetDefaultObject<UNavArea>() : NULL;
 			Flags.AreaFlags = DefArea ? DefArea->GetAreaFlags() : 0;
+			// @todo what is this literal?
 			Flags.PathFlags = (PolyFlags & GetNavLinkFlag()) ? 4 : 0;
 		}
 	}
@@ -2375,6 +2428,32 @@ void ARecastNavMesh::ConditionalConstructGenerator()
 	}
 }
 
+void ARecastNavMesh::UpdateGenerationProperties(const FRecastNavMeshGenerationProperties& GenerationProps)
+{
+	TilePoolSize = GenerationProps.TilePoolSize;
+	TileSizeUU = GenerationProps.TileSizeUU;
+	CellSize = GenerationProps.CellSize;
+	CellHeight = GenerationProps.CellHeight;
+	AgentRadius = GenerationProps.AgentRadius;
+	AgentHeight = GenerationProps.AgentHeight;
+	AgentMaxSlope = GenerationProps.AgentMaxSlope;
+	AgentMaxStepHeight = GenerationProps.AgentMaxStepHeight;
+	MinRegionArea = GenerationProps.MinRegionArea;
+	MergeRegionSize = GenerationProps.MergeRegionSize;
+	MaxSimplificationError = GenerationProps.MaxSimplificationError;
+	TileNumberHardLimit = GenerationProps.TileNumberHardLimit;
+	RegionPartitioning = GenerationProps.RegionPartitioning;
+	LayerPartitioning = GenerationProps.LayerPartitioning;
+	RegionChunkSplits = GenerationProps.RegionChunkSplits;
+	LayerChunkSplits = GenerationProps.LayerChunkSplits;
+	bSortNavigationAreasByCost = GenerationProps.bSortNavigationAreasByCost;
+	bPerformVoxelFiltering = GenerationProps.bPerformVoxelFiltering;
+	bMarkLowHeightAreas = GenerationProps.bMarkLowHeightAreas;
+	bFilterLowSpanSequences = GenerationProps.bFilterLowSpanSequences;
+	bFilterLowSpanFromTileCache = GenerationProps.bFilterLowSpanFromTileCache;
+	bFixedTilePoolSize = GenerationProps.bFixedTilePoolSize;
+}
+
 bool ARecastNavMesh::IsVoxelCacheEnabled()
 {
 #if RECAST_ASYNC_REBUILDING
diff --git a/Engine/Source/Runtime/NavigationSystem/Private/NavMesh/RecastNavMeshGenerator.cpp b/Engine/Source/Runtime/NavigationSystem/Private/NavMesh/RecastNavMeshGenerator.cpp
--- a/Engine/Source/Runtime/NavigationSystem/Private/NavMesh/RecastNavMeshGenerator.cpp
+++ b/Engine/Source/Runtime/NavigationSystem/Private/NavMesh/RecastNavMeshGenerator.cpp
@@ -3311,7 +3311,6 @@ void FRecastNavMeshGenerator::ConfigureBuildProperties(FRecastBuildConfig& OutCo
 	const float CellSize = DestNavMesh->CellSize;
 	const float CellHeight = DestNavMesh->CellHeight;
 	const float AgentHeight = DestNavMesh->AgentHeight;
-	const float MaxAgentHeight = DestNavMesh->AgentMaxHeight;
 	const float AgentMaxSlope = DestNavMesh->AgentMaxSlope;
 	const float AgentMaxClimb = DestNavMesh->AgentMaxStepHeight;
 	const float AgentRadius = DestNavMesh->AgentRadius;
@@ -3341,7 +3340,6 @@ void FRecastNavMeshGenerator::ConfigureBuildProperties(FRecastBuildConfig& OutCo
 	OutConfig.maxVertsPerPoly = (int32)MAX_VERTS_PER_POLY;
 	OutConfig.detailSampleDist = 600.0f;
 	OutConfig.detailSampleMaxError = 1.0f;
-	OutConfig.PolyMaxHeight = (int32)ceilf(MaxAgentHeight / CellHeight);
 
 	OutConfig.minRegionArea = (int32)rcSqr(DestNavMesh->MinRegionArea / CellSize);
 	OutConfig.mergeRegionArea = (int32)rcSqr(DestNavMesh->MergeRegionSize / CellSize);
diff --git a/Engine/Source/Runtime/NavigationSystem/Public/NavMesh/RecastNavMesh.h b/Engine/Source/Runtime/NavigationSystem/Public/NavMesh/RecastNavMesh.h
--- a/Engine/Source/Runtime/NavigationSystem/Public/NavMesh/RecastNavMesh.h
+++ b/Engine/Source/Runtime/NavigationSystem/Public/NavMesh/RecastNavMesh.h
@@ -43,6 +43,7 @@ class UCanvas;
 class UNavArea;
 class UPrimitiveComponent;
 class URecastNavMeshDataChunk;
+class ARecastNavMesh;
 struct FRecastAreaNavModifierElement;
 class dtNavMesh;
 class dtQueryFilter;
@@ -206,6 +207,112 @@ namespace ERecastNamedFilter
 #endif //WITH_RECAST
 
 
+USTRUCT()
+struct NAVIGATIONSYSTEM_API FRecastNavMeshGenerationProperties
+{
+	GENERATED_BODY()
+
+	/** maximum number of tiles NavMesh can hold */
+	UPROPERTY(EditAnywhere, Category = Generation, meta = (editcondition = "bFixedTilePoolSize"))
+	int32 TilePoolSize;
+
+	/** size of single tile, expressed in uu */
+	UPROPERTY(EditAnywhere, Category = Generation, meta = (ClampMin = "300.0"))
+	float TileSizeUU;
+
+	/** horizontal size of voxelization cell */
+	UPROPERTY(EditAnywhere, Category = Generation, meta = (ClampMin = "1.0", ClampMax = "1024.0"))
+	float CellSize;
+
+	/** vertical size of voxelization cell */
+	UPROPERTY(EditAnywhere, Category = Generation, meta = (ClampMin = "1.0", ClampMax = "1024.0"))
+	float CellHeight;
+
+	/** Radius of largest agent that can freely traverse the generated navmesh */
+	UPROPERTY(EditAnywhere, Category = Generation, meta = (ClampMin = "0.0"))
+	float AgentRadius;
+
+	/** Size of the tallest agent that will path with this navmesh. */
+	UPROPERTY(EditAnywhere, Category = Generation, meta = (ClampMin = "0.0"))
+	float AgentHeight;
+
+	/* The maximum slope (angle) that the agent can move on. */
+	UPROPERTY(EditAnywhere, Category = Generation, meta = (ClampMin = "0.0", ClampMax = "89.0", UIMin = "0.0", UIMax = "89.0"))
+	float AgentMaxSlope;
+
+	/** Largest vertical step the agent can perform */
+	UPROPERTY(EditAnywhere, Category = Generation, meta = (ClampMin = "0.0"))
+	float AgentMaxStepHeight;
+
+	/* The minimum dimension of area. Areas smaller than this will be discarded */
+	UPROPERTY(EditAnywhere, Category = Generation, meta = (ClampMin = "0.0"))
+	float MinRegionArea;
+
+	/* The size limit of regions to be merged with bigger regions (watershed partitioning only) */
+	UPROPERTY(EditAnywhere, Category = Generation, meta = (ClampMin = "0.0"))
+	float MergeRegionSize;
+
+	/** How much navigable shapes can get simplified - the higher the value the more freedom */
+	UPROPERTY(EditAnywhere, Category = Generation, meta = (ClampMin = "0.0"))
+	float MaxSimplificationError;
+
+	/** Absolute hard limit to number of navmesh tiles. Be very, very careful while modifying it while
+	*	having big maps with navmesh. A single, empty tile takes 176 bytes and empty tiles are
+	*	allocated up front (subject to change, but that's where it's at now)
+	*	@note TileNumberHardLimit is always rounded up to the closest power of 2 */
+	UPROPERTY(EditAnywhere, Category = Generation, meta = (ClampMin = "1", UIMin = "1"), AdvancedDisplay)
+	int32 TileNumberHardLimit;
+
+	/** partitioning method for creating navmesh polys */
+	UPROPERTY(EditAnywhere, Category = Generation, AdvancedDisplay)
+	TEnumAsByte<ERecastPartitioning::Type> RegionPartitioning;
+
+	/** partitioning method for creating tile layers */
+	UPROPERTY(EditAnywhere, Category = Generation, AdvancedDisplay)
+	TEnumAsByte<ERecastPartitioning::Type> LayerPartitioning;
+
+	/** number of chunk splits (along single axis) used for region's partitioning: ChunkyMonotone */
+	UPROPERTY(EditAnywhere, Category = Generation, AdvancedDisplay)
+	int32 RegionChunkSplits;
+
+	/** number of chunk splits (along single axis) used for layer's partitioning: ChunkyMonotone */
+	UPROPERTY(EditAnywhere, Category = Generation, AdvancedDisplay)
+	int32 LayerChunkSplits;
+
+	/** Controls whether Navigation Areas will be sorted by cost before application
+	 *	to navmesh during navmesh generation. This is relevant when there are
+	 *	areas overlapping and we want to have area cost express area relevancy
+	 *	as well. Setting it to true will result in having area sorted by cost,
+	 *	but it will also increase navmesh generation cost a bit */
+	UPROPERTY(EditAnywhere, Category = Generation)
+	uint32 bSortNavigationAreasByCost : 1;
+
+	/** controls whether voxel filtering will be applied (via FRecastTileGenerator::ApplyVoxelFilter).
+	 *	Results in generated navmesh better fitting navigation bounds, but hits (a bit) generation performance */
+	UPROPERTY(EditAnywhere, Category = Generation, AdvancedDisplay)
+	uint32 bPerformVoxelFiltering : 1;
+
+	/** mark areas with insufficient free height above instead of cutting them out (accessible only for area modifiers using replace mode) */
+	UPROPERTY(EditAnywhere, Category = Generation, AdvancedDisplay)
+	uint32 bMarkLowHeightAreas : 1;
+
+	/** if set, only single low height span will be allowed under valid one */
+	UPROPERTY(EditAnywhere, Category = Generation, AdvancedDisplay)
+	uint32 bFilterLowSpanSequences : 1;
+
+	/** if set, only low height spans with corresponding area modifier will be stored in tile cache (reduces memory, can't modify without full tile rebuild) */
+	UPROPERTY(EditAnywhere, Category = Generation, AdvancedDisplay)
+	uint32 bFilterLowSpanFromTileCache : 1;
+
+	/** if true, the NavMesh will allocate fixed size pool for tiles, should be enabled to support streaming */
+	UPROPERTY(EditAnywhere, Category = Generation)
+	uint32 bFixedTilePoolSize : 1;
+
+	FRecastNavMeshGenerationProperties();
+	FRecastNavMeshGenerationProperties(const ARecastNavMesh& RecastNavMesh);
+};
+
+
 /**
  *	Structure to handle nav mesh tile's raw data persistence and releasing
  */
@@ -274,7 +381,7 @@ namespace FNavMeshConfig
 }
 
 
-UCLASS(config=Engine, defaultconfig, hidecategories=(Input,Rendering,Tags,Transform,"Utilities|Transformation",Actor,Layers,Replication), notplaceable)
+UCLASS(config=Engine, defaultconfig, hidecategories=(Input,Rendering,Tags,"Utilities|Transformation",Actor,Layers,Replication), notplaceable)
 class NAVIGATIONSYSTEM_API ARecastNavMesh : public ANavigationData
 {
 	GENERATED_UCLASS_BODY()
@@ -376,17 +483,19 @@ class NAVIGATIONSYSTEM_API ARecastNavMesh : public ANavigationData
 	UPROPERTY(EditAnywhere, Category = Generation, config, meta = (ClampMin = "0.0"))
 	float AgentRadius;
 
+	/** Size of the tallest agent that will path with this navmesh. */
 	UPROPERTY(EditAnywhere, Category = Generation, config, meta = (ClampMin = "0.0"))
 	float AgentHeight;
 
-	/** Size of the tallest agent that will path with this navmesh. */
-	UPROPERTY(EditAnywhere, Category=Generation, config, meta=(ClampMin = "0.0"))
+	UE_DEPRECATED(4.24, "ARecastNavMesh.AgentMaxHeight has been deprecated as it has no use. Use AgentHeight instead.")
+	UPROPERTY(VisibleAnywhere, Category=Generation, config, meta=(ClampMin = "0.0", DisplayName="DEPRECATED_AgentMaxHeight"))
 	float AgentMaxHeight;
 
 	/* The maximum slope (angle) that the agent can move on. */ 
 	UPROPERTY(EditAnywhere, Category=Generation, config, meta=(ClampMin = "0.0", ClampMax = "89.0", UIMin = "0.0", UIMax = "89.0" ))
 	float AgentMaxSlope;
 
+	/** Largest vertical step the agent can perform */
 	UPROPERTY(EditAnywhere, Category = Generation, config, meta = (ClampMin = "0.0"))
 	float AgentMaxStepHeight;
 
@@ -455,15 +564,15 @@ class NAVIGATIONSYSTEM_API ARecastNavMesh : public ANavigationData
 	int32 LayerChunkSplits;
 
 	/** Controls whether Navigation Areas will be sorted by cost before application 
-	 *	to navmesh during navmesh generation. This is relevant then there are
+	 *	to navmesh during navmesh generation. This is relevant when there are
 	 *	areas overlapping and we want to have area cost express area relevancy
 	 *	as well. Setting it to true will result in having area sorted by cost,
 	 *	but it will also increase navmesh generation cost a bit */
 	UPROPERTY(EditAnywhere, Category=Generation, config)
 	uint32 bSortNavigationAreasByCost:1;
 
-	/** controls whether voxel filterring will be applied (via FRecastTileGenerator::ApplyVoxelFilter). 
-	 *	Results in generated navemesh better fitting navigation bounds, but hits (a bit) generation performance */
+	/** controls whether voxel filtering will be applied (via FRecastTileGenerator::ApplyVoxelFilter). 
+	 *	Results in generated navmesh better fitting navigation bounds, but hits (a bit) generation performance */
 	UPROPERTY(EditAnywhere, Category=Generation, config, AdvancedDisplay)
 	uint32 bPerformVoxelFiltering:1;
 
@@ -901,6 +1010,7 @@ public:
 	virtual bool SupportsRuntimeGeneration() const override;
 	virtual bool SupportsStreaming() const override;
 	virtual void ConditionalConstructGenerator() override;
+	void UpdateGenerationProperties(const FRecastNavMeshGenerationProperties& GenerationProps);
 	bool ShouldGatherDataOnGameThread() const { return bDoFullyAsyncNavDataGathering == false; }
 	int32 GetTileNumberHardLimit() const { return TileNumberHardLimit; }
 
diff --git a/Engine/Source/Runtime/NavigationSystem/Public/NavMesh/RecastNavMeshGenerator.h b/Engine/Source/Runtime/NavigationSystem/Public/NavMesh/RecastNavMeshGenerator.h
--- a/Engine/Source/Runtime/NavigationSystem/Public/NavMesh/RecastNavMeshGenerator.h
+++ b/Engine/Source/Runtime/NavigationSystem/Public/NavMesh/RecastNavMeshGenerator.h
@@ -50,6 +50,7 @@ struct FRecastBuildConfig : public rcConfig
 	/** chunk size for ChunkyMonotone partitioning */
 	int32 TileCacheChunkSize;
 
+	UE_DEPRECATED(4.24, "FRecastBuildConfig.PolyMaxHeight has been deprecated as it has no use")
 	int32 PolyMaxHeight;
 	/** indicates what's the limit of navmesh polygons per tile. This value is calculated from other
 	 *	factors - DO NOT SET IT TO ARBITRARY VALUE */
@@ -78,7 +79,10 @@ struct FRecastBuildConfig : public rcConfig
 		bMarkLowHeightAreas = false;
 		bFilterLowSpanSequences = false;
 		bFilterLowSpanFromTileCache = false;
+		PRAGMA_DISABLE_DEPRECATION_WARNINGS
+		// Still initializing, even though the property is deprecated, to avoid static analysis warnings
 		PolyMaxHeight = 10;
+		PRAGMA_ENABLE_DEPRECATION_WARNINGS
 		MaxPolysPerTile = -1;
 		AgentIndex = 0;
 	}
