Fixed NavigationSystemConfigOverride never creating an instance of NavigationSystem in  PIE if it gets loaded as part of a sublevel

diff --git a/Engine/Source/Runtime/NavigationSystem/Private/NavSystemConfigOverride.cpp b/Engine/Source/Runtime/NavigationSystem/Private/NavSystemConfigOverride.cpp
--- a/Engine/Source/Runtime/NavigationSystem/Private/NavSystemConfigOverride.cpp
+++ b/Engine/Source/Runtime/NavigationSystem/Private/NavSystemConfigOverride.cpp
@@ -75,7 +75,7 @@ void ANavSystemConfigOverride::PostLoad()
 		if (World->bIsWorldInitialized
 			&& NavigationSystemConfig
 #if WITH_EDITOR
-			&& (GIsEditorLoadingPackage == false)
+			&& (GIsEditorLoadingPackage == false || World->IsGameWorld())
 #endif // WITH_EDITOR
 			)
 		{
