Fixed wrong copyright dates in NavigationDataHandler

diff --git a/Engine/Source/Runtime/NavigationSystem/Private/NavigationDataHandler.cpp b/Engine/Source/Runtime/NavigationSystem/Private/NavigationDataHandler.cpp
--- a/Engine/Source/Runtime/NavigationSystem/Private/NavigationDataHandler.cpp
+++ b/Engine/Source/Runtime/NavigationSystem/Private/NavigationDataHandler.cpp
@@ -1,4 +1,4 @@
-// Copyright 1998-2018 Epic Games, Inc. All Rights Reserved.
+// Copyright 1998-2019 Epic Games, Inc. All Rights Reserved.
 
 #include "NavigationDataHandler.h"
 #include "NavMesh/RecastNavMeshGenerator.h"
