Removed the octree reference from FNavigationOctreeElement

Submitting to Main as a follow up to 7308751

Yoan.StAmant
#okforgitub public

diff --git a/Engine/Source/Runtime/NavigationSystem/Private/NavigationDataHandler.cpp b/Engine/Source/Runtime/NavigationSystem/Private/NavigationDataHandler.cpp
--- a/Engine/Source/Runtime/NavigationSystem/Private/NavigationDataHandler.cpp
+++ b/Engine/Source/Runtime/NavigationSystem/Private/NavigationDataHandler.cpp
@@ -108,7 +108,7 @@ void FNavigationDataHandler::AddElementToNavOctree(const FNavigationDirtyElement
 		return;
 	}
 
-	FNavigationOctreeElement GeneratedData(*OctreeController.NavOctree.Get(), *ElementOwner);
+	FNavigationOctreeElement GeneratedData(*ElementOwner);
 	const FBox ElementBounds = DirtyElement.NavInterface->GetNavigationBounds();
 
 	UObject* NavigationParent = DirtyElement.NavInterface->GetNavigationParent();
@@ -401,7 +401,7 @@ void FNavigationDataHandler::AddLevelCollisionToOctree(ULevel& Level)
 
 		if (!ElementId && LevelGeom && LevelGeom->Num() > 0)
 		{
-			FNavigationOctreeElement BSPElem(*OctreeController.NavOctree.Get(), Level);
+			FNavigationOctreeElement BSPElem(Level);
 			FRecastNavMeshGenerator::ExportVertexSoupGeometry(*LevelGeom, *BSPElem.Data);
 
 			const FBox& Bounds = BSPElem.Data->Bounds;
diff --git a/Engine/Source/Runtime/NavigationSystem/Private/NavigationOctree.cpp b/Engine/Source/Runtime/NavigationSystem/Private/NavigationOctree.cpp
--- a/Engine/Source/Runtime/NavigationSystem/Private/NavigationOctree.cpp
+++ b/Engine/Source/Runtime/NavigationSystem/Private/NavigationOctree.cpp
@@ -4,10 +4,6 @@
 #include "AI/Navigation/NavRelevantInterface.h"
 #include "NavigationSystem.h"
 
-namespace
-{
-	FNavigationOctree DummyNavOctree(FVector(0), 0);
-}
 
 //----------------------------------------------------------------------//
 // FNavigationOctree
@@ -202,7 +198,7 @@ const FNavigationRelevantData* FNavigationOctree::GetDataForID(const FOctreeElem
 	return &*OctreeElement.Data;
 }
 
-void FNavigationOctree::SetElementId(const UObject& Object, FOctreeElementId Id)
+void FNavigationOctree::SetElementIdImpl(const UObject& Object, FOctreeElementId Id)
 {
 	ObjectToOctreeId.Add(HashObject(Object), Id);
 }
@@ -213,23 +209,13 @@ void FNavigationOctree::SetElementId(const UObject& Object, FOctreeElementId Id)
 #if NAVSYS_DEBUG
 FORCENOINLINE
 #endif // NAVSYS_DEBUG
-void FNavigationOctreeSemantics::SetElementId(const FNavigationOctreeElement& Element, FOctreeElementId Id)
+void FNavigationOctreeSemantics::SetElementId(FNavigationOctreeSemantics::FOctree& OctreeOwner, const FNavigationOctreeElement& Element, FOctreeElementId Id)
 {
 
 	const bool bEvenIfPendingKill = true;
 	UObject* ElementOwner = Element.GetOwner(bEvenIfPendingKill);
 	if (ElementOwner)
 	{
-		Element.GetOwnerOctree().SetElementId(*ElementOwner, Id);
+		((FNavigationOctree&)OctreeOwner).SetElementIdImpl(*ElementOwner, Id);
 	}
 }
-
-//----------------------------------------------------------------------//
-//  FNavigationOctreeElement
-//----------------------------------------------------------------------//
-FNavigationOctreeElement::FNavigationOctreeElement(UObject& SourceObject)
-	: Data(new FNavigationRelevantData(SourceObject))
-	, OwnerOctree(DummyNavOctree)
-{
-
-}
\ No newline at end of file
diff --git a/Engine/Source/Runtime/NavigationSystem/Public/NavigationOctree.h b/Engine/Source/Runtime/NavigationSystem/Public/NavigationOctree.h
--- a/Engine/Source/Runtime/NavigationSystem/Public/NavigationOctree.h
+++ b/Engine/Source/Runtime/NavigationSystem/Public/NavigationOctree.h
@@ -20,22 +20,15 @@ struct NAVIGATIONSYSTEM_API FNavigationOctreeElement
 {
 	FBoxSphereBounds Bounds;
 	TSharedRef<FNavigationRelevantData, ESPMode::ThreadSafe> Data;
-private:
-	FNavigationOctree& OwnerOctree;
 
 public:
-	UE_DEPRECATED(4.24, "Single-paramater FNavigationOctreeElement ctor is deprecated. Use the one explicitly passing the FNavigationOctree reference in")
-	explicit FNavigationOctreeElement(UObject& SourceObject);
-
-	FNavigationOctreeElement(FNavigationOctree& InOwnerOctree, UObject& SourceObject)
+	explicit FNavigationOctreeElement(UObject& SourceObject)
 		: Data(new FNavigationRelevantData(SourceObject))
-		, OwnerOctree(InOwnerOctree)
 	{}
 
 	FNavigationOctreeElement(const FNavigationOctreeElement& Other)
 		: Bounds(Other.Bounds)
 		, Data(Other.Data)
-		, OwnerOctree(Other.OwnerOctree)
 	{}
 
 	FNavigationOctreeElement& operator=(const FNavigationOctreeElement& Other)
@@ -86,12 +79,11 @@ public:
 	}
 
 	FORCEINLINE UObject* GetOwner(bool bEvenIfPendingKill = false) const { return Data->SourceObject.Get(bEvenIfPendingKill); }
-
-	FORCEINLINE FNavigationOctree& GetOwnerOctree() const { return OwnerOctree; }
 };
 
 struct FNavigationOctreeSemantics
 {
+	typedef TOctree<FNavigationOctreeElement, FNavigationOctreeSemantics> FOctree;
 	enum { MaxElementsPerLeaf = 16 };
 	enum { MinInclusiveElementsPerNode = 7 };
 	enum { MaxNodeDepth = 12 };
@@ -117,7 +109,7 @@ struct FNavigationOctreeSemantics
 #if NAVSYS_DEBUG
 	FORCENOINLINE 
 #endif // NAVSYS_DEBUG
-	static void SetElementId(const FNavigationOctreeElement& Element, FOctreeElementId Id);
+	static void SetElementId(FOctree& OctreeOwner, const FNavigationOctreeElement& Element, FOctreeElementId Id);
 };
 
 class NAVIGATIONSYSTEM_API FNavigationOctree : public TOctree<FNavigationOctreeElement, FNavigationOctreeSemantics>, public TSharedFromThis<FNavigationOctree, ESPMode::ThreadSafe>
@@ -161,7 +153,6 @@ public:
 	void DemandLazyDataGathering(const FNavigationOctreeElement& Element);
 	void DemandLazyDataGathering(FNavigationRelevantData& ElementData);
 
-	void SetElementId(const UObject& Object, FOctreeElementId Id);
 
 	FORCEINLINE static uint32 HashObject(const UObject& Object)
 	{
@@ -170,6 +161,9 @@ public:
 
 protected:
 	friend struct FNavigationOctreeController;
+	friend struct FNavigationOctreeSemantics;
+
+	void SetElementIdImpl(const UObject& Object, FOctreeElementId Id);
 
 	TMap<uint32, FOctreeElementId> ObjectToOctreeId;
 	ENavDataGatheringMode DefaultGeometryGatheringMode;
