Added some ensure in NavigationSystem to track down some server crashes related to obsolete entries in ObjectToOctreeId

#rnx

diff --git a/Engine/Source/Runtime/NavigationSystem/Private/NavigationDataHandler.cpp b/Engine/Source/Runtime/NavigationSystem/Private/NavigationDataHandler.cpp
--- a/Engine/Source/Runtime/NavigationSystem/Private/NavigationDataHandler.cpp
+++ b/Engine/Source/Runtime/NavigationSystem/Private/NavigationDataHandler.cpp
@@ -22,7 +22,7 @@ FNavigationDataHandler::FNavigationDataHandler(FNavigationOctreeController& InOc
 
 void FNavigationDataHandler::RemoveNavOctreeElementId(const FOctreeElementId& ElementId, int32 UpdateFlags)
 {
-	if (OctreeController.NavOctree->IsValidElementId(ElementId))
+	if (ensure(OctreeController.IsValidElement(ElementId)))
 	{
 		const FNavigationOctreeElement& ElementData = OctreeController.NavOctree->GetElementById(ElementId);
 		const int32 DirtyFlag = GetDirtyFlagHelper(UpdateFlags, ElementData.Data->GetDirtyFlag());
@@ -126,11 +126,11 @@ void FNavigationDataHandler::AddElementToNavOctree(const FNavigationDirtyElement
 			ParentNode.bInvalidRequest = true;
 		}
 
-		const FOctreeElementId* UseParentId = ParentId ? ParentId : OctreeController.GetObjectsNavOctreeId(*NavigationParent);
-		if (UseParentId && OctreeController.NavOctree->IsValidElementId(*UseParentId))
+		const FOctreeElementId* ElementId = ParentId ? ParentId : OctreeController.GetObjectsNavOctreeId(*NavigationParent);
+		if (ElementId && ensure(OctreeController.IsValidElement(*ElementId)))
 		{
 			UE_LOG(LogNavOctree, Log, TEXT("ADD %s to %s"), *GetNameSafe(ElementOwner), *GetNameSafe(NavigationParent));
-			OctreeController.NavOctree->AppendToNode(*UseParentId, DirtyElement.NavInterface, ElementBounds, GeneratedData);
+			OctreeController.NavOctree->AppendToNode(*ElementId, DirtyElement.NavInterface, ElementBounds, GeneratedData);
 		}
 		else
 		{
@@ -307,15 +307,16 @@ void FNavigationDataHandler::UpdateNavOctreeParentChain(UObject& ElementOwner, b
 bool FNavigationDataHandler::UpdateNavOctreeElementBounds(UActorComponent& Comp, const FBox& NewBounds, const FBox& DirtyArea)
 {
 	const FOctreeElementId* ElementId = OctreeController.GetObjectsNavOctreeId(Comp);
-	if (ElementId && ElementId->IsValidId())
+	if (ElementId != nullptr && ensure(OctreeController.IsValidElement(*ElementId)))
 	{
 		OctreeController.NavOctree->UpdateNode(*ElementId, NewBounds);
 
 		// Add dirty area
 		if (DirtyArea.IsValid)
 		{
+			// Refresh ElementId since components may be stored in a different node after updating bounds
 			ElementId = OctreeController.GetObjectsNavOctreeId(Comp);
-			if (ElementId && ElementId->IsValidId())
+			if (ElementId != nullptr && ensure(OctreeController.IsValidElement(*ElementId)))
 			{
 				FNavigationOctreeElement& ElementData = OctreeController.NavOctree->GetElementById(*ElementId);
 				DirtyAreasController.AddArea(DirtyArea, ElementData.Data->GetDirtyFlag());
@@ -426,7 +427,7 @@ void FNavigationDataHandler::RemoveLevelCollisionFromOctree(ULevel& Level)
 
 		if (ElementId != nullptr)
 		{
-			if (OctreeController.NavOctree->IsValidElementId(*ElementId))
+			if (ensure(OctreeController.IsValidElement(*ElementId)))
 			{
 				// mark area occupied by given actor as dirty
 				FNavigationOctreeElement& ElementData = OctreeController.NavOctree->GetElementById(*ElementId);
diff --git a/Engine/Source/Runtime/NavigationSystem/Private/NavigationOctree.cpp b/Engine/Source/Runtime/NavigationSystem/Private/NavigationOctree.cpp
--- a/Engine/Source/Runtime/NavigationSystem/Private/NavigationOctree.cpp
+++ b/Engine/Source/Runtime/NavigationSystem/Private/NavigationOctree.cpp
@@ -136,6 +136,14 @@ void FNavigationOctree::AddNode(UObject* ElementOb, INavRelevantInterface* NavEl
 	NodesMemory += ElementMemory;
 	INC_MEMORY_STAT_BY(STAT_Navigation_CollisionTreeMemory, ElementMemory);
 
+	if (ElementOb != nullptr)
+	{
+		ensure(ObjectToOctreeId.Find(HashObject(*ElementOb)) == nullptr);
+
+		// Reserve entry in our map before octree calls us back with the assigned Id. This is required since we ensure that all set id already exists
+		ObjectToOctreeId.Add(HashObject(*ElementOb), FOctreeElementId());
+	}	
+
 	AddElement(Element);
 }
 
@@ -200,7 +208,15 @@ const FNavigationRelevantData* FNavigationOctree::GetDataForID(const FOctreeElem
 
 void FNavigationOctree::SetElementIdImpl(const UObject& Object, FOctreeElementId Id)
 {
-	ObjectToOctreeId.Add(HashObject(Object), Id);
+	ensure(IsValidElementId(Id));
+
+	FOctreeElementId* ExistingElementId = ObjectToOctreeId.Find(HashObject(Object));
+	// We expect ObjectToOctreeId to already contain an entry for Object since it should have already 
+	// been added in FNavigationOctree::AddNode
+	if (ensure(ExistingElementId))
+	{
+		*ExistingElementId = Id;
+	}
 }
 
 //----------------------------------------------------------------------//
@@ -211,10 +227,9 @@ FORCENOINLINE
 #endif // NAVSYS_DEBUG
 void FNavigationOctreeSemantics::SetElementId(FNavigationOctreeSemantics::FOctree& OctreeOwner, const FNavigationOctreeElement& Element, FOctreeElementId Id)
 {
-
 	const bool bEvenIfPendingKill = true;
 	UObject* ElementOwner = Element.GetOwner(bEvenIfPendingKill);
-	if (ElementOwner)
+	if (ensure(ElementOwner != nullptr))
 	{
 		((FNavigationOctree&)OctreeOwner).SetElementIdImpl(*ElementOwner, Id);
 	}
diff --git a/Engine/Source/Runtime/NavigationSystem/Private/NavigationOctreeController.cpp b/Engine/Source/Runtime/NavigationSystem/Private/NavigationOctreeController.cpp
--- a/Engine/Source/Runtime/NavigationSystem/Private/NavigationOctreeController.cpp
+++ b/Engine/Source/Runtime/NavigationSystem/Private/NavigationOctreeController.cpp
@@ -27,16 +27,13 @@ void FNavigationOctreeController::SetNavigableGeometryStoringMode(FNavigationOct
 bool FNavigationOctreeController::GetNavOctreeElementData(const UObject& NodeOwner, int32& DirtyFlags, FBox& DirtyBounds)
 {
 	const FOctreeElementId* ElementId = GetObjectsNavOctreeId(NodeOwner);
-	if (ElementId != NULL)
+	if (ElementId != nullptr && ensure(IsValidElement(*ElementId)))
 	{
-		if (NavOctree->IsValidElementId(*ElementId))
-		{
-			// mark area occupied by given actor as dirty
-			FNavigationOctreeElement& ElementData = NavOctree->GetElementById(*ElementId);
-			DirtyFlags = ElementData.Data->GetDirtyFlag();
-			DirtyBounds = ElementData.Bounds.GetBox();
-			return true;
-		}
+		// mark area occupied by given actor as dirty
+		FNavigationOctreeElement& ElementData = NavOctree->GetElementById(*ElementId);
+		DirtyFlags = ElementData.Data->GetDirtyFlag();
+		DirtyBounds = ElementData.Bounds.GetBox();
+		return true;
 	}
 
 	return false;
@@ -44,11 +41,11 @@ bool FNavigationOctreeController::GetNavOctreeElementData(const UObject& NodeOwn
 
 const FNavigationRelevantData* FNavigationOctreeController::GetDataForObject(const UObject& Object) const
 {
-	const FOctreeElementId* OctreeID = GetObjectsNavOctreeId(Object);
+	const FOctreeElementId* ElementId = GetObjectsNavOctreeId(Object);
 
-	if (OctreeID != nullptr && ensure(NavOctree.IsValid()) && NavOctree->IsValidElementId(*OctreeID))
+	if (ElementId != nullptr && ensure(IsValidElement(*ElementId)))
 	{
-		return NavOctree->GetDataForID(*OctreeID);
+		return NavOctree->GetDataForID(*ElementId);
 	}
 
 	return nullptr;
diff --git a/Engine/Source/Runtime/NavigationSystem/Public/NavigationOctreeController.h b/Engine/Source/Runtime/NavigationSystem/Public/NavigationOctreeController.h
--- a/Engine/Source/Runtime/NavigationSystem/Public/NavigationOctreeController.h
+++ b/Engine/Source/Runtime/NavigationSystem/Public/NavigationOctreeController.h
@@ -39,6 +39,7 @@ struct NAVIGATIONSYSTEM_API FNavigationOctreeController
 	bool IsNavigationOctreeLocked() const;
 	/** basically says if navoctree has been created already */
 	bool IsValid() const { return NavOctree.IsValid(); }
+	bool IsValidElement(const FOctreeElementId& ElementId) const;
 
 private:
 	static uint32 HashObject(const UObject& Object);
@@ -96,3 +97,8 @@ FORCEINLINE void FNavigationOctreeController::SetNavigationOctreeLock(bool bLock
 { 
 	bNavOctreeLock = bLock; 
 }
+
+FORCEINLINE bool FNavigationOctreeController::IsValidElement(const FOctreeElementId& ElementId) const
+{
+	return IsValid() && NavOctree->IsValidElementId(ElementId);
+}
