Made sure NavigationSystemConfig respects GEngine->NavigationSystemClassName #UE4

diff --git a/Engine/Source/Runtime/Engine/Private/AI/NavigationSystemConfig.cpp b/Engine/Source/Runtime/Engine/Private/AI/NavigationSystemConfig.cpp
--- a/Engine/Source/Runtime/Engine/Private/AI/NavigationSystemConfig.cpp
+++ b/Engine/Source/Runtime/Engine/Private/AI/NavigationSystemConfig.cpp
@@ -9,7 +9,7 @@
 UNavigationSystemConfig::UNavigationSystemConfig(const FObjectInitializer& ObjectInitializer)
 	: Super(ObjectInitializer)
 {
-	if (GEngine && *GEngine->NavigationSystemClass != nullptr)
+	if (GEngine)
 	{
 		NavigationSystemClass = GEngine->NavigationSystemClassName;
 	}
diff --git a/Engine/Source/Runtime/NavigationSystem/Private/NavigationSystem.cpp b/Engine/Source/Runtime/NavigationSystem/Private/NavigationSystem.cpp
--- a/Engine/Source/Runtime/NavigationSystem/Private/NavigationSystem.cpp
+++ b/Engine/Source/Runtime/NavigationSystem/Private/NavigationSystem.cpp
@@ -4438,7 +4438,8 @@ void UNavigationSystemModuleConfig::PostInitProperties()
 
 void UNavigationSystemModuleConfig::UpdateWithNavSysCDO(const UNavigationSystemV1& NavSysCDO)
 {
-	if (NavigationSystemClass.ResolveClass() == NavSysCDO.GetClass())
+	UClass* MyClass = NavigationSystemClass.ResolveClass();
+	if (MyClass != nullptr && MyClass->IsChildOf(NavSysCDO.GetClass()))
 	{
 		bStrictlyStatic = NavSysCDO.bStaticRuntimeNavigation;
 		bCreateOnClient = NavSysCDO.bAllowClientSideNavigation;
