Fixed a bug in how UNavModifierComponent handles FKConvexElems #UE4

diff --git a/Engine/Source/Runtime/NavigationSystem/Private/NavModifierComponent.cpp b/Engine/Source/Runtime/NavigationSystem/Private/NavModifierComponent.cpp
--- a/Engine/Source/Runtime/NavigationSystem/Private/NavModifierComponent.cpp
+++ b/Engine/Source/Runtime/NavigationSystem/Private/NavModifierComponent.cpp
@@ -72,7 +72,8 @@ void UNavModifierComponent::CalcAndCacheBounds() const
 						const FKConvexElem& ElemInfo = BodySetup->AggGeom.ConvexElems[ConvexIdx];
 						FTransform ElemTM = ElemInfo.GetTransform();
 						ElemTM.ScaleTranslation(Scale3D);
-						ElemTM *= FTransform(ParentTM.GetLocation());
+						// no need to translate with ParentTM since FKConvexElem 
+						// we're getting are already translated
 
 						const FBox ConvexBounds = ElemInfo.CalcAABB(ElemTM, Scale3D);
 						ComponentBounds.Add(FRotatedBox(ConvexBounds, ElemTM.GetRotation() * ParentTM.GetRotation()));
