Manual merge of CL#4578611 from 4.21. Original description:

Added an additional condition to FNavigationSystem::ShouldLoadNavigationOnClient to avoid triggering a casting ensure #UE4

diff --git a/Engine/Source/Runtime/NavigationSystem/Private/NavigationSystem.cpp b/Engine/Source/Runtime/NavigationSystem/Private/NavigationSystem.cpp
--- a/Engine/Source/Runtime/NavigationSystem/Private/NavigationSystem.cpp
+++ b/Engine/Source/Runtime/NavigationSystem/Private/NavigationSystem.cpp
@@ -125,13 +125,12 @@ namespace FNavigationSystem
 			const UNavigationSystemV1* NavSys = Cast<UNavigationSystemV1>(World->GetNavigationSystem());
 			return NavSys && NavSys->ShouldLoadNavigationOnClient(&NavData);
 		}
-		else
+		else if (GEngine->NavigationSystemClass && GEngine->NavigationSystemClass->IsChildOf<UNavigationSystemV1>())
 		{
-			const UNavigationSystemV1* NavSysCDO = (*GEngine->NavigationSystemClass != nullptr)
-				? (GEngine->NavigationSystemClass->GetDefaultObject<const UNavigationSystemV1>())
-				: (const UNavigationSystemV1*)nullptr;
+			const UNavigationSystemV1* NavSysCDO = GEngine->NavigationSystemClass->GetDefaultObject<const UNavigationSystemV1>();
 			return NavSysCDO && NavSysCDO->ShouldLoadNavigationOnClient(&NavData);
 		}
+		return false;
 	}
 
 	bool ShouldDiscardSubLevelNavData(ANavigationData& NavData)
