PR #5390: Engine crash caused by a null pointer in navigation system (Contributed by [REDACTED_PERSON])

diff --git a/Engine/Source/Runtime/NavigationSystem/Private/NavMesh/NavMeshPath.cpp b/Engine/Source/Runtime/NavigationSystem/Private/NavMesh/NavMeshPath.cpp
--- a/Engine/Source/Runtime/NavigationSystem/Private/NavMesh/NavMeshPath.cpp
+++ b/Engine/Source/Runtime/NavigationSystem/Private/NavMesh/NavMeshPath.cpp
@@ -110,8 +110,11 @@ const TArray<FNavigationPortalEdge>& FNavMeshPath::GeneratePathCorridorEdges() c
 	if (CorridorLength != 0 && IsInGameThread() && NavigationDataUsed.IsValid())
 	{
 		const ARecastNavMesh* MyOwner = Cast<ARecastNavMesh>(GetNavigationDataUsed());
-		MyOwner->GetEdgesForPathCorridor(&PathCorridor, &PathCorridorEdges);
-		bCorridorEdgesGenerated = PathCorridorEdges.Num() > 0;
+		if (MyOwner)
+		{
+			MyOwner->GetEdgesForPathCorridor(&PathCorridor, &PathCorridorEdges);
+			bCorridorEdgesGenerated = (PathCorridorEdges.Num() > 0);
+		}
 	}
 #endif // WITH_RECAST
 	return PathCorridorEdges;
@@ -121,7 +124,7 @@ void FNavMeshPath::PerformStringPulling(const FVector& StartLoc, const FVector&
 {
 #if WITH_RECAST
 	const ARecastNavMesh* MyOwner = Cast<ARecastNavMesh>(GetNavigationDataUsed());
-	if (PathCorridor.Num())
+	if (::IsValid(MyOwner) && PathCorridor.Num())
 	{
 		bStringPulled = MyOwner->FindStraightPath(StartLoc, EndLoc, PathCorridor, PathPoints, &CustomLinkIds);
 	}
@@ -282,7 +285,7 @@ void FNavMeshPath::OffsetFromCorners(float Distance)
 	SCOPE_CYCLE_COUNTER(STAT_Navigation_OffsetFromCorners);
 
 	const ARecastNavMesh* MyOwner = Cast<ARecastNavMesh>(GetNavigationDataUsed());
-	if (PathPoints.Num() == 0 || PathPoints.Num() > 100)
+	if (MyOwner == nullptr || PathPoints.Num() == 0 || PathPoints.Num() > 100)
 	{
 		// skip it, there is not need to offset that path from performance point of view
 		return;
@@ -731,8 +734,11 @@ bool FNavMeshPath::GetNodeFlags(int32 NodeIdx, FNavMeshNodeFlags& Flags) const
 		{
 #if WITH_RECAST
 			const ARecastNavMesh* MyOwner = Cast<ARecastNavMesh>(GetNavigationDataUsed());
-			MyOwner->GetPolyFlags(PathCorridor[NodeIdx], Flags);
-			bResult = true;
+			if (MyOwner)
+			{
+				MyOwner->GetPolyFlags(PathCorridor[NodeIdx], Flags);
+				bResult = true;
+			}
 #endif	// WITH_RECAST
 		}
 	}
@@ -811,6 +817,10 @@ void FNavMeshPath::DescribeSelfToVisLog(FVisualLogEntry* Snapshot) const
 	int32 NumAreaMark = 1;
 
 	ARecastNavMesh* NavMesh = Cast<ARecastNavMesh>(GetNavigationDataUsed());
+	if (NavMesh == nullptr)
+	{
+		return;
+	}
 	NavMesh->BeginBatchQuery();
 
 	TArray<FVector> Verts;
