Deprecated K2_GetRandomPointInNavigableRadius, and replaced with K2_GetRandomLocationInNavigableRadius. Changed to BlueprintCallable.
[RN] fix general

diff --git a/Engine/Source/Runtime/NavigationSystem/Private/NavigationSystem.cpp b/Engine/Source/Runtime/NavigationSystem/Private/NavigationSystem.cpp
--- a/Engine/Source/Runtime/NavigationSystem/Private/NavigationSystem.cpp
+++ b/Engine/Source/Runtime/NavigationSystem/Private/NavigationSystem.cpp
@@ -3874,10 +3874,11 @@ bool UNavigationSystemV1::K2_GetRandomReachablePointInRadius(UObject* WorldConte
 	return bResult;
 }
 
-bool UNavigationSystemV1::K2_GetRandomPointInNavigableRadius(UObject* WorldContextObject, const FVector& Origin, FVector& RandomLocation, float Radius, ANavigationData* NavData, TSubclassOf<UNavigationQueryFilter> FilterClass)
+bool UNavigationSystemV1::K2_GetRandomLocationInNavigableRadius(UObject* WorldContextObject, const FVector& Origin, FVector& RandomLocation, float Radius, ANavigationData* NavData, TSubclassOf<UNavigationQueryFilter> FilterClass)
 {
 	FNavLocation RandomPoint(Origin);
 	bool bResult = false;
+	RandomLocation = Origin;
 
 	UWorld* World = GEngine->GetWorldFromContextObject(WorldContextObject, EGetWorldErrorMode::LogAndReturnNull);
 	UNavigationSystemV1* NavSys = FNavigationSystem::GetCurrent<UNavigationSystemV1>(World);
@@ -3886,8 +3887,11 @@ bool UNavigationSystemV1::K2_GetRandomPointInNavigableRadius(UObject* WorldConte
 		ANavigationData* UseNavData = NavData ? NavData : NavSys->GetDefaultNavDataInstance(FNavigationSystem::DontCreate);
 		if (UseNavData)
 		{
-			bResult = NavSys->GetRandomPointInNavigableRadius(Origin, Radius, RandomPoint, UseNavData, UNavigationQueryFilter::GetQueryFilter(*UseNavData, WorldContextObject, FilterClass));
-			RandomLocation = RandomPoint.Location;
+			if (NavSys->GetRandomPointInNavigableRadius(Origin, Radius, RandomPoint, UseNavData, UNavigationQueryFilter::GetQueryFilter(*UseNavData, WorldContextObject, FilterClass)))
+			{
+				bResult = true;
+				RandomLocation = RandomPoint.Location;
+			}
 		}
 	}
 
@@ -4321,6 +4325,11 @@ FVector UNavigationSystemV1::GetRandomPointInNavigableRadius(UObject* WorldConte
 	return RandomPoint.Location;
 }
 
+bool UNavigationSystemV1::K2_GetRandomPointInNavigableRadius(UObject* WorldContextObject, const FVector& Origin, FVector& RandomLocation, float Radius, ANavigationData* NavData, TSubclassOf<UNavigationQueryFilter> FilterClass)
+{
+	return K2_GetRandomLocationInNavigableRadius(WorldContextObject, Origin, RandomLocation, Radius, NavData, FilterClass);
+}
+
 void UNavigationSystemV1::SimpleMoveToActor(AController* Controller, const AActor* Goal)
 {
 	UE_LOG(LogNavigation, Error, TEXT("SimpleMoveToActor is deprecated. Use UAIBlueprintHelperLibrary::SimpleMoveToActor instead"));
diff --git a/Engine/Source/Runtime/NavigationSystem/Public/NavigationSystem.h b/Engine/Source/Runtime/NavigationSystem/Public/NavigationSystem.h
--- a/Engine/Source/Runtime/NavigationSystem/Public/NavigationSystem.h
+++ b/Engine/Source/Runtime/NavigationSystem/Public/NavigationSystem.h
@@ -240,8 +240,8 @@ public:
 
 	/** Generates a random location in navigable space within given radius of Origin.
 	 *	@return Return Value represents if the call was successful */
-	UFUNCTION(BlueprintPure, Category = "AI|Navigation", meta = (WorldContext = "WorldContextObject", DisplayName = "GetRandomPointInNavigableRadius", ScriptName = "GetRandomPointInNavigableRadius"))
-	static bool K2_GetRandomPointInNavigableRadius(UObject* WorldContextObject, const FVector& Origin, FVector& RandomLocation, float Radius, ANavigationData* NavData = NULL, TSubclassOf<UNavigationQueryFilter> FilterClass = NULL);
+	UFUNCTION(BlueprintCallable, Category = "AI|Navigation", meta = (WorldContext = "WorldContextObject", DisplayName = "GetRandomLocationInNavigableRadius", ScriptName = "GetRandomLocationInNavigableRadius"))
+	static bool K2_GetRandomLocationInNavigableRadius(UObject* WorldContextObject, const FVector& Origin, FVector& RandomLocation, float Radius, ANavigationData* NavData = NULL, TSubclassOf<UNavigationQueryFilter> FilterClass = NULL);
 	
 	/** Potentially expensive. Use with caution. Consider using UPathFollowingComponent::GetRemainingPathCost instead */
 	UFUNCTION(BlueprintPure, Category="AI|Navigation", meta=(WorldContext="WorldContextObject" ) )
@@ -996,6 +996,9 @@ public:
 	static TSubclassOf<UNavAreaBase> GetDefaultWalkableArea() { return FNavigationSystem::GetDefaultWalkableArea(); }
 	UE_DEPRECATED(4.20, "UNavigationSystem::GetDefaultObstacleArea is deprecated. Use FNavigationSystem::GetDefaultObstacleArea instead")
 	static TSubclassOf<UNavAreaBase> GetDefaultObstacleArea() { return FNavigationSystem::GetDefaultObstacleArea(); }
+	UE_DEPRECATED(4.22, "This version is deprecated.  Please use GetRandomLocationInNavigableRadius instead")
+	UFUNCTION(BlueprintPure, Category = "AI|Navigation", meta = (WorldContext = "WorldContextObject", DisplayName = "GetRandomPointInNavigableRadius", ScriptName = "GetRandomPointInNavigableRadius"))
+	static bool K2_GetRandomPointInNavigableRadius(UObject* WorldContextObject, const FVector& Origin, FVector& RandomLocation, float Radius, ANavigationData* NavData = NULL, TSubclassOf<UNavigationQueryFilter> FilterClass = NULL);
 };
 
 //----------------------------------------------------------------------//
