Added special handling (including logging) when there are no walkable areas while generating recast navmesh tile #UE4

This is the case where there are some spans, but not enough to create a walkable surface.

diff --git a/Engine/Source/Runtime/NavigationSystem/Private/NavMesh/RecastNavMeshGenerator.cpp b/Engine/Source/Runtime/NavigationSystem/Private/NavMesh/RecastNavMeshGenerator.cpp
--- a/Engine/Source/Runtime/NavigationSystem/Private/NavMesh/RecastNavMeshGenerator.cpp
+++ b/Engine/Source/Runtime/NavigationSystem/Private/NavMesh/RecastNavMeshGenerator.cpp
@@ -2441,7 +2441,16 @@ bool FRecastTileGenerator::GenerateCompressedLayers(FNavMeshBuildContext& BuildC
 		}
 		if (!rcBuildCompactHeightfield(&BuildContext, TileConfig.walkableHeight, TileConfig.walkableClimb, *RasterContext.SolidHF, *RasterContext.CompactHF))
 		{
-			BuildContext.log(RC_LOG_ERROR, "GenerateCompressedLayers: Could not build compact data.");
+			const int SpanCount = rcGetHeightFieldSpanCount(&BuildContext, *RasterContext.SolidHF);
+			if (SpanCount > 0)
+			{
+				BuildContext.log(RC_LOG_ERROR, "GenerateCompressedLayers: Could not build compact data.");
+			}
+			// else there's just no spans to walk on (no spans at all or too small/sparse)
+			else
+			{
+				BuildContext.log(RC_LOG_WARNING, "GenerateCompressedLayers: no walkable spans - aborting");
+			}
 			return false;
 		}
 	}
diff --git a/Engine/Source/Runtime/Navmesh/Private/Recast/Recast.cpp b/Engine/Source/Runtime/Navmesh/Private/Recast/Recast.cpp
--- a/Engine/Source/Runtime/Navmesh/Private/Recast/Recast.cpp
+++ b/Engine/Source/Runtime/Navmesh/Private/Recast/Recast.cpp
@@ -406,11 +406,18 @@ bool rcBuildCompactHeightfield(rcContext* ctx, const int walkableHeight, const i
 {
 	rcAssert(ctx);
 	
+// @UE4 BEGIN: early-out when no walkable spans 
+	const int spanCount = rcGetHeightFieldSpanCount(ctx, hf);
+	if (spanCount == 0)
+	{
+		// no spans to speak of, bail out.
+		return false;
+	}
+// @UE4 END
 	ctx->startTimer(RC_TIMER_BUILD_COMPACTHEIGHTFIELD);
 	
 	const int w = hf.width;
 	const int h = hf.height;
-	const int spanCount = rcGetHeightFieldSpanCount(ctx, hf);
 
 	// Fill in header.
 	chf.width = w;
