Changed the way we store UNavigationSystemV1.CrowdManagerClass so that it works with game-specific crowd manager class implementations #UE4

diff --git a/Engine/Source/Runtime/NavigationSystem/Private/NavigationSystem.cpp b/Engine/Source/Runtime/NavigationSystem/Private/NavigationSystem.cpp
--- a/Engine/Source/Runtime/NavigationSystem/Private/NavigationSystem.cpp
+++ b/Engine/Source/Runtime/NavigationSystem/Private/NavigationSystem.cpp
@@ -906,9 +906,13 @@ void UNavigationSystemV1::RegisterNavigationDataInstances()
 
 void UNavigationSystemV1::CreateCrowdManager()
 {
-	if (CrowdManagerClass)
+	UClass* CrowdManagerClassInstance = CrowdManagerClass.Get();
+	if (CrowdManagerClassInstance)
 	{
-		SetCrowdManager(NewObject<UCrowdManagerBase>(this, CrowdManagerClass));
+		UCrowdManagerBase* ManagerInstance = NewObject<UCrowdManagerBase>(this, CrowdManagerClassInstance);
+		// creating an instance when we have a valid class should never fail
+		check(ManagerInstance);
+		SetCrowdManager(ManagerInstance);
 	}
 }
 
diff --git a/Engine/Source/Runtime/NavigationSystem/Public/NavigationSystem.h b/Engine/Source/Runtime/NavigationSystem/Public/NavigationSystem.h
--- a/Engine/Source/Runtime/NavigationSystem/Public/NavigationSystem.h
+++ b/Engine/Source/Runtime/NavigationSystem/Public/NavigationSystem.h
@@ -114,7 +114,7 @@ public:
 
 protected:
 	UPROPERTY(config, EditAnywhere, BlueprintReadOnly, Category = Navigation)
-	TSubclassOf<UCrowdManagerBase> CrowdManagerClass;
+	TSoftClassPtr<UCrowdManagerBase> CrowdManagerClass;
 
 	/** Should navigation system spawn default Navigation Data when there's none and there are navigation bounds present? */
 	UPROPERTY(config, EditAnywhere, Category=NavigationSystem)
