Fixed a navmes generation bug that resulted in missing parts of navmesh when new navmesh bounds had been added  #UE4

The bug was being exibited when there was already a navmesh instance on the map and another navmesh bounds volume had been added. The whole detour navmesh got dropped, but only the area covered by the new bounds had been rebuild. I've fixed it so that as part of nav bounds update a new detour navmesh instance gets created and areas covered by _all_ bounds get dirtied.

diff --git a/Engine/Source/Runtime/NavigationSystem/Private/NavMesh/RecastNavMeshGenerator.cpp b/Engine/Source/Runtime/NavigationSystem/Private/NavMesh/RecastNavMeshGenerator.cpp
--- a/Engine/Source/Runtime/NavigationSystem/Private/NavMesh/RecastNavMeshGenerator.cpp
+++ b/Engine/Source/Runtime/NavigationSystem/Private/NavMesh/RecastNavMeshGenerator.cpp
@@ -3637,8 +3637,22 @@ void FRecastNavMeshGenerator::OnNavigationBoundsChanged()
 		int32 MaxRequestedTiles = CaclulateMaxTilesCount(InclusionBounds, Config.tileSize * Config.cs, AvgLayersPerTile);
 		if (DetourMesh->getMaxTiles() != MaxRequestedTiles)
 		{
-			// Destroy current NavMesh, it will be allocated with a new size on next build request
+			// Destroy current NavMesh
 			DestNavMesh->GetRecastNavMeshImpl()->SetRecastMesh(nullptr);
+
+			// if there are any valid bounds recreate detour navmesh instance
+			// and mark all bounds as dirty
+			if (InclusionBounds.Num() > 0)
+			{
+				TArray<FNavigationDirtyArea> AsDirtyAreas;
+				AsDirtyAreas.Reserve(InclusionBounds.Num());
+				for (const FBox& BBox : InclusionBounds)
+				{
+					AsDirtyAreas.Add(FNavigationDirtyArea(BBox, ENavigationDirtyFlag::NavigationBounds));
+				}
+				
+				RebuildDirtyAreas(AsDirtyAreas);
+			}
 		}
 	}
 }
@@ -4024,7 +4038,7 @@ FBox FRecastNavMeshGenerator::GrowBoundingBox(const FBox& BBox, bool bIncludeAge
 	return FBox(BBox.Min - BBoxGrowth - BBoxGrowOffsetMin, BBox.Max + BBoxGrowth);
 }
 
-static bool IntercestBounds(const FBox& TestBox, const TNavStatArray<FBox>& Bounds)
+static bool IntersectBounds(const FBox& TestBox, const TNavStatArray<FBox>& Bounds)
 {
 	for (const FBox& Box : Bounds)
 	{
@@ -4090,7 +4104,7 @@ void FRecastNavMeshGenerator::MarkDirtyTiles(const TArray<FNavigationDirtyArea>&
 			AdjustedAreaBounds = GrowBoundingBox(CutDownArea, DirtyArea.HasFlag(ENavigationDirtyFlag::UseAgentHeight));
 
 			// @TODO this and the following test share some work in common
-			if (IntercestBounds(AdjustedAreaBounds, InclusionBounds) == false)
+			if (IntersectBounds(AdjustedAreaBounds, InclusionBounds) == false)
 			{
 				continue;
 			}
@@ -4121,7 +4135,7 @@ void FRecastNavMeshGenerator::MarkDirtyTiles(const TArray<FNavigationDirtyArea>&
 					const FBox TileBox = CalculateTileBounds(TileX, TileY, RcNavMeshOrigin, TotalNavBounds, TileSizeInWorldUnits);
 
 					// do per tile check since we can have lots of tiles inbetween navigable bounds volumes
-					if (IntercestBounds(TileBox, InclusionBounds) == false)
+					if (IntersectBounds(TileBox, InclusionBounds) == false)
 					{
 						// Skip this tile
 						continue;
