CIS fix #UE4

diff --git a/Engine/Source/Runtime/NavigationSystem/Private/NavigationSystem.cpp b/Engine/Source/Runtime/NavigationSystem/Private/NavigationSystem.cpp
--- a/Engine/Source/Runtime/NavigationSystem/Private/NavigationSystem.cpp
+++ b/Engine/Source/Runtime/NavigationSystem/Private/NavigationSystem.cpp
@@ -416,7 +416,7 @@ UNavigationSystemV1::UNavigationSystemV1(const FObjectInitializer& ObjectInitial
 		SetDefaultObstacleArea(UNavArea_Obstacle::StaticClass());
 		
 		const FTransform RecastToUnrealTransfrom(Recast2UnrealMatrix());
-		SetCoordTransformFrom(ENavigationCoordSystem::Navigation, RecastToUnrealTransfrom);
+		SetCoordTransform(ENavigationCoordSystem::Navigation, ENavigationCoordSystem::Unreal, RecastToUnrealTransfrom);
 	}
 
 #if WITH_EDITOR
