Unhidden Activation category for NavigationInvokerComponent to allow users to controll initial bAutoActivate value #UE4

diff --git a/Engine/Source/Runtime/NavigationSystem/Public/NavigationInvokerComponent.h b/Engine/Source/Runtime/NavigationSystem/Public/NavigationInvokerComponent.h
--- a/Engine/Source/Runtime/NavigationSystem/Public/NavigationInvokerComponent.h
+++ b/Engine/Source/Runtime/NavigationSystem/Public/NavigationInvokerComponent.h
@@ -9,7 +9,7 @@
 
 class UNavigationSystemV1;
 
-UCLASS(ClassGroup = (Navigation), meta = (BlueprintSpawnableComponent), hidecategories = (Activation))
+UCLASS(ClassGroup = (Navigation), meta = (BlueprintSpawnableComponent))
 class NAVIGATIONSYSTEM_API UNavigationInvokerComponent : public UActorComponent
 {
 	GENERATED_BODY()
