Extended NavigationInvokerComponent's setter/getter functions set #UE4

diff --git a/Engine/Source/Runtime/NavigationSystem/Private/NavigationInvokerComponent.cpp b/Engine/Source/Runtime/NavigationSystem/Private/NavigationInvokerComponent.cpp
--- a/Engine/Source/Runtime/NavigationSystem/Private/NavigationInvokerComponent.cpp
+++ b/Engine/Source/Runtime/NavigationSystem/Private/NavigationInvokerComponent.cpp
@@ -44,3 +44,9 @@ void UNavigationInvokerComponent::RegisterWithNavigationSystem(UNavigationSystem
 		}
 	}
 }
+
+void UNavigationInvokerComponent::SetGenerationRadii(const float GenerationRadius, const float RemovalRadius)
+{
+	TileGenerationRadius = GenerationRadius;
+	TileRemovalRadius = RemovalRadius;
+}
diff --git a/Engine/Source/Runtime/NavigationSystem/Public/NavigationInvokerComponent.h b/Engine/Source/Runtime/NavigationSystem/Public/NavigationInvokerComponent.h
--- a/Engine/Source/Runtime/NavigationSystem/Public/NavigationInvokerComponent.h
+++ b/Engine/Source/Runtime/NavigationSystem/Public/NavigationInvokerComponent.h
@@ -27,7 +27,13 @@ public:
 
 	void RegisterWithNavigationSystem(UNavigationSystemV1& NavSys);
 
-protected:
+	/** Sets generation/removal ranges. Doesn't force navigation system's update.
+	 *	Will get picked up the next time NavigationSystem::UpdateInvokers gets called */
+	void SetGenerationRadii(const float GenerationRadius, const float RemovalRadius);
+
+	float GetGenerationRadius() const { return TileGenerationRadius; }
+	float GetRemovalRadius() const { return TileRemovalRadius; }
+
 	virtual void Activate(bool bReset = false) override;
 	virtual void Deactivate() override;
 };
