Made overriding navigation-invoker-related functionality possible #UE4

Made editor-time rebuilding with nav invokers more reliable.
Minor optimization in FRecastNavMeshGenerator::ProcessTileTasksAsync to avoid redundant condition checking and memory reallocation.

diff --git a/Engine/Source/Runtime/NavigationSystem/Private/NavMesh/PImplRecastNavMesh.cpp b/Engine/Source/Runtime/NavigationSystem/Private/NavMesh/PImplRecastNavMesh.cpp
--- a/Engine/Source/Runtime/NavigationSystem/Private/NavMesh/PImplRecastNavMesh.cpp
+++ b/Engine/Source/Runtime/NavigationSystem/Private/NavMesh/PImplRecastNavMesh.cpp
@@ -2171,7 +2171,9 @@ void FPImplRecastNavMesh::GetDebugGeometry(FRecastDebugGeometry& OutGeometry, in
 
 	const FRecastNavMeshGenerator* Generator = static_cast<const FRecastNavMeshGenerator*>(NavMeshOwner->GetGenerator());
 
-	if (Generator && Generator->IsBuildingRestrictedToActiveTiles())
+	if (Generator && Generator->IsBuildingRestrictedToActiveTiles()
+		// if not active tiles try drawing all tiles
+		&& NavMeshOwner->GetActiveTiles().Num() > 0)
 	{
 		const TArray<FIntPoint>& ActiveTiles = NavMeshOwner->GetActiveTiles();
 		for (const FIntPoint& TileLocation : ActiveTiles)
diff --git a/Engine/Source/Runtime/NavigationSystem/Private/NavMesh/RecastNavMeshGenerator.cpp b/Engine/Source/Runtime/NavigationSystem/Private/NavMesh/RecastNavMeshGenerator.cpp
--- a/Engine/Source/Runtime/NavigationSystem/Private/NavMesh/RecastNavMeshGenerator.cpp
+++ b/Engine/Source/Runtime/NavigationSystem/Private/NavMesh/RecastNavMeshGenerator.cpp
@@ -59,28 +59,6 @@ FORCEINLINE bool DoesBoxContainBox(const FBox& BigBox, const FBox& SmallBox)
 	return DoesBoxContainOrOverlapVector(BigBox, SmallBox.Min) && DoesBoxContainOrOverlapVector(BigBox, SmallBox.Max);
 }
 
-struct FRcTileBox
-{
-	int32 XMin, XMax, YMin, YMax;
-
-	FRcTileBox(const FBox& UnrealBounds, const FVector& RcNavMeshOrigin, const float TileSizeInWorldUnits)
-	{
-		check(TileSizeInWorldUnits > 0);
-
-		const FBox RcAreaBounds = Unreal2RecastBox(UnrealBounds);
-		XMin = FMath::FloorToInt((RcAreaBounds.Min.X - RcNavMeshOrigin.X) / TileSizeInWorldUnits);
-		XMax = FMath::FloorToInt((RcAreaBounds.Max.X - RcNavMeshOrigin.X) / TileSizeInWorldUnits);
-		YMin = FMath::FloorToInt((RcAreaBounds.Min.Z - RcNavMeshOrigin.Z) / TileSizeInWorldUnits);
-		YMax = FMath::FloorToInt((RcAreaBounds.Max.Z - RcNavMeshOrigin.Z) / TileSizeInWorldUnits);
-	}
-
-	FORCEINLINE bool Contains(const FIntPoint& Point) const 
-	{
-		return Point.X >= XMin && Point.X <= XMax
-			&& Point.Y >= YMin && Point.Y <= YMax;
-	}
-};
-
 int32 GetTilesCountHelper(const dtNavMesh* DetourMesh)
 {
 	int32 NumTiles = 0;
@@ -3920,7 +3898,11 @@ TArray<uint32> FRecastNavMeshGenerator::AddGeneratedTiles(FRecastTileGenerator&
 	dtNavMesh* DetourMesh = DestNavMesh->GetRecastNavMeshImpl()->GetRecastMesh();
 	const int32 FirstDirtyTileIndex = TileGenerator.GetDirtyLayersMask().Find(true);
 
-	if (DetourMesh != nullptr && IsInActiveSet(FIntPoint(TileX, TileY))
+	if (DetourMesh != nullptr 
+		// no longer testing this here, we can live with a stray unwanted tile here 
+		// and there. It will be removed the next time around the invokers get
+		// updated 
+		// && IsInActiveSet(FIntPoint(TileX, TileY))
 		&& FirstDirtyTileIndex != INDEX_NONE)
 	{
 		TArray<FNavMeshTileData> TileLayers = TileGenerator.GetNavigationData();
@@ -4075,7 +4057,7 @@ static bool IntersectBounds(const FBox& TestBox, const TNavStatArray<FBox>& Boun
 
 namespace 
 {
-	FBox CalculateBoxIntercetion(const FBox& BoxA, const FBox& BoxB)
+	FBox CalculateBoxIntersection(const FBox& BoxA, const FBox& BoxB)
 	{
 		// assumes boxes overlap
 		ensure(BoxA.Intersect(BoxB));
@@ -4165,7 +4147,7 @@ void FRecastNavMeshGenerator::MarkDirtyTiles(const TArray<FNavigationDirtyArea>&
 				continue;
 			}
 
-			const FBox CutDownArea = CalculateBoxIntercetion(GetTotalBounds(), DirtyArea.Bounds);
+			const FBox CutDownArea = CalculateBoxIntersection(GetTotalBounds(), DirtyArea.Bounds);
 			AdjustedAreaBounds = GrowBoundingBox(CutDownArea, DirtyArea.HasFlag(ENavigationDirtyFlag::UseAgentHeight));
 
 			// @TODO this and the following test share some work in common
@@ -4177,7 +4159,7 @@ void FRecastNavMeshGenerator::MarkDirtyTiles(const TArray<FNavigationDirtyArea>&
 			// check if any of inclusion volumes encapsulates this box
 			// using CutDownArea not AdjustedAreaBounds since if the area is on the border of navigable space
 			// then FindInclusionBoundEncapsulatingBox can produce false negative
-			bDoTileInclusionTest = FindInclusionBoundEncapsulatingBox(CutDownArea) == INDEX_NONE;
+			bDoTileInclusionTest = (FindInclusionBoundEncapsulatingBox(CutDownArea) == INDEX_NONE);
 		}
 		
 		const FRcTileBox TileBox(AdjustedAreaBounds, RcNavMeshOrigin, TileSizeInWorldUnits);
@@ -4191,7 +4173,7 @@ void FRecastNavMeshGenerator::MarkDirtyTiles(const TArray<FNavigationDirtyArea>&
 					continue;
 				}
 
-				if (DirtyArea.HasFlag(ENavigationDirtyFlag::NavigationBounds) == false && bDoTileInclusionTest == true)
+				if (bDoTileInclusionTest == true && DirtyArea.HasFlag(ENavigationDirtyFlag::NavigationBounds) == false)
 				{
 					const FBox TileBounds = CalculateTileBounds(TileX, TileY, RcNavMeshOrigin, TotalNavBounds, TileSizeInWorldUnits);
 
@@ -4375,16 +4357,16 @@ TArray<uint32> FRecastNavMeshGenerator::ProcessTileTasksAsync(const int32 NumTas
 			}
 
 			// Remove submitted element from pending list
-			PendingDirtyTiles.RemoveAt(ElementIdx);
-
-			// Release memory, list could be quite big after map load
-			if (PendingDirtyTiles.Num() == 0)
-			{
-				PendingDirtyTiles.Empty(32);
-			}
+			PendingDirtyTiles.RemoveAt(ElementIdx, 1, /*bAllowShrinking=*/false);
 			NumProcessedTasks++;
 		}
 	}
+
+	// Release memory, list could be quite big after map load
+	if (NumProcessedTasks > 0 && PendingDirtyTiles.Num() == 0)
+	{
+		PendingDirtyTiles.Empty(32);
+	}
 	
 	// Collect completed tasks and apply generated data to navmesh
 	for (int32 Idx = RunningDirtyTiles.Num() - 1; Idx >=0; --Idx)
diff --git a/Engine/Source/Runtime/NavigationSystem/Private/NavigationSystem.cpp b/Engine/Source/Runtime/NavigationSystem/Private/NavigationSystem.cpp
--- a/Engine/Source/Runtime/NavigationSystem/Private/NavigationSystem.cpp
+++ b/Engine/Source/Runtime/NavigationSystem/Private/NavigationSystem.cpp
@@ -981,46 +981,49 @@ void UNavigationSystemV1::Tick(float DeltaSeconds)
 		}
 		INC_FLOAT_STAT_BY(STAT_Navigation_CumulativeBuildTime,(float)ThisTime*1000);
 	}
-	
-	if (bGenerateNavigationOnlyAroundNavigationInvokers)
-	{
-		UpdateInvokers();
-	}
+		
+	DirtyAreasUpdateTime += DeltaSeconds;
 
+	if (IsNavigationBuildingLocked() == false)
 	{
-		SCOPE_CYCLE_COUNTER(STAT_Navigation_TickMarkDirty);
-
-		DirtyAreasUpdateTime += DeltaSeconds;
-		const float DirtyAreasUpdateDeltaTime = 1.0f / DirtyAreasUpdateFreq;
-		const bool bCanRebuildNow = (DirtyAreasUpdateTime >= DirtyAreasUpdateDeltaTime) || !bIsGame;
-		const bool bIsLocked = IsNavigationBuildingLocked();
+		if (bGenerateNavigationOnlyAroundNavigationInvokers)
+		{
+			UpdateInvokers();
+		}
 
-		if (DirtyAreas.Num() > 0 && bCanRebuildNow && !bIsLocked)
 		{
-			for (int32 NavDataIndex = 0; NavDataIndex < NavDataSet.Num(); ++NavDataIndex)
+			SCOPE_CYCLE_COUNTER(STAT_Navigation_TickMarkDirty);
+
+			const float DirtyAreasUpdateDeltaTime = 1.0f / DirtyAreasUpdateFreq;
+			const bool bCanRebuildNow = (DirtyAreasUpdateTime >= DirtyAreasUpdateDeltaTime) || !bIsGame;
+
+			if (DirtyAreas.Num() > 0 && bCanRebuildNow)
 			{
-				ANavigationData* NavData = NavDataSet[NavDataIndex];
-				if (NavData)
+				for (int32 NavDataIndex = 0; NavDataIndex < NavDataSet.Num(); ++NavDataIndex)
 				{
-					NavData->RebuildDirtyAreas(DirtyAreas);
+					ANavigationData* NavData = NavDataSet[NavDataIndex];
+					if (NavData)
+					{
+						NavData->RebuildDirtyAreas(DirtyAreas);
+					}
 				}
-			}
 
-			DirtyAreasUpdateTime = 0;
-			DirtyAreas.Reset();
+				DirtyAreasUpdateTime = 0;
+				DirtyAreas.Reset();
+			}
 		}
-	}
-
-	// Tick navigation mesh async builders
-	if (!bAsyncBuildPaused)
-	{
-		SCOPE_CYCLE_COUNTER(STAT_Navigation_TickAsyncBuild);
 
-		for (ANavigationData* NavData : NavDataSet)
+		// Tick navigation mesh async builders
+		if (bAsyncBuildPaused == false)
 		{
-			if (NavData)
+			SCOPE_CYCLE_COUNTER(STAT_Navigation_TickAsyncBuild);
+
+			for (ANavigationData* NavData : NavDataSet)
 			{
-				NavData->TickAsyncBuild(DeltaSeconds);
+				if (NavData)
+				{
+					NavData->TickAsyncBuild(DeltaSeconds);
+				}
 			}
 		}
 	}
@@ -3300,6 +3303,13 @@ void UNavigationSystemV1::Build()
 
 	// make sure freshly created navigation instances are registered before we try to build them
 	ProcessRegistrationCandidates();
+	
+	// update invokers in case we're not updating navmesh automatically, in which case
+	// navigation generators wouldn't have up-to-date info.
+	if (bGenerateNavigationOnlyAroundNavigationInvokers)
+	{
+		UpdateInvokers();
+	}
 
 	// and now iterate through all registered and just start building them
 	RebuildAll();
@@ -4211,7 +4221,7 @@ void UNavigationSystemV1::UpdateInvokers()
 	const float CurrentTime = World->GetTimeSeconds();
 	if (CurrentTime >= NextInvokersUpdateTime)
 	{
-		TArray<FNavigationInvokerRaw> InvokerLocations;
+		InvokerLocations.Reset();
 
 		if (Invokers.Num() > 0)
 		{
diff --git a/Engine/Source/Runtime/NavigationSystem/Public/NavMesh/RecastNavMesh.h b/Engine/Source/Runtime/NavigationSystem/Public/NavMesh/RecastNavMesh.h
--- a/Engine/Source/Runtime/NavigationSystem/Public/NavMesh/RecastNavMesh.h
+++ b/Engine/Source/Runtime/NavigationSystem/Public/NavMesh/RecastNavMesh.h
@@ -888,8 +888,8 @@ public:
 	bool ShouldGatherDataOnGameThread() const { return bDoFullyAsyncNavDataGathering == false; }
 	int32 GetTileNumberHardLimit() const { return TileNumberHardLimit; }
 
-	void UpdateActiveTiles(const TArray<FNavigationInvokerRaw>& InvokerLocations);
-	void RemoveTiles(const TArray<FIntPoint>& Tiles);
+	virtual void UpdateActiveTiles(const TArray<FNavigationInvokerRaw>& InvokerLocations);
+	virtual void RemoveTiles(const TArray<FIntPoint>& Tiles);
 	void RebuildTile(const TArray<FIntPoint>& Tiles);
 
 protected:
diff --git a/Engine/Source/Runtime/NavigationSystem/Public/NavMesh/RecastNavMeshGenerator.h b/Engine/Source/Runtime/NavigationSystem/Public/NavMesh/RecastNavMeshGenerator.h
--- a/Engine/Source/Runtime/NavigationSystem/Public/NavMesh/RecastNavMeshGenerator.h
+++ b/Engine/Source/Runtime/NavigationSystem/Public/NavMesh/RecastNavMeshGenerator.h
@@ -11,6 +11,7 @@
 #include "Async/AsyncWork.h"
 #include "UObject/GCObject.h"
 #include "AI/NavDataGenerator.h"
+#include "NavMesh/RecastHelpers.h"
 
 #if WITH_RECAST
 
@@ -150,6 +151,28 @@ struct FRecastAreaNavModifierElement
 	TArray<FTransform>	PerInstanceTransform;
 };
 
+struct FRcTileBox
+{
+	int32 XMin, XMax, YMin, YMax;
+
+	FRcTileBox(const FBox& UnrealBounds, const FVector& RcNavMeshOrigin, const float TileSizeInWorldUnits)
+	{
+		check(TileSizeInWorldUnits > 0);
+
+		const FBox RcAreaBounds = Unreal2RecastBox(UnrealBounds);
+		XMin = FMath::FloorToInt((RcAreaBounds.Min.X - RcNavMeshOrigin.X) / TileSizeInWorldUnits);
+		XMax = FMath::FloorToInt((RcAreaBounds.Max.X - RcNavMeshOrigin.X) / TileSizeInWorldUnits);
+		YMin = FMath::FloorToInt((RcAreaBounds.Min.Z - RcNavMeshOrigin.Z) / TileSizeInWorldUnits);
+		YMax = FMath::FloorToInt((RcAreaBounds.Max.Z - RcNavMeshOrigin.Z) / TileSizeInWorldUnits);
+	}
+
+	FORCEINLINE bool Contains(const FIntPoint& Point) const
+	{
+		return Point.X >= XMin && Point.X <= XMax
+			&& Point.Y >= YMin && Point.Y <= YMax;
+	}
+};
+
 enum class ETimeSliceWorkResult : uint8
 {
 	Failed,
@@ -485,7 +508,7 @@ protected:
 	void UpdateNavigationBounds();
 		
 	// Sorts pending build tiles by proximity to player, so tiles closer to player will get generated first
-	void SortPendingBuildTiles();
+	virtual void SortPendingBuildTiles();
 
 	/** Instantiates dtNavMesh and configures it for tiles generation. Returns false if failed */
 	bool ConstructTiledNavMesh();
@@ -494,7 +517,7 @@ protected:
 	void CalcNavMeshProperties(int32& MaxTiles, int32& MaxPolys);
 	
 	/** Marks grid tiles affected by specified areas as dirty */
-	void MarkDirtyTiles(const TArray<FNavigationDirtyArea>& DirtyAreas);
+	virtual void MarkDirtyTiles(const TArray<FNavigationDirtyArea>& DirtyAreas);
 
 	void RemoveLayers(const FIntPoint& Tile, TArray<uint32>& UpdatedTiles);
 	
@@ -530,7 +553,7 @@ public:
 
 protected:
 	bool IsInActiveSet(const FIntPoint& Tile) const;
-	void RestrictBuildingToActiveTiles(bool InRestrictBuildingToActiveTiles);
+	virtual void RestrictBuildingToActiveTiles(bool InRestrictBuildingToActiveTiles);
 	
 	/** Blocks until build for specified list of tiles is complete and discard results */
 	void DiscardCurrentBuildingTasks();
diff --git a/Engine/Source/Runtime/NavigationSystem/Public/NavigationSystem.h b/Engine/Source/Runtime/NavigationSystem/Public/NavigationSystem.h
--- a/Engine/Source/Runtime/NavigationSystem/Public/NavigationSystem.h
+++ b/Engine/Source/Runtime/NavigationSystem/Public/NavigationSystem.h
@@ -217,6 +217,8 @@ protected:
 
 private:
 	TMap<AActor*, FNavigationInvoker> Invokers;
+	/** Contains pre-digested and cached invokers' info. Generated by UpdateInvokers */
+	TArray<FNavigationInvokerRaw> InvokerLocations;
 
 	float NextInvokersUpdateTime;
 	void UpdateInvokers();
@@ -502,12 +504,14 @@ public:
 	//----------------------------------------------------------------------//
 	// Active tiles
 	//----------------------------------------------------------------------//
-	void RegisterInvoker(AActor& Invoker, float TileGenerationRadius, float TileRemovalRadius);
-	void UnregisterInvoker(AActor& Invoker);
+	virtual void RegisterInvoker(AActor& Invoker, float TileGenerationRadius, float TileRemovalRadius);
+	virtual void UnregisterInvoker(AActor& Invoker);
 
 	static void RegisterNavigationInvoker(AActor& Invoker, float TileGenerationRadius, float TileRemovalRadius);
 	static void UnregisterNavigationInvoker(AActor& Invoker);
 
+	const TArray<FNavigationInvokerRaw>& GetInvokerLocations() const { return InvokerLocations; }
+
 	//----------------------------------------------------------------------//
 	// Bookkeeping 
 	//----------------------------------------------------------------------//
@@ -578,7 +582,7 @@ public:
 	/** updates bounds of all components implementing INavRelevantInterface */
 	static void UpdateNavOctreeBounds(AActor* Actor);
 
-	void AddDirtyArea(const FBox& NewArea, int32 Flags);
+	virtual void AddDirtyArea(const FBox& NewArea, int32 Flags);
 	void AddDirtyAreas(const TArray<FBox>& NewAreas, int32 Flags);
 	bool HasDirtyAreasQueued() const;
 
