PR #5813: Adding concept of supported nav agents to Smart Links (Contributed by [REDACTED_PERSON])
Added virtual method INavLinkCustomInterface::GetSupportedAgents to allow derived classes to provide their supported agents when a NavigationLink is created by INavLinkCustomInterface::GetModifier.
Added property SupportedAgents to NavLinkCustomComponent and override GetSupportedAgents to forward the filter to the NavigationLink.

diff --git a/Engine/Source/Runtime/NavigationSystem/Private/NavLinkCustomComponent.cpp b/Engine/Source/Runtime/NavigationSystem/Private/NavLinkCustomComponent.cpp
--- a/Engine/Source/Runtime/NavigationSystem/Private/NavLinkCustomComponent.cpp
+++ b/Engine/Source/Runtime/NavigationSystem/Private/NavLinkCustomComponent.cpp
@@ -58,6 +58,11 @@ void UNavLinkCustomComponent::GetLinkData(FVector& LeftPt, FVector& RightPt, ENa
 	Direction = LinkDirection;
 }
 
+void UNavLinkCustomComponent::GetSupportedAgents(FNavAgentSelector& OutSupportedAgents) const
+{
+	OutSupportedAgents = SupportedAgents;
+}
+
 TSubclassOf<UNavArea> UNavLinkCustomComponent::GetLinkAreaClass() const
 {
 	return bLinkEnabled ? EnabledAreaClass : DisabledAreaClass;
diff --git a/Engine/Source/Runtime/NavigationSystem/Private/NavigationSystemTypes.cpp b/Engine/Source/Runtime/NavigationSystem/Private/NavigationSystemTypes.cpp
--- a/Engine/Source/Runtime/NavigationSystem/Private/NavigationSystemTypes.cpp
+++ b/Engine/Source/Runtime/NavigationSystem/Private/NavigationSystemTypes.cpp
@@ -312,6 +312,7 @@ FNavigationLink INavLinkCustomInterface::GetModifier(const INavLinkCustomInterfa
 
 	ENavLinkDirection::Type LinkDirection = ENavLinkDirection::BothWays;
 	CustomNavLink->GetLinkData(LinkMod.Left, LinkMod.Right, LinkDirection);
+	CustomNavLink->GetSupportedAgents(LinkMod.SupportedAgents);
 	LinkMod.Direction = LinkDirection;
 
 	return LinkMod;
diff --git a/Engine/Source/Runtime/NavigationSystem/Public/NavLinkCustomComponent.h b/Engine/Source/Runtime/NavigationSystem/Public/NavLinkCustomComponent.h
--- a/Engine/Source/Runtime/NavigationSystem/Public/NavLinkCustomComponent.h
+++ b/Engine/Source/Runtime/NavigationSystem/Public/NavLinkCustomComponent.h
@@ -34,6 +34,7 @@ class NAVIGATIONSYSTEM_API UNavLinkCustomComponent : public UNavRelevantComponen
 
 	// BEGIN INavLinkCustomInterface
 	virtual void GetLinkData(FVector& LeftPt, FVector& RightPt, ENavLinkDirection::Type& Direction) const override;
+	virtual void GetSupportedAgents(FNavAgentSelector& OutSupportedAgents) const override;
 	virtual TSubclassOf<UNavArea> GetLinkAreaClass() const override;
 	virtual uint32 GetLinkId() const override;
 	virtual void UpdateLinkId(uint32 NewUniqueId) override;
@@ -71,6 +72,9 @@ class NAVIGATIONSYSTEM_API UNavLinkCustomComponent : public UNavRelevantComponen
 	void SetDisabledArea(TSubclassOf<UNavArea> AreaClass);
 	TSubclassOf<UNavArea> GetDisabledArea() const { return DisabledAreaClass; }
 
+	void SetSupportedAgents(const FNavAgentSelector& InSupportedAgents) { SupportedAgents = InSupportedAgents; }
+	FNavAgentSelector GetSupportedAgents() const { return SupportedAgents; }
+
 	/** add box obstacle during generation of navigation data
 	  * this can be used to create empty area under doors */
 	void AddNavigationObstacle(TSubclassOf<UNavArea> AreaClass, const FVector& BoxExtent, const FVector& BoxOffset = FVector::ZeroVector);
@@ -142,6 +146,10 @@ protected:
 	UPROPERTY(EditAnywhere, Category=SmartLink)
 	TSubclassOf<UNavArea> DisabledAreaClass;
 
+	/** restrict area only to specified agents */
+	UPROPERTY(EditAnywhere, Category=SmartLink)
+	FNavAgentSelector SupportedAgents;
+
 	/** start point, relative to owner */
 	UPROPERTY(EditAnywhere, Category=SmartLink)
 	FVector LinkRelativeStart;
diff --git a/Engine/Source/Runtime/NavigationSystem/Public/NavLinkCustomInterface.h b/Engine/Source/Runtime/NavigationSystem/Public/NavLinkCustomInterface.h
--- a/Engine/Source/Runtime/NavigationSystem/Public/NavLinkCustomInterface.h
+++ b/Engine/Source/Runtime/NavigationSystem/Public/NavLinkCustomInterface.h
@@ -37,6 +37,9 @@ class NAVIGATIONSYSTEM_API INavLinkCustomInterface
 	/** Get basic link data: two points (relative to owner) and direction */
 	virtual void GetLinkData(FVector& LeftPt, FVector& RightPt, ENavLinkDirection::Type& Direction) const {};
 
+	/** Get agents supported by this link */
+	virtual void GetSupportedAgents(FNavAgentSelector& OutSupportedAgents) const {};
+
 	/** Get basic link data: area class (null = default walkable) */
 	virtual TSubclassOf<UNavArea> GetLinkAreaClass() const { return nullptr; }
 
