Made NavModifierComponent work properly even if its owner's root component was marked as nav-irrelevant #UE4

While at it had to mess around with RegastNavMeshGenerator's initialization - the old way in some cases resulted in overriding once configured properties or needless dtNavMesh instance recreation.
In the process I've removed the Generator->Init() call from ARecastNavMesh::ConditionalConstructGenerator since now we call Init as part of generator's construction.

diff --git a/Engine/Source/Runtime/NavigationSystem/Private/NavMesh/RecastNavMesh.cpp b/Engine/Source/Runtime/NavigationSystem/Private/NavMesh/RecastNavMesh.cpp
--- a/Engine/Source/Runtime/NavigationSystem/Private/NavMesh/RecastNavMesh.cpp
+++ b/Engine/Source/Runtime/NavigationSystem/Private/NavMesh/RecastNavMesh.cpp
@@ -2393,7 +2393,6 @@ void ARecastNavMesh::ConditionalConstructGenerator()
 		if (Generator)
 		{
 			NavDataGenerator = MakeShareable(Generator);
-			Generator->Init();
 		}
 
 		UNavigationSystemV1* NavSys = FNavigationSystem::GetCurrent<UNavigationSystemV1>(World);
diff --git a/Engine/Source/Runtime/NavigationSystem/Private/NavMesh/RecastNavMeshGenerator.cpp b/Engine/Source/Runtime/NavigationSystem/Private/NavMesh/RecastNavMeshGenerator.cpp
--- a/Engine/Source/Runtime/NavigationSystem/Private/NavMesh/RecastNavMeshGenerator.cpp
+++ b/Engine/Source/Runtime/NavigationSystem/Private/NavMesh/RecastNavMeshGenerator.cpp
@@ -3231,6 +3231,63 @@
 #endif
 
 	INC_DWORD_STAT_BY(STAT_NavigationMemory, sizeof(*this));
+
+	Init();
+
+	int32 MaxTiles = 0;
+	int32 MaxPolysPerTile = 0;
+
+	// recreate navmesh if no data was loaded, or when loaded data doesn't match current grid layout
+	bool bRecreateNavmesh = true;
+	if (DestNavMesh->HasValidNavmesh())
+	{
+		const bool bGameStaticNavMesh = IsGameStaticNavMesh(DestNavMesh);
+		const dtNavMeshParams* SavedNavParams = DestNavMesh->GetRecastNavMeshImpl()->DetourNavMesh->getParams();
+		if (SavedNavParams)
+		{
+			if (bGameStaticNavMesh)
+			{
+				bRecreateNavmesh = false;
+				MaxTiles = SavedNavParams->maxTiles;
+				MaxPolysPerTile = SavedNavParams->maxPolys;
+			}
+			else
+			{
+				const float TileDim = Config.tileSize * Config.cs;
+				if (SavedNavParams->tileHeight == TileDim && SavedNavParams->tileWidth == TileDim)
+				{
+					const FVector Orig = Recast2UnrealPoint(SavedNavParams->orig);
+					const FVector OrigError(FMath::Fmod(Orig.X, TileDim), FMath::Fmod(Orig.Y, TileDim), FMath::Fmod(Orig.Z, TileDim));
+					if (OrigError.IsNearlyZero())
+					{
+						bRecreateNavmesh = false;
+					}
+				}
+
+				// if new navmesh needs more tiles, force recreating
+				if (!bRecreateNavmesh)
+				{
+					CalcNavMeshProperties(MaxTiles, MaxPolysPerTile);
+					if (FMath::Log2(MaxTiles) != FMath::Log2(SavedNavParams->maxTiles))
+					{
+						bRecreateNavmesh = true;
+					}
+				}
+			}
+		};
+	}
+
+	if (bRecreateNavmesh)
+	{
+		// recreate navmesh from scratch if no data was loaded
+		ConstructTiledNavMesh();
+	}
+	else
+	{
+		// otherwise just update generator params
+		Config.MaxPolysPerTile = MaxPolysPerTile;
+		NumActiveTiles = GetTilesCountHelper(DestNavMesh->GetRecastNavMeshImpl()->DetourNavMesh);
+	}
 }
 
 FRecastNavMeshGenerator::~FRecastNavMeshGenerator()
@@ -3310,6 +3367,14 @@
 	
 	AdditionalCachedData = FRecastNavMeshCachedData::Construct(DestNavMesh);
 
+	if (Config.MaxPolysPerTile <= 0 && DestNavMesh->HasValidNavmesh())
+	{
+		const dtNavMeshParams* SavedNavParams = DestNavMesh->GetRecastNavMeshImpl()->DetourNavMesh->getParams();
+		if (SavedNavParams)
+		{
+			Config.MaxPolysPerTile = SavedNavParams->maxPolys;
+		}
+	}
 	UpdateNavigationBounds();
 
 	/** setup maximum number of active tile generator*/
@@ -3325,61 +3390,6 @@
 	}
 
 	bInitialized = true;
-
-	int32 MaxTiles = 0;
-	int32 MaxPolysPerTile = 0;
-
-	// recreate navmesh if no data was loaded, or when loaded data doesn't match current grid layout
-	bool bRecreateNavmesh = true;
-	if (DestNavMesh->HasValidNavmesh())
-	{
-		const bool bGameStaticNavMesh = IsGameStaticNavMesh(DestNavMesh);
-		const dtNavMeshParams* SavedNavParams = DestNavMesh->GetRecastNavMeshImpl()->DetourNavMesh->getParams();
-		if (SavedNavParams)
-		{
-			if (bGameStaticNavMesh)
-			{
-				bRecreateNavmesh = false;
-				MaxTiles = SavedNavParams->maxTiles;
-				MaxPolysPerTile = SavedNavParams->maxPolys;
-			}
-			else
-			{
-				const float TileDim = Config.tileSize * Config.cs;
-				if (SavedNavParams->tileHeight == TileDim && SavedNavParams->tileWidth == TileDim)
-				{
-					const FVector Orig = Recast2UnrealPoint(SavedNavParams->orig);
-					const FVector OrigError(FMath::Fmod(Orig.X, TileDim), FMath::Fmod(Orig.Y, TileDim), FMath::Fmod(Orig.Z, TileDim));
-					if (OrigError.IsNearlyZero())
-					{
-						bRecreateNavmesh = false;
-					}
-				}
-
-				// if new navmesh needs more tiles, force recreating
-				if (!bRecreateNavmesh)
-				{
-					CalcNavMeshProperties(MaxTiles, MaxPolysPerTile);
-					if (FMath::Log2(MaxTiles) != FMath::Log2(SavedNavParams->maxTiles))
-					{
-						bRecreateNavmesh = true;
-					}
-				}
-			}
-		};
-	}
-
-	if (bRecreateNavmesh)
-	{
-		// recreate navmesh from scratch if no data was loaded
-		ConstructTiledNavMesh();
-	}
-	else
-	{
-		// otherwise just update generator params
-		Config.MaxPolysPerTile = MaxPolysPerTile;
-		NumActiveTiles = GetTilesCountHelper(DestNavMesh->GetRecastNavMeshImpl()->DetourNavMesh);
-	}
 }
 
 void FRecastNavMeshGenerator::UpdateNavigationBounds()
diff --git a/Engine/Source/Runtime/NavigationSystem/Private/NavModifierComponent.cpp b/Engine/Source/Runtime/NavigationSystem/Private/NavModifierComponent.cpp
--- a/Engine/Source/Runtime/NavigationSystem/Private/NavModifierComponent.cpp
+++ b/Engine/Source/Runtime/NavigationSystem/Private/NavModifierComponent.cpp
@@ -16,9 +16,21 @@ UNavModifierComponent::UNavModifierComponent(const FObjectInitializer& ObjectIni
 
 void UNavModifierComponent::CalcAndCacheBounds() const
 {
-	const AActor* MyOwner = GetOwner();
+	AActor* MyOwner = GetOwner();
 	if (MyOwner)
 	{
+		CachedTransform = MyOwner->GetActorTransform();
+		
+		if (TransformUpdateHandle.IsValid() == false && MyOwner->GetRootComponent())
+		{
+			// binding to get notifies when the root component moves. We need
+			// this only when the rootcomp is nav-irrelevant (since the default 
+			// mechanisms won't kick in) but we're binding without checking it since
+			// this property can change without re-running CalcAndCacheBounds.
+			// We're filtering for nav relevancy in OnTransformUpdated.
+			TransformUpdateHandle = MyOwner->GetRootComponent()->TransformUpdated.AddUObject(this, &UNavModifierComponent::OnTransformUpdated);
+		}
+
 		Bounds = FBox(ForceInit);
 		ComponentBounds.Reset();
 		for (UActorComponent* Component : MyOwner->GetComponents())
@@ -117,3 +129,16 @@ void UNavModifierComponent::SetAreaClass(TSubclassOf<UNavArea> NewAreaClass)
 		RefreshNavigationModifiers();
 	}
 }
+
+void UNavModifierComponent::OnTransformUpdated(USceneComponent* RootComponent, EUpdateTransformFlags UpdateTransformFlags, ETeleportType Teleport)
+{
+	// force bounds recaching next time GetNavigationBounds gets called.
+	bBoundsInitialized = false;
+
+	// otherwise the update will be handled by the default path
+	if (RootComponent && RootComponent->CanEverAffectNavigation() == false)
+	{
+		// since the parent is not nav-relevant we need to manually tell nav sys to update
+		RefreshNavigationModifiers();
+	}
+}
diff --git a/Engine/Source/Runtime/NavigationSystem/Public/NavModifierComponent.h b/Engine/Source/Runtime/NavigationSystem/Public/NavModifierComponent.h
--- a/Engine/Source/Runtime/NavigationSystem/Public/NavModifierComponent.h
+++ b/Engine/Source/Runtime/NavigationSystem/Public/NavModifierComponent.h
@@ -35,6 +35,8 @@ class NAVIGATIONSYSTEM_API UNavModifierComponent : public UNavRelevantComponent
 	void SetAreaClass(TSubclassOf<UNavArea> NewAreaClass);
 
 protected:
+	void OnTransformUpdated(USceneComponent* RootComponent, EUpdateTransformFlags UpdateTransformFlags, ETeleportType Teleport);
+
 	struct FRotatedBox
 	{
 		FBox Box;
@@ -45,4 +47,8 @@ protected:
 	};
 
 	mutable TArray<FRotatedBox> ComponentBounds;
+	mutable FDelegateHandle TransformUpdateHandle;
+	/** cached in CalcAndCacheBounds and tested in GetNavigationData to see if
+	 *	cached data is still valid */
+	mutable FTransform CachedTransform;
 };
