Replaced UNavigationSystemV1::ConditionalPopulateNavOctreeActor with UNavigationSystemV1::AddLevelToOctree #UE4

Removed UNavigationSystemV1::ConditionalPopulateNavOctreeActor has been added to allow one of the internal project filter out some actors being added to nav octree. I'm removing this function before it ships with 4.23. I've replaced the function with UNavigationSystemV1::AddLevelToOctree that gets called in the same circumstances but a bit "higher level" - it's operating on the whole ULevel rather than individual actors, which ends up being more efficient.

diff --git a/Engine/Source/Runtime/NavigationSystem/Private/NavigationSystem.cpp b/Engine/Source/Runtime/NavigationSystem/Private/NavigationSystem.cpp
--- a/Engine/Source/Runtime/NavigationSystem/Private/NavigationSystem.cpp
+++ b/Engine/Source/Runtime/NavigationSystem/Private/NavigationSystem.cpp
@@ -645,17 +645,9 @@ bool UNavigationSystemV1::ConditionalPopulateNavOctree()
 			for (int32 LevelIndex = 0; LevelIndex < World->GetNumLevels(); ++LevelIndex)
 			{
 				ULevel* Level = World->GetLevel(LevelIndex);
-				AddLevelCollisionToOctree(Level);
-
-				for (int32 ActorIndex = 0; ActorIndex < Level->Actors.Num(); ActorIndex++)
+				if (ensure(Level))
 				{
-					AActor* Actor = Level->Actors[ActorIndex];
-
-					const bool bLegalActor = Actor && !Actor->IsPendingKill();
-					if (bLegalActor)
-					{
-						ConditionalPopulateNavOctreeActor(*Actor);
-					}
+					AddLevelToOctree(*Level);
 				}
 			}
 		}
@@ -3029,11 +3021,6 @@ void UNavigationSystemV1::OnActorUnregistered(AActor* Actor)
 	}
 }
 
-void UNavigationSystemV1::ConditionalPopulateNavOctreeActor(AActor& Actor)
-{
-	UpdateActorAndComponentsInNavOctree(Actor);
-}
-
 void UNavigationSystemV1::FindElementsInNavOctree(const FBox& QueryBox, const FNavigationOctreeFilter& Filter, TArray<FNavigationOctreeElement>& Elements)
 {
 	if (NavOctree.IsValid() == false)
@@ -3709,6 +3696,22 @@ void UNavigationSystemV1::OnLevelRemovedFromWorld(ULevel* InLevel, UWorld* InWor
 	}
 }
 
+void UNavigationSystemV1::AddLevelToOctree(ULevel& Level)
+{
+	AddLevelCollisionToOctree(&Level);
+
+	for (int32 ActorIndex = 0; ActorIndex < Level.Actors.Num(); ActorIndex++)
+	{
+		AActor* Actor = Level.Actors[ActorIndex];
+
+		const bool bLegalActor = Actor && !Actor->IsPendingKill();
+		if (bLegalActor)
+		{
+			UpdateActorAndComponentsInNavOctree(*Actor);
+		}
+	}
+}
+
 void UNavigationSystemV1::AddLevelCollisionToOctree(ULevel* Level)
 {
 #if WITH_RECAST
diff --git a/Engine/Source/Runtime/NavigationSystem/Public/NavigationSystem.h b/Engine/Source/Runtime/NavigationSystem/Public/NavigationSystem.h
--- a/Engine/Source/Runtime/NavigationSystem/Public/NavigationSystem.h
+++ b/Engine/Source/Runtime/NavigationSystem/Public/NavigationSystem.h
@@ -712,8 +712,9 @@ protected:
 	/** spawn new crowd manager */
 	virtual void UpdateAbstractNavData();
 
-	/** Called when ConditionalPopulateNavOctree processes each Actor */
-	virtual void ConditionalPopulateNavOctreeActor(AActor& Actor);
+	/** Called during ConditionalPopulateNavOctree and gives subclassess a chance 
+	 *	to influence what gets added */
+	virtual void AddLevelToOctree(ULevel& Level);
 	
 public:
 	/** Called upon UWorld destruction to release what needs to be released */
