Re-registering NavAreas modified in the editor to support better debugging via NavigationTestingActors (for example)

diff --git a/Engine/Source/Runtime/NavigationSystem/Private/NavAreas/NavArea.cpp b/Engine/Source/Runtime/NavigationSystem/Private/NavAreas/NavArea.cpp
--- a/Engine/Source/Runtime/NavigationSystem/Private/NavAreas/NavArea.cpp
+++ b/Engine/Source/Runtime/NavigationSystem/Private/NavAreas/NavArea.cpp
@@ -106,4 +106,27 @@ void UNavArea::CopyFrom(TSubclassOf<UNavArea> AreaClass)
 	}
 }
 
+#if WITH_EDITOR
+void UNavArea::PostEditChangeProperty(FPropertyChangedEvent& PropertyChangedEvent)
+{
+	static const FName NAME_DefaultCost = GET_MEMBER_NAME_CHECKED(UNavArea, DefaultCost);
+	static const FName NAME_FixedAreaEnteringCost = GET_MEMBER_NAME_CHECKED(UNavArea, FixedAreaEnteringCost);
+
+	Super::PostEditChangeProperty(PropertyChangedEvent);
 
+	if (HasAnyFlags(RF_ClassDefaultObject)
+#if WITH_HOT_RELOAD
+		&& !GIsHotReload
+#endif // WITH_HOT_RELOAD
+		)
+	{
+		const FName PropertyName = (PropertyChangedEvent.Property != nullptr) ? PropertyChangedEvent.Property->GetFName() : NAME_None;
+		if (PropertyName == NAME_DefaultCost
+			|| PropertyName == NAME_FixedAreaEnteringCost)
+		{
+			UNavigationSystemV1::RequestAreaUnregistering(GetClass());
+			RegisterArea();
+		}
+	}
+}
+#endif // WITH_EDITOR
diff --git a/Engine/Source/Runtime/NavigationSystem/Public/NavAreas/NavArea.h b/Engine/Source/Runtime/NavigationSystem/Public/NavAreas/NavArea.h
--- a/Engine/Source/Runtime/NavigationSystem/Public/NavAreas/NavArea.h
+++ b/Engine/Source/Runtime/NavigationSystem/Public/NavAreas/NavArea.h
@@ -85,6 +85,9 @@ public:
 	virtual void PostLoad() override;
 	virtual void PostInitProperties() override;
 	virtual void Serialize(FArchive& Ar) override;
+#if WITH_EDITOR
+	virtual void PostEditChangeProperty(FPropertyChangedEvent& PropertyChangedEvent);
+#endif // WITH_EDITOR
 
 	FORCEINLINE uint16 GetAreaFlags() const { return AreaFlags; }
 	FORCEINLINE bool HasFlags(uint16 InFlags) const { return (InFlags & AreaFlags) != 0; }
