Added logging to the two cases that can make RecastNavMeshGenerator recreate the dtNavMesh instance (dumping the serialized one)

diff --git a/Engine/Source/Runtime/NavigationSystem/Private/NavMesh/RecastNavMeshGenerator.cpp b/Engine/Source/Runtime/NavigationSystem/Private/NavMesh/RecastNavMeshGenerator.cpp
--- a/Engine/Source/Runtime/NavigationSystem/Private/NavMesh/RecastNavMeshGenerator.cpp
+++ b/Engine/Source/Runtime/NavigationSystem/Private/NavMesh/RecastNavMeshGenerator.cpp
@@ -3262,6 +3262,11 @@ FRecastNavMeshGenerator::FRecastNavMeshGenerator(ARecastNavMesh& InDestNavMesh)
 					{
 						bRecreateNavmesh = false;
 					}
+					else
+					{
+						UE_LOG(LogNavigation, Warning, TEXT("Recreating dtNavMesh instance due to saved navmesh origin (%s, usually the RecastNavMesh location) not being aligned with tile size (%d uu) ")
+							, *Orig.ToString(), TileDim);
+					}
 				}
 
 				// if new navmesh needs more tiles, force recreating
@@ -3271,6 +3276,9 @@ FRecastNavMeshGenerator::FRecastNavMeshGenerator(ARecastNavMesh& InDestNavMesh)
 					if (FMath::Log2(MaxTiles) != FMath::Log2(SavedNavParams->maxTiles))
 					{
 						bRecreateNavmesh = true;
+						UE_LOG(LogNavigation, Warning, TEXT("Recreating dtNavMesh instance due mismatch in number of bytes required to store serialized maxTiles (%d, %d bits) vs calculated maxtiles (%d, %d bits)")
+							, SavedNavParams->maxTiles, FMath::CeilToInt(FMath::Log2(SavedNavParams->maxTiles))
+							, MaxTiles, FMath::CeilToInt(FMath::Log2(MaxTiles)));
 					}
 				}
 			}
