Fixed NavModifierVolume not updating navmesh upon undoing changes in the editor

diff --git a/Engine/Source/Runtime/NavigationSystem/Private/NavModifierVolume.cpp b/Engine/Source/Runtime/NavigationSystem/Private/NavModifierVolume.cpp
--- a/Engine/Source/Runtime/NavigationSystem/Private/NavModifierVolume.cpp
+++ b/Engine/Source/Runtime/NavigationSystem/Private/NavModifierVolume.cpp
@@ -59,6 +59,10 @@ void ANavModifierVolume::PostEditUndo()
 {
 	Super::PostEditUndo();
 
+	if (GetBrushComponent())
+	{
+		GetBrushComponent()->BuildSimpleBrushCollision();
+	}
 	FNavigationSystem::UpdateActorData(*this);
 }
 
diff --git a/Engine/Source/Runtime/NavigationSystem/Private/NavigationSystem.cpp b/Engine/Source/Runtime/NavigationSystem/Private/NavigationSystem.cpp
--- a/Engine/Source/Runtime/NavigationSystem/Private/NavigationSystem.cpp
+++ b/Engine/Source/Runtime/NavigationSystem/Private/NavigationSystem.cpp
@@ -3784,10 +3784,14 @@ void UNavigationSystemV1::OnPostLoadMap(UWorld*)
 #if WITH_EDITOR
 void UNavigationSystemV1::OnActorMoved(AActor* Actor)
 {
-	if (Cast<ANavMeshBoundsVolume>(Actor) != NULL)
+	if (Cast<ANavMeshBoundsVolume>(Actor))
 	{
 		OnNavigationBoundsUpdated((ANavMeshBoundsVolume*)Actor);
 	}
+	else if (Actor)
+	{
+		UpdateActorAndComponentsInNavOctree(*Actor, /*bUpdateAttachedActors=*/true);
+	}
 }
 #endif // WITH_EDITOR
 
