Fixed a bug in ARecastNavMesh::GetRandomReachablePointInRadius that resulted in erronous behavior for big(ish) Radius values.

The bug resulted from us using a strange extent ({Radius,Radius,Radius}) when projecting the initial location to navmesh. I've used a lot more reasonable DefaultQueryExtent for this purpose with Z being set to BIG_NUMBER. When a big Radius value was used the underlying recast code was expected to query a lot of polygons, but there's a hard limit of 128 polygons to process in multiple places in Recast lib (for performance reasons). These two facts combined resulted in finding awkward query starting point when using large Radius. The new approach first projects the query origin 'down' taking advantage of the assumption that the querier location is (more or less) on navmesh, since the whole query is about reachable locations.

diff --git a/Engine/Source/Runtime/NavigationSystem/Private/NavMesh/RecastNavMesh.cpp b/Engine/Source/Runtime/NavigationSystem/Private/NavMesh/RecastNavMesh.cpp
--- a/Engine/Source/Runtime/NavigationSystem/Private/NavMesh/RecastNavMesh.cpp
+++ b/Engine/Source/Runtime/NavigationSystem/Private/NavMesh/RecastNavMesh.cpp
@@ -975,26 +975,28 @@ bool ARecastNavMesh::GetRandomReachablePointInRadius(const FVector& Origin, floa
 	if (QueryFilter)
 	{
 		// find starting poly
+		const FVector ProjectionExtent(NavDataConfig.DefaultQueryExtent.X, NavDataConfig.DefaultQueryExtent.Y, BIG_NUMBER);
+		const FVector RcExtent = Unreal2RecastPoint(ProjectionExtent).GetAbs();
 		// convert start/end pos to Recast coords
-		const float Extent[3] = { Radius, Radius, Radius };
 		const FVector RecastOrigin = Unreal2RecastPoint(Origin);
 		NavNodeRef OriginPolyID = INVALID_NAVNODEREF;
-		NavQuery.findNearestPoly(&RecastOrigin.X, Extent, QueryFilter, &OriginPolyID, nullptr);
+		NavQuery.findNearestPoly(&RecastOrigin.X, &RcExtent.X, QueryFilter, &OriginPolyID, nullptr);
 
-		dtPolyRef Poly;
-		float RandPt[3];
-		dtStatus Status = NavQuery.findRandomPointAroundCircle(OriginPolyID, &RecastOrigin.X, Radius
-			, QueryFilter, FMath::FRand, &Poly, RandPt);
-
-		if (dtStatusSucceed(Status))
-		{
-			OutResult = FNavLocation(Recast2UnrealPoint(RandPt), Poly);
-			return true;
-		}
-		else
+		if (OriginPolyID != INVALID_NAVNODEREF)
 		{
-			OutResult = FNavLocation(Origin, OriginPolyID);
+			dtPolyRef Poly;
+			float RandPt[3];
+			dtStatus Status = NavQuery.findRandomPointAroundCircle(OriginPolyID, &RecastOrigin.X, Radius
+				, QueryFilter, FMath::FRand, &Poly, RandPt);
+
+			if (dtStatusSucceed(Status))
+			{
+				OutResult = FNavLocation(Recast2UnrealPoint(RandPt), Poly);
+				return true;
+			}
 		}
+
+		OutResult = FNavLocation(Origin, OriginPolyID);
 	}
 
 	return false;
