Made sure we mark all the navigable space dirty after recreating dtNavMesh instance due to mismatching navmesh properties.

diff --git a/Engine/Source/Runtime/NavigationSystem/Private/NavMesh/RecastNavMeshGenerator.cpp b/Engine/Source/Runtime/NavigationSystem/Private/NavMesh/RecastNavMeshGenerator.cpp
--- a/Engine/Source/Runtime/NavigationSystem/Private/NavMesh/RecastNavMeshGenerator.cpp
+++ b/Engine/Source/Runtime/NavigationSystem/Private/NavMesh/RecastNavMeshGenerator.cpp
@@ -3265,7 +3265,7 @@ FRecastNavMeshGenerator::FRecastNavMeshGenerator(ARecastNavMesh& InDestNavMesh)
 					else
 					{
 						UE_LOG(LogNavigation, Warning, TEXT("Recreating dtNavMesh instance due to saved navmesh origin (%s, usually the RecastNavMesh location) not being aligned with tile size (%d uu) ")
-							, *Orig.ToString(), TileDim);
+							, *Orig.ToString(), int(TileDim));
 					}
 				}
 
@@ -3289,6 +3289,8 @@ FRecastNavMeshGenerator::FRecastNavMeshGenerator(ARecastNavMesh& InDestNavMesh)
 	{
 		// recreate navmesh from scratch if no data was loaded
 		ConstructTiledNavMesh();
+		// mark all the areas we need to update, which is the whole (known) navigable space
+		MarkNavBoundsDirty();
 	}
 	else
 	{
@@ -3544,19 +3546,7 @@ bool FRecastNavMeshGenerator::RebuildAll()
 
 	ConstructTiledNavMesh();
 	
-	// if rebuilding all no point in keeping "old" invalidated areas
-	TArray<FNavigationDirtyArea> DirtyAreas;
-	for (FBox AreaBounds : InclusionBounds)
-	{
-		FNavigationDirtyArea DirtyArea(AreaBounds, ENavigationDirtyFlag::All | ENavigationDirtyFlag::NavigationBounds);
-		DirtyAreas.Add(DirtyArea);
-	}
-
-	if (DirtyAreas.Num())
-	{
-		MarkDirtyTiles(DirtyAreas);
-	}
-	else
+	if (MarkNavBoundsDirty() == false)
 	{
 		// There are no navigation bounds to build, probably navmesh was resized and we just need to update debug draw
 		DestNavMesh->RequestDrawingUpdate();
@@ -4141,6 +4131,24 @@ int32 FRecastNavMeshGenerator::GetDirtyTilesCount(const FBox& AreaBounds) const
 	return DirtyPendingCount + RunningCount;
 }
 
+bool FRecastNavMeshGenerator::MarkNavBoundsDirty()
+{
+	// if rebuilding all no point in keeping "old" invalidated areas
+	TArray<FNavigationDirtyArea> DirtyAreas;
+	for (FBox AreaBounds : InclusionBounds)
+	{
+		FNavigationDirtyArea DirtyArea(AreaBounds, ENavigationDirtyFlag::All | ENavigationDirtyFlag::NavigationBounds);
+		DirtyAreas.Add(DirtyArea);
+	}
+
+	if (DirtyAreas.Num())
+	{
+		MarkDirtyTiles(DirtyAreas);
+		return true;
+	}
+	return false;
+}
+
 void FRecastNavMeshGenerator::MarkDirtyTiles(const TArray<FNavigationDirtyArea>& DirtyAreas)
 {
 	QUICK_SCOPE_CYCLE_COUNTER(STAT_RecastNavMeshGenerator_MarkDirtyTiles);
diff --git a/Engine/Source/Runtime/NavigationSystem/Public/NavMesh/RecastNavMeshGenerator.h b/Engine/Source/Runtime/NavigationSystem/Public/NavMesh/RecastNavMeshGenerator.h
--- a/Engine/Source/Runtime/NavigationSystem/Public/NavMesh/RecastNavMeshGenerator.h
+++ b/Engine/Source/Runtime/NavigationSystem/Public/NavMesh/RecastNavMeshGenerator.h
@@ -520,6 +520,9 @@ protected:
 	/** Marks grid tiles affected by specified areas as dirty */
 	virtual void MarkDirtyTiles(const TArray<FNavigationDirtyArea>& DirtyAreas);
 
+	/** Marks all tiles overlapping with InclusionBounds dirty (via MarkDirtyTiles). */
+	bool MarkNavBoundsDirty();
+
 	void RemoveLayers(const FIntPoint& Tile, TArray<uint32>& UpdatedTiles);
 	
 #if RECAST_ASYNC_REBUILDING
