Added a dedicated warning for when we attempt to creat a dtNavMesh instance with maxTiles == 0.

Also, made sure we release the faultu dtNavMesh instance if we're not going to use it.

diff --git a/Engine/Source/Runtime/NavigationSystem/Private/NavMesh/RecastNavMeshGenerator.cpp b/Engine/Source/Runtime/NavigationSystem/Private/NavMesh/RecastNavMeshGenerator.cpp
--- a/Engine/Source/Runtime/NavigationSystem/Private/NavMesh/RecastNavMeshGenerator.cpp
+++ b/Engine/Source/Runtime/NavigationSystem/Private/NavMesh/RecastNavMeshGenerator.cpp
@@ -3464,18 +3464,31 @@ bool FRecastNavMeshGenerator::ConstructTiledNavMesh()
 		CalcNavMeshProperties(TiledMeshParameters.maxTiles, TiledMeshParameters.maxPolys);
 		Config.MaxPolysPerTile = TiledMeshParameters.maxPolys;
 
-		const dtStatus status = DetourMesh->init(&TiledMeshParameters);
-
-		if (dtStatusFailed(status))
+		if (TiledMeshParameters.maxTiles == 0)
 		{
-			UE_LOG(LogNavigation, Warning, TEXT("ConstructTiledNavMesh: Could not init navmesh.") );
+			UE_LOG(LogNavigation, Warning, TEXT("ConstructTiledNavMesh: Failed to create navmesh of size 0."));
 			bSuccess = false;
 		}
 		else
 		{
-			bSuccess = true;
-			NumActiveTiles = GetTilesCountHelper(DetourMesh);
-			DestNavMesh->GetRecastNavMeshImpl()->SetRecastMesh(DetourMesh);
+			const dtStatus status = DetourMesh->init(&TiledMeshParameters);
+
+			if (dtStatusFailed(status))
+			{
+				UE_LOG(LogNavigation, Warning, TEXT("ConstructTiledNavMesh: Could not init navmesh."));
+				bSuccess = false;
+			}
+			else
+			{
+				bSuccess = true;
+				NumActiveTiles = GetTilesCountHelper(DetourMesh);
+				DestNavMesh->GetRecastNavMeshImpl()->SetRecastMesh(DetourMesh);
+			}
+		}
+
+		if (bSuccess == false)
+		{
+			dtFreeNavMesh(DetourMesh);
 		}
 	}
 	else
