RecastNavMeshGenerator: Fixed regression introduced by 6675381 by configuring flag bRestrictBuildingToActiveTiles before marking NavBounds dirty for newly created navmesh. This prevents the whole navmesh to be generated when using Navigation invokers.

#rnx

diff --git a/Engine/Source/Runtime/NavigationSystem/Private/NavMesh/RecastNavMeshGenerator.cpp b/Engine/Source/Runtime/NavigationSystem/Private/NavMesh/RecastNavMeshGenerator.cpp
--- a/Engine/Source/Runtime/NavigationSystem/Private/NavMesh/RecastNavMeshGenerator.cpp
+++ b/Engine/Source/Runtime/NavigationSystem/Private/NavMesh/RecastNavMeshGenerator.cpp
@@ -3289,7 +3289,13 @@ FRecastNavMeshGenerator::FRecastNavMeshGenerator(ARecastNavMesh& InDestNavMesh)
 	{
 		// recreate navmesh from scratch if no data was loaded
 		ConstructTiledNavMesh();
-		// mark all the areas we need to update, which is the whole (known) navigable space
+
+		// mark all the areas we need to update, which is the whole (known) navigable space if not restricted to active tiles
+		const UNavigationSystemV1* NavSys = FNavigationSystem::GetCurrent<UNavigationSystemV1>(GetWorld());
+		if (NavSys)
+		{
+			bRestrictBuildingToActiveTiles = NavSys->IsActiveTilesGenerationEnabled();
+		}
 		MarkNavBoundsDirty();
 	}
 	else
