Fix regression in reachitecture of NavSystem

#rnx

diff --git a/Engine/Source/Runtime/NavigationSystem/Private/NavigationDataHandler.cpp b/Engine/Source/Runtime/NavigationSystem/Private/NavigationDataHandler.cpp
--- a/Engine/Source/Runtime/NavigationSystem/Private/NavigationDataHandler.cpp
+++ b/Engine/Source/Runtime/NavigationSystem/Private/NavigationDataHandler.cpp
@@ -454,7 +454,14 @@ void FNavigationDataHandler::UpdateActorAndComponentsInNavOctree(AActor& Actor)
 		if (CompNavElement)
 		{
 			// Component != null is implied by successful INavRelevantInterface cast 
-			UpdateNavOctreeElement(*Component, *CompNavElement, FNavigationOctreeController::OctreeUpdate_Default);
+			if (Actor.IsComponentRelevantForNavigation(Component))
+			{
+				UpdateNavOctreeElement(*Component, *CompNavElement, FNavigationOctreeController::OctreeUpdate_Default);
+			}
+			else
+			{
+				UnregisterNavOctreeElement(*Component, *CompNavElement, FNavigationOctreeController::OctreeUpdate_Default);
+			}
 		}
 	}
 }
