ANavSystemConfigOverride: Postpone override effect to next tick when world is ticking to prevent deletion of the navigation system while its getting updated.

#rnx

diff --git a/Engine/Source/Runtime/NavigationSystem/Private/NavSystemConfigOverride.cpp b/Engine/Source/Runtime/NavigationSystem/Private/NavSystemConfigOverride.cpp
--- a/Engine/Source/Runtime/NavigationSystem/Private/NavSystemConfigOverride.cpp
+++ b/Engine/Source/Runtime/NavigationSystem/Private/NavSystemConfigOverride.cpp
@@ -62,7 +62,23 @@ ANavSystemConfigOverride::ANavSystemConfigOverride(const FObjectInitializer& Obj
 void ANavSystemConfigOverride::PostLoad()
 {
 	Super::PostLoad();
+	
+	UWorld* World = GetWorld();
+	if (World != nullptr)
+	{
+		if (World->bInTick)
+		{
+			World->GetTimerManager().SetTimerForNextTick(this, &ANavSystemConfigOverride::ApplyConfig);
+		}
+		else
+		{
+			ApplyConfig();
+		}
+	}
+}
 
+void ANavSystemConfigOverride::ApplyConfig()
+{
 	UWorld* World = GetWorld();
 	if (World && NavigationSystemConfig)
 	{
diff --git a/Engine/Source/Runtime/NavigationSystem/Public/NavSystemConfigOverride.h b/Engine/Source/Runtime/NavigationSystem/Public/NavSystemConfigOverride.h
--- a/Engine/Source/Runtime/NavigationSystem/Public/NavSystemConfigOverride.h
+++ b/Engine/Source/Runtime/NavigationSystem/Public/NavSystemConfigOverride.h
@@ -53,4 +53,6 @@ protected:
 	/** Called only in the editor mode*/
 	void InitializeForWorld(UNavigationSystemBase* NewNavSys, UWorld* World, const FNavigationSystemRunMode RunMode);
 #endif // WITH_EDITOR
+
+	void ApplyConfig();
 };
