NavLinkComponent: Fixed NavigationLink AreaClass cache when the component gets reregistered after running construction scripts for blueprint classes.
NavLinkRenderingProxy: Added support for components implementing INavLinkHostInterface so it can render links from NavLinkComponent.
NavLinkProxy: Appended link from the SmartLinkComponent when it's NavigationRelevent so it gets rendered with the simple links

diff --git a/Engine/Source/Runtime/AIModule/Private/Navigation/NavLinkProxy.cpp b/Engine/Source/Runtime/AIModule/Private/Navigation/NavLinkProxy.cpp
--- a/Engine/Source/Runtime/AIModule/Private/Navigation/NavLinkProxy.cpp
+++ b/Engine/Source/Runtime/AIModule/Private/Navigation/NavLinkProxy.cpp
@@ -194,9 +194,16 @@ bool ANavLinkProxy::GetNavigationLinksClasses(TArray<TSubclassOf<UNavLinkDefinit
 bool ANavLinkProxy::GetNavigationLinksArray(TArray<FNavigationLink>& OutLink, TArray<FNavigationSegmentLink>& OutSegments) const
 {
 	OutLink.Append(PointLinks);
+
+	const bool bIsSmartLinkActive = (SmartLinkComp && SmartLinkComp->IsNavigationRelevant());
+	if (bIsSmartLinkActive)
+	{
+		OutLink.Add(SmartLinkComp->GetLinkModifier());
+	}
+
 	OutSegments.Append(SegmentLinks);
 
-	return (PointLinks.Num() > 0) || (SegmentLinks.Num() > 0);
+	return (PointLinks.Num() > 0) || (SegmentLinks.Num() > 0) || bIsSmartLinkActive;
 }
 
 FBox ANavLinkProxy::GetComponentsBoundingBox(bool bNonColliding) const
diff --git a/Engine/Source/Runtime/NavigationSystem/Private/NavLinkComponent.cpp b/Engine/Source/Runtime/NavigationSystem/Private/NavLinkComponent.cpp
--- a/Engine/Source/Runtime/NavigationSystem/Private/NavLinkComponent.cpp
+++ b/Engine/Source/Runtime/NavigationSystem/Private/NavLinkComponent.cpp
@@ -67,10 +67,7 @@ void UNavLinkComponent::PostEditChangeProperty(FPropertyChangedEvent& PropertyCh
 
 	if (PropertyChangedEvent.MemberProperty && PropertyChangedEvent.MemberProperty->GetFName() == GET_MEMBER_NAME_CHECKED(UNavLinkComponent, Links))
 	{
-		for (FNavigationLink& Link : Links)
-		{
-			Link.InitializeAreaClass(/*bForceRefresh=*/true);
-		}
+		InitializeLinksAreaClasses();
 	}
 }
 
@@ -78,23 +75,33 @@ void UNavLinkComponent::PostEditUndo()
 {
 	Super::PostEditUndo();
 
-	for (FNavigationLink& Link : Links)
-	{
-		Link.InitializeAreaClass(/*bForceRefresh=*/true);
-	}
+	InitializeLinksAreaClasses();
 }
 
 void UNavLinkComponent::PostEditImport()
 {
 	Super::PostEditImport();
 
+	InitializeLinksAreaClasses();
+}
+
+#endif // WITH_EDITOR
+
+void UNavLinkComponent::OnRegister()
+{
+	Super::OnRegister();
+
+	InitializeLinksAreaClasses();
+}
+
+void UNavLinkComponent::InitializeLinksAreaClasses()
+{
 	for (FNavigationLink& Link : Links)
 	{
-		Link.InitializeAreaClass(/*bForceRefresh=*/true);
+		Link.InitializeAreaClass();
 	}
 }
 
-#endif // WITH_EDITOR
 //----------------------------------------------------------------------//
 // UNavLinkTrivial
 //----------------------------------------------------------------------//
diff --git a/Engine/Source/Runtime/NavigationSystem/Private/NavLinkRenderingComponent.cpp b/Engine/Source/Runtime/NavigationSystem/Private/NavLinkRenderingComponent.cpp
--- a/Engine/Source/Runtime/NavigationSystem/Private/NavLinkRenderingComponent.cpp
+++ b/Engine/Source/Runtime/NavigationSystem/Private/NavLinkRenderingComponent.cpp
@@ -106,7 +106,12 @@ FNavLinkRenderingProxy::FNavLinkRenderingProxy(const UPrimitiveComponent* InComp
 	: FPrimitiveSceneProxy(InComponent)
 {
 	LinkOwnerActor = InComponent->GetOwner();
-	LinkOwnerHost = Cast<INavLinkHostInterface>(InComponent->GetOwner());
+	LinkOwnerHost = Cast<INavLinkHostInterface>((UPrimitiveComponent*)InComponent);
+
+	if (LinkOwnerHost == nullptr)
+	{
+		LinkOwnerHost = Cast<INavLinkHostInterface>(LinkOwnerActor);
+	}
 
 	if (LinkOwnerActor != NULL && LinkOwnerHost != NULL)
 	{
diff --git a/Engine/Source/Runtime/NavigationSystem/Public/NavLinkComponent.h b/Engine/Source/Runtime/NavigationSystem/Public/NavLinkComponent.h
--- a/Engine/Source/Runtime/NavigationSystem/Public/NavLinkComponent.h
+++ b/Engine/Source/Runtime/NavigationSystem/Public/NavLinkComponent.h
@@ -36,4 +36,9 @@ class NAVIGATIONSYSTEM_API UNavLinkComponent : public UPrimitiveComponent, publi
 	virtual void PostEditUndo() override;
 	virtual void PostEditImport() override;
 #endif // WITH_EDITOR
+
+	virtual void OnRegister() override;
+
+protected:
+	void InitializeLinksAreaClasses();
 };
