NavLinkCustomComponent: Restored NavLinkUserId during blueprint classes reconstruction process. This fixes an issue when reloading a map and link Ids were bumped by the reconstruction process and were no longer matching the ids saved in the NavigationData.

#rn NavLinkUserId for NavLinkCustomComponents used in blueprint classes are now properly restored when runnning construction scripts

diff --git a/Engine/Source/Runtime/NavigationSystem/Private/NavLinkCustomComponent.cpp b/Engine/Source/Runtime/NavigationSystem/Private/NavLinkCustomComponent.cpp
--- a/Engine/Source/Runtime/NavigationSystem/Private/NavLinkCustomComponent.cpp
+++ b/Engine/Source/Runtime/NavigationSystem/Private/NavLinkCustomComponent.cpp
@@ -42,6 +42,39 @@ void UNavLinkCustomComponent::PostLoad()
 	INavLinkCustomInterface::UpdateUniqueId(NavLinkUserId);
 }
 
+TStructOnScope<FActorComponentInstanceData> UNavLinkCustomComponent::GetComponentInstanceData() const
+{
+	TStructOnScope<FActorComponentInstanceData> InstanceData = MakeStructOnScope<FActorComponentInstanceData, FNavLinkCustomInstanceData>(this);
+	FNavLinkCustomInstanceData* NavLinkCustomInstanceData = InstanceData.Cast<FNavLinkCustomInstanceData>();
+	NavLinkCustomInstanceData->NavLinkUserId = NavLinkUserId;
+
+	return InstanceData;
+}
+
+void UNavLinkCustomComponent::ApplyComponentInstanceData(FNavLinkCustomInstanceData* NavLinkData)
+{
+	check(NavLinkData);
+
+	if (NavLinkUserId != NavLinkData->NavLinkUserId)
+	{
+		// Registered component has its link registered in the navigation system. 
+		// In such case, we need to unregister current Id and register with the Id from the instance data.
+		const bool bIsLinkRegistrationUpdateRequired = IsRegistered();
+
+		if (bIsLinkRegistrationUpdateRequired)
+		{
+			UNavigationSystemV1::RequestCustomLinkUnregistering(*this, this);
+		}
+
+		NavLinkUserId = NavLinkData->NavLinkUserId;
+
+		if (bIsLinkRegistrationUpdateRequired)
+		{
+			UNavigationSystemV1::RequestCustomLinkRegistering(*this, this);
+		}
+	}
+}
+
 #if WITH_EDITOR
 void UNavLinkCustomComponent::PostEditImport()
 {
diff --git a/Engine/Source/Runtime/NavigationSystem/Public/NavLinkCustomComponent.h b/Engine/Source/Runtime/NavigationSystem/Public/NavLinkCustomComponent.h
--- a/Engine/Source/Runtime/NavigationSystem/Public/NavLinkCustomComponent.h
+++ b/Engine/Source/Runtime/NavigationSystem/Public/NavLinkCustomComponent.h
@@ -51,8 +51,11 @@ class NAVIGATIONSYSTEM_API UNavLinkCustomComponent : public UNavRelevantComponen
 	//~ Begin UActorComponent Interface
 	virtual void OnRegister() override;
 	virtual void OnUnregister() override;
+	virtual TStructOnScope<FActorComponentInstanceData> GetComponentInstanceData() const override;
 	//~ End UActorComponent Interface
 
+	void ApplyComponentInstanceData(struct FNavLinkCustomInstanceData* ComponentInstanceData);
+
 	//~ Begin UObject Interface
 	virtual void PostLoad() override;
 #if WITH_EDITOR
@@ -220,3 +223,37 @@ protected:
 	/** gather agents to notify about state change */
 	void CollectNearbyAgents(TArray<UObject*>& NotifyList);
 };
+
+/** Used to store navlink data during RerunConstructionScripts */
+USTRUCT()
+struct FNavLinkCustomInstanceData : public FActorComponentInstanceData
+{
+	GENERATED_BODY()
+
+public:
+	FNavLinkCustomInstanceData() = default;
+	FNavLinkCustomInstanceData(const UNavLinkCustomComponent* SourceComponent)
+		: FActorComponentInstanceData(SourceComponent)
+		, NavLinkUserId(0)
+	{}
+
+	virtual ~FNavLinkCustomInstanceData() = default;
+
+	virtual bool ContainsData() const override
+	{
+		return true;
+	}
+
+	virtual void ApplyToComponent(UActorComponent* Component, const ECacheApplyPhase CacheApplyPhase) override
+	{
+		Super::ApplyToComponent(Component, CacheApplyPhase);
+
+		if (CacheApplyPhase == ECacheApplyPhase::PostUserConstructionScript)
+		{
+			CastChecked<UNavLinkCustomComponent>(Component)->ApplyComponentInstanceData(this);
+		}
+	}
+
+	UPROPERTY()
+	uint32 NavLinkUserId;
+};
\ No newline at end of file
