- Moved the work done in ProcessPendingNavigation from a Task to be triggered via the navigation system tick, after the other required actors were processed such as navigation bound volumes and that the owning level of the navigation data is full loaded.
- Added a call to ConditionalPopulateNavOctree to fix navigation not being regenerated when adding a level with some navigation data to a current world that didn't have one.
- Renamed boolean to better names
- MarkPendingKill the navigation instance when no compatible nav bound for it in the navigation system insead of in the nav mesh generator (except for the abstract nav data)

#rnx

Merging manually CL 7842753 using //Release-10.30 to //Main

diff --git a/Engine/Source/Runtime/Engine/Classes/AI/NavigationSystemBase.h b/Engine/Source/Runtime/Engine/Classes/AI/NavigationSystemBase.h
--- a/Engine/Source/Runtime/Engine/Classes/AI/NavigationSystemBase.h
+++ b/Engine/Source/Runtime/Engine/Classes/AI/NavigationSystemBase.h
@@ -224,9 +224,6 @@ public:
 
 	virtual void InitializeForWorld(UWorld& World, FNavigationSystemRunMode Mode) PURE_VIRTUAL(UNavigationSystemBase::InitializeForWorld, );
 
-	/** Called after a NavSystemConfigOverride has finished overriding the old NavSystem */
-	virtual void OnNavSystemOverriden(UNavigationSystemBase* PreviousNavSystem) { };
-
 	/** 
 	 *	If you're using NavigationSysstem module consider calling 
 	 *	FNavigationSystem::GetCurrent<UNavigationSystemV1>()->GetDefaultNavDataInstance 
diff --git a/Engine/Source/Runtime/NavigationSystem/Private/NavMesh/RecastNavMeshGenerator.cpp b/Engine/Source/Runtime/NavigationSystem/Private/NavMesh/RecastNavMeshGenerator.cpp
--- a/Engine/Source/Runtime/NavigationSystem/Private/NavMesh/RecastNavMeshGenerator.cpp
+++ b/Engine/Source/Runtime/NavigationSystem/Private/NavMesh/RecastNavMeshGenerator.cpp
@@ -3565,17 +3565,8 @@ bool FRecastNavMeshGenerator::RebuildAll()
 	
 	if (MarkNavBoundsDirty() == false)
 	{
-		const UNavigationSystemV1* NavSys = FNavigationSystem::GetCurrent<UNavigationSystemV1>(GetWorld());
-		if (NavSys && !NavSys->ShouldGenerateNavDataWhenNoCompatibleNavBound())
-		{
-			// There are no navigation bounds to build, no point keeping the nav mesh
-			DestNavMesh->Destroy();
-		}
-		else
-		{
-			// There are no navigation bounds to build, probably navmesh was resized and we just need to update debug draw
-			DestNavMesh->RequestDrawingUpdate();
-		}
+		// There are no navigation bounds to build, probably navmesh was resized and we just need to update debug draw
+		DestNavMesh->RequestDrawingUpdate();
 	}
 
 	return true;
diff --git a/Engine/Source/Runtime/NavigationSystem/Private/NavSystemConfigOverride.cpp b/Engine/Source/Runtime/NavigationSystem/Private/NavSystemConfigOverride.cpp
--- a/Engine/Source/Runtime/NavigationSystem/Private/NavSystemConfigOverride.cpp
+++ b/Engine/Source/Runtime/NavigationSystem/Private/NavSystemConfigOverride.cpp
@@ -4,6 +4,7 @@
 #include "Engine/World.h"
 #include "GameFramework/WorldSettings.h"
 #include "AI/NavigationSystemBase.h"
+#include "NavigationSystem.h"
 #include "TimerManager.h"
 
 #if WITH_EDITORONLY_DATA
@@ -96,7 +97,7 @@ void ANavSystemConfigOverride::ApplyConfig()
 			)
 		{
 			
-			UNavigationSystemBase* PrevNavSysBase = World->GetNavigationSystem();
+			UNavigationSystemV1* PrevNavSys = Cast<UNavigationSystemV1>(World->GetNavigationSystem());
 			World->SetNavigationSystem(nullptr);
 
 			const FNavigationSystemRunMode RunMode = World->WorldType == EWorldType::Editor
@@ -123,10 +124,10 @@ void ANavSystemConfigOverride::ApplyConfig()
 				FNavigationSystem::AddNavigationSystemToWorld(*World, RunMode, NavigationSystemConfig, /*bInitializeForWorld=*/true);
 			}
 
-			UNavigationSystemBase* NewNavSys = World->GetNavigationSystem();
-			if (NewNavSys && PrevNavSysBase)
+			UNavigationSystemV1* NewNavSys = Cast<UNavigationSystemV1>(World->GetNavigationSystem());
+			if (NewNavSys && PrevNavSys)
 			{
-				NewNavSys->OnNavSystemOverriden(PrevNavSysBase);
+				NewNavSys->FinalizeOverridingNavSystem(*PrevNavSys);
 			}
 		}
 	}
diff --git a/Engine/Source/Runtime/NavigationSystem/Private/NavigationSystem.cpp b/Engine/Source/Runtime/NavigationSystem/Private/NavigationSystem.cpp
--- a/Engine/Source/Runtime/NavigationSystem/Private/NavigationSystem.cpp
+++ b/Engine/Source/Runtime/NavigationSystem/Private/NavigationSystem.cpp
@@ -245,8 +245,8 @@ bool UNavigationSystemV1::bIsPIEActive = false;
 UNavigationSystemV1::UNavigationSystemV1(const FObjectInitializer& ObjectInitializer)
 	: Super(ObjectInitializer)
 	, bTickWhilePaused(false)
-	, bGenerateNavDataWhenNoCompatibleNavBound(true)
-	, bUseNavDataInAdditionalLevelWhenDuplicatedAgent(false)
+	, bShouldGenerateNavDataOnlyWhenCompatibleNavBound(false)
+	, bAllowEquivalentNavDataOverride(false)
 	, bWholeWorldNavigable(false)
 	, bSkipAgentHeightCheckWhenPickingNavData(false)
 	, OperationMode(FNavigationSystemRunMode::InvalidMode)
@@ -873,7 +873,7 @@ void UNavigationSystemV1::RegisterNavigationDataInstances()
 		ANavigationData* Nav = (*It);
 		if (Nav != NULL && Nav->IsPendingKill() == false && Nav->IsRegistered() == false)
 		{
-			RequestRegistration(Nav, false);
+			RequestRegistration(Nav);
 			bProcessRegistration = true;
 		}
 	}
@@ -940,6 +940,11 @@ void UNavigationSystemV1::Tick(float DeltaSeconds)
 		PendingNavBoundsUpdates.Reset();
 	}
 
+	if (NavDataRegistrationQueue.Num() > 0)
+	{
+		ProcessRegistrationCandidates();
+	}
+
 	if (DefaultOctreeController.PendingOctreeUpdates.Num() > 0)
 	{
 		SCOPE_CYCLE_COUNTER(STAT_Navigation_AddingActorsToNavOctree);
@@ -1746,29 +1751,13 @@ void UNavigationSystemV1::ApplyWorldOffset(const FVector& InOffset, bool bWorldS
 //----------------------------------------------------------------------//
 // Bookkeeping 
 //----------------------------------------------------------------------//
-void UNavigationSystemV1::RequestRegistration(ANavigationData* NavData, bool bTriggerRegistrationProcessing)
+void UNavigationSystemV1::RequestRegistration(ANavigationData* NavData)
 {
 	FScopeLock RegistrationLock(&NavDataRegistrationSection);
 
 	if (NavDataRegistrationQueue.Num() < REGISTRATION_QUEUE_SIZE)
 	{
 		NavDataRegistrationQueue.AddUnique(NavData);
-
-		// checking if bWorldInitDone since requesting out-of-order registration
-		// processing when we're still setting up can result in odd cases,
-		// like initializing navmesh generators while the nav system doesn't have
-		// the navmesh bounds collected yet.
-		if (bTriggerRegistrationProcessing && bWorldInitDone)
-		{
-			// trigger registration candidates processing
-			DECLARE_CYCLE_STAT(TEXT("FSimpleDelegateGraphTask.Process registration candidates"),
-				STAT_FSimpleDelegateGraphTask_ProcessRegistrationCandidates,
-				STATGROUP_TaskGraphTasks);
-
-			FSimpleDelegateGraphTask::CreateAndDispatchWhenReady(
-				FSimpleDelegateGraphTask::FDelegate::CreateUObject(this, &UNavigationSystemV1::ProcessRegistrationCandidates),
-				GET_STATID(STAT_FSimpleDelegateGraphTask_ProcessRegistrationCandidates), NULL, ENamedThreads::GameThread);
-		}
 	}
 	else
 	{
@@ -1778,47 +1767,47 @@ void UNavigationSystemV1::RequestRegistration(ANavigationData* NavData, bool bTr
 
 void UNavigationSystemV1::ProcessRegistrationCandidates()
 {
-	//if (FUObjectThreadContext::Get().IsRoutingPostLoad)
-	//{
-	//	// postopne
-	//	return;
-	//}
-
 	FScopeLock RegistrationLock(&NavDataRegistrationSection);
 
 	if (NavDataRegistrationQueue.Num() == 0)
 	{
 		return;
 	}
-
-	ANavigationData** NavDataPtr = NavDataRegistrationQueue.GetData();
+	
 	const int CandidatesCount = NavDataRegistrationQueue.Num();
-
-	for (int32 CandidateIndex = 0; CandidateIndex < CandidatesCount; ++CandidateIndex, ++NavDataPtr)
+	int32 NumNavDataProcessed = 0;
+	for (int32 CandidateIndex = CandidatesCount - 1; CandidateIndex >= 0; --CandidateIndex)
 	{
-		if (*NavDataPtr != NULL)
+		ANavigationData* NavDataPtr = NavDataRegistrationQueue[CandidateIndex];
+		ULevel* OwningLevel = NavDataPtr != nullptr ? NavDataPtr->GetLevel() : nullptr;
+		if (OwningLevel && OwningLevel->bIsVisible)
 		{
-			ERegistrationResult Result = RegisterNavData(*NavDataPtr);
+			const ERegistrationResult Result = RegisterNavData(NavDataPtr);
 
-			if (Result == RegistrationSuccessful)
+			if (Result != RegistrationSuccessful && Result != RegistrationFailed_DataPendingKill)
 			{
-				continue;
-			}
-			else if (Result != RegistrationFailed_DataPendingKill)
-			{
-				(*NavDataPtr)->CleanUpAndMarkPendingKill();
-				if ((*NavDataPtr) == MainNavData)
+				NavDataPtr->CleanUpAndMarkPendingKill();
+				if (NavDataPtr == MainNavData)
 				{
 					MainNavData = NULL;
 				}
 			}
-		}
-	}
 
-	MainNavData = GetDefaultNavDataInstance(FNavigationSystem::DontCreate);
+			NumNavDataProcessed++;
+			NavDataRegistrationQueue.RemoveAtSwap(CandidateIndex);
+		}
+	}	
 	
-	// we processed all candidates so clear the queue
-	NavDataRegistrationQueue.Reset();
+	if (NumNavDataProcessed)
+	{
+		MainNavData = GetDefaultNavDataInstance(FNavigationSystem::DontCreate);
+
+		// See if any of registered navigation data now needs NavOctree
+		if (DefaultOctreeController.IsValid() == false && RequiresNavOctree() == true)
+		{
+			ConditionalPopulateNavOctree();
+		}
+	}
 }
 
 void UNavigationSystemV1::ProcessCustomLinkPendingRegistration()
@@ -1855,6 +1844,15 @@ UNavigationSystemV1::ERegistrationResult UNavigationSystemV1::RegisterNavData(AN
 	{
 		return RegistrationSuccessful;
 	}
+    // Don't allow nav data to register
+	else if (ShouldGenerateNavigationEverywhere() == false && ShouldGenerateNavDataOnlyWhenCompatibleNavBound() == true && NavData->IsA(AAbstractNavData::StaticClass()) == false)
+	{
+		TArray<FBox> SupportedBounds;
+		if (GetNavigationBoundsForNavData(*NavData, SupportedBounds) <= 0)
+		{
+			return RegistrationFailed_NotSuitable;
+		}
+	}
 
 	FScopeLock Lock(&NavDataRegistration);
 
@@ -1880,11 +1878,11 @@ UNavigationSystemV1::ERegistrationResult UNavigationSystemV1::RegisterNavData(AN
 		TWeakObjectPtr<ANavigationData>* NavDataForAgent = AgentToNavDataMap.Find(NavConfig);
 		ANavigationData* NavDataInstanceForAgent = NavDataForAgent ? NavDataForAgent->Get() : nullptr;
 
-		if (NavDataInstanceForAgent == nullptr || ShouldUseNavDataInAdditionalLevelWhenDuplicatedAgent())
+		if (NavDataInstanceForAgent == nullptr || ShouldAllowEquivalentNavDataOverride())
 		{
 			if (NavData->IsA(AAbstractNavData::StaticClass()) == false)
 			{
-				if (NavDataInstanceForAgent != nullptr && ShouldUseNavDataInAdditionalLevelWhenDuplicatedAgent())
+				if (NavDataInstanceForAgent != nullptr && ShouldAllowEquivalentNavDataOverride())
 				{
 					if (MainNavData == NavDataInstanceForAgent)
 					{
@@ -1967,6 +1965,7 @@ void UNavigationSystemV1::UnregisterNavData(ANavigationData* NavData)
     AgentToNavDataMap.Remove(NavData->GetConfig());
 
 	FScopeLock Lock(&NavDataRegistration);
+	NavDataRegistrationQueue.Remove(NavData);
 	NavData->OnUnregistered();
 }
 
@@ -2301,13 +2300,12 @@ void UNavigationSystemV1::InitializeForWorld(UWorld& World, FNavigationSystemRun
 	OnWorldInitDone(Mode);
 }
 
-void UNavigationSystemV1::OnNavSystemOverriden(UNavigationSystemBase* PreviousNavSystem)
+void UNavigationSystemV1::FinalizeOverridingNavSystem(UNavigationSystemV1& PreviousNavSystem)
 {
     // When a Navigation Override Config loads it may happen that we have some nav data queued to be registered, keep them
-	const UNavigationSystemV1* PrevNavSystemV1 = Cast<UNavigationSystemV1>(PreviousNavSystem);
-	if (PreviousNavSystem && PrevNavSystemV1->NavDataRegistrationQueue.Num() > 0)
+	if (PreviousNavSystem.NavDataRegistrationQueue.Num() > 0)
 	{
-		NavDataRegistrationQueue = PrevNavSystemV1->NavDataRegistrationQueue;		
+		NavDataRegistrationQueue = PreviousNavSystem.NavDataRegistrationQueue;
 	}
 }
 
@@ -2910,7 +2908,7 @@ void UNavigationSystemV1::PerformNavigationBoundsUpdate(const TArray<FNavigation
 		}
 	}
 
-	if (!IsNavigationBuildingLocked() || ShouldUseNavDataInAdditionalLevelWhenDuplicatedAgent())
+	if (!IsNavigationBuildingLocked())
 	{
 		if (UpdatedAreas.Num())
 		{
@@ -2922,7 +2920,7 @@ void UNavigationSystemV1::PerformNavigationBoundsUpdate(const TArray<FNavigation
 				}
 			}
 		}
-				
+
 		// Propagate to generators areas that needs to be updated
 		AddDirtyAreas(UpdatedAreas, ENavigationDirtyFlag::All | ENavigationDirtyFlag::NavigationBounds);
 	}
@@ -3242,8 +3240,7 @@ void UNavigationSystemV1::RebuildAll(bool bIsLoadTime)
 
 	DefaultDirtyAreasController.Reset();
 
-	// Going backward as if bGenerateNavDataWhenNoCompatibleNavBound is false, we will destroy the nav data in RebuildAll
-	for (int32 NavDataIndex = NavDataSet.Num() - 1; NavDataIndex >= 0 ; --NavDataIndex)
+	for (int32 NavDataIndex = 0; NavDataIndex < NavDataSet.Num(); ++NavDataIndex)
 	{
 		ANavigationData* NavData = NavDataSet[NavDataIndex];
 				
@@ -3330,18 +3327,20 @@ void UNavigationSystemV1::OnLevelAddedToWorld(ULevel* InLevel, UWorld* InWorld)
 
 		if (!InLevel->IsPersistentLevel())
 		{
-			if (ShouldUseNavDataInAdditionalLevelWhenDuplicatedAgent())
+#if WITH_EDITOR
+            // Registering NavigationData in its PostLoad() isn't enough because when adding/removing level in editor, it calls Remove/Load to refresh streaming levels
+			if (GIsEditor)
 			{
 				for (TActorIterator<ANavigationData> It(InWorld); It; ++It)
 				{
 					ANavigationData* NavData = *It;
 					if (NavData->GetLevel() == InLevel)
 					{
-						RegisterNavData(NavData);
+						RequestRegistration(NavData);
 					}
 				}
 			}
-
+#endif
 			for (ANavigationData* NavData : NavDataSet)
 			{
 				if (NavData)
@@ -3371,20 +3370,21 @@ void UNavigationSystemV1::OnLevelRemovedFromWorld(ULevel* InLevel, UWorld* InWor
 						NavData->OnStreamingLevelRemoved(InLevel, InWorld);
 					}
 					else
-					{			
+					{
                         if (NavData->IsRegistered())
                         {
                             UnregisterNavData(NavData);
 
-                            if (ShouldUseNavDataInAdditionalLevelWhenDuplicatedAgent())
+                            if (ShouldAllowEquivalentNavDataOverride())
                             {
                                 for (TActorIterator<ANavigationData> It(InWorld); It; ++It)
                                 {
                                     ANavigationData* CurrentNavData = *It;
                                     const FNavDataConfig& TestConfig = NavData->GetConfig();
-                                    if (CurrentNavData->GetLevel() != InLevel && GetTypeHash(CurrentNavData->GetConfig()) == GetTypeHash(NavData->GetConfig()))
+                                    if (!CurrentNavData->IsRegistered() && CurrentNavData->GetLevel() != InLevel && GetTypeHash(CurrentNavData->GetConfig()) == GetTypeHash(NavData->GetConfig()))
                                     {
                                         RequestRegistration(CurrentNavData);
+										break;
                                     }
                                 }
                             }
@@ -3521,9 +3521,18 @@ void UNavigationSystemV1::CleanUp(FNavigationSystem::ECleanupMode Mode)
 			UnregisterNavData(NavData);
 		}
 	}	
-	NavDataSet.Reset();
-
-	AgentToNavDataMap.Reset();
+	if (NavDataSet.Num())
+	{
+        UE_LOG(LogNavigation, Error, TEXT("UNavigationSystemV1::CleanUp still have an agent in NavDataSet loaded after unregister them all!!"));
+        NavDataSet.Reset();
+    }
+	
+	if (AgentToNavDataMap.Num())
+	{
+		UE_LOG(LogNavigation, Error, TEXT("UNavigationSystemV1::CleanUp still have an agent mapped to navigation data after clean up!"));
+		AgentToNavDataMap.Reset();
+	}
+	
 	MainNavData = nullptr;
 
 	// reset unique link Id for new map
@@ -4253,8 +4262,8 @@ void UNavigationSystemModuleConfig::UpdateWithNavSysCDO(const UNavigationSystemV
 		bCreateOnClient = NavSysCDO.bAllowClientSideNavigation;
 		bAutoSpawnMissingNavData = NavSysCDO.bAutoCreateNavigationData;
 		bSpawnNavDataInNavBoundsLevel = NavSysCDO.bSpawnNavDataInNavBoundsLevel;
-		bGenerateNavDataWhenNoCompatibleNavBound = NavSysCDO.bGenerateNavDataWhenNoCompatibleNavBound;
-		bUseNavDataInAdditionalLevelWhenDuplicatedAgent = NavSysCDO.bUseNavDataInAdditionalLevelWhenDuplicatedAgent;
+		bShouldGenerateNavDataOnlyWhenCompatibleNavBound = NavSysCDO.bShouldGenerateNavDataOnlyWhenCompatibleNavBound;
+		bAllowEquivalentNavDataOverride = NavSysCDO.bAllowEquivalentNavDataOverride;
 	}
 }
 
@@ -4276,8 +4285,8 @@ UNavigationSystemBase* UNavigationSystemModuleConfig::CreateAndConfigureNavigati
 	{
 		NavSysInstance->bAutoCreateNavigationData = bAutoSpawnMissingNavData;
 		NavSysInstance->bSpawnNavDataInNavBoundsLevel = bSpawnNavDataInNavBoundsLevel;
-		NavSysInstance->bGenerateNavDataWhenNoCompatibleNavBound = bGenerateNavDataWhenNoCompatibleNavBound;
-		NavSysInstance->bUseNavDataInAdditionalLevelWhenDuplicatedAgent = bUseNavDataInAdditionalLevelWhenDuplicatedAgent;
+		NavSysInstance->bShouldGenerateNavDataOnlyWhenCompatibleNavBound = bShouldGenerateNavDataOnlyWhenCompatibleNavBound;
+		NavSysInstance->bAllowEquivalentNavDataOverride  = bAllowEquivalentNavDataOverride;
 		if (bStrictlyStatic)
 		{
 			NavSysInstance->ConfigureAsStatic();
diff --git a/Engine/Source/Runtime/NavigationSystem/Public/NavigationSystem.h b/Engine/Source/Runtime/NavigationSystem/Public/NavigationSystem.h
--- a/Engine/Source/Runtime/NavigationSystem/Public/NavigationSystem.h
+++ b/Engine/Source/Runtime/NavigationSystem/Public/NavigationSystem.h
@@ -142,13 +142,13 @@ protected:
 	UPROPERTY()
 	uint32 bSupportRebuilding : 1; 
 
-	/** When set to TRUE, navigation data will be generated even if no navigation bounds is compatible to build it */
+	/** When set to TRUE, navigation data will be generated only when at least one navigation bounds is compatible for building it */
 	UPROPERTY(EditAnywhere, Category = Navigation)
-	uint32 bGenerateNavDataWhenNoCompatibleNavBound : 1;
+	uint32 bShouldGenerateNavDataOnlyWhenCompatibleNavBound : 1;
 
-	/** When set to TRUE, if there conflicting navigation data using same NavConfig, the last one loaded (via the additional level) will be used when rebuild navmesh */
+	/** When set to TRUE, equivalent navigation data will override the existing one in the system */
 	UPROPERTY(EditAnywhere, Category = Navigation)
-	uint32 bUseNavDataInAdditionalLevelWhenDuplicatedAgent : 1;
+	uint32 bAllowEquivalentNavDataOverride  : 1;
 public:
 	/** if set to true will result navigation system not rebuild navigation until 
 	 *	a call to ReleaseInitialBuildingLock() is called. Does not influence 
@@ -494,8 +494,8 @@ public:
 	bool ShouldAllowClientSideNavigation() const { return bAllowClientSideNavigation; }
 	virtual bool ShouldLoadNavigationOnClient(ANavigationData* NavData = nullptr) const { return bAllowClientSideNavigation; }
 	virtual bool ShouldDiscardSubLevelNavData(ANavigationData* NavData = nullptr) const { return bShouldDiscardSubLevelNavData; }
-	bool ShouldGenerateNavDataWhenNoCompatibleNavBound() const { return bGenerateNavDataWhenNoCompatibleNavBound; }
-	bool ShouldUseNavDataInAdditionalLevelWhenDuplicatedAgent() const { return bUseNavDataInAdditionalLevelWhenDuplicatedAgent; }
+	bool ShouldGenerateNavDataOnlyWhenCompatibleNavBound() const { return bShouldGenerateNavDataOnlyWhenCompatibleNavBound; }
+	bool ShouldAllowEquivalentNavDataOverride() const { return bAllowEquivalentNavDataOverride ; }
 
 	FBox GetWorldBounds() const;
 	
@@ -544,7 +544,7 @@ public:
 	/** adds NavData to registration candidates queue - NavDataRegistrationQueue
 	 *	@return true if registration request was successful, false if given NavData 
 	 *	was deemed unsuitable for registration consideration */
-	virtual void RequestRegistration(ANavigationData* NavData, bool bTriggerRegistrationProcessing = true);
+	virtual void RequestRegistration(ANavigationData* NavData);
 
 protected:
 
@@ -801,7 +801,7 @@ public:
 	virtual void InitializeForWorld(UWorld& World, FNavigationSystemRunMode Mode) override;
 
 	/** Called after a NavSystemConfigOverride has finished overriding the old NavSystem */
-	virtual void OnNavSystemOverriden(UNavigationSystemBase* PreviousNavSystem) override;
+	virtual void FinalizeOverridingNavSystem(UNavigationSystemV1& PreviousNavSystem);
 
 	// Fetch the array of all nav-agent properties.
 	void GetNavAgentPropertiesArray(TArray<FNavAgentProperties>& OutNavAgentProperties) const;
@@ -1087,13 +1087,13 @@ protected:
 	UPROPERTY(EditAnywhere, Category = Navigation)
 	uint32 bSpawnNavDataInNavBoundsLevel : 1;
 
-	/** When set to TRUE, navigation data will be generated even if no navigation bounds is compatible to build it */
+	/** When set to TRUE, navigation data will be generated only when at least one navigation bounds is compatible for building it */
 	UPROPERTY(EditAnywhere, Category = Navigation)
-	uint32 bGenerateNavDataWhenNoCompatibleNavBound : 1;
+	uint32 bShouldGenerateNavDataOnlyWhenCompatibleNavBound : 1;
 
-	/** When set to TRUE, if there conflicting navigation data using same NavConfig, the last one loaded (via the additional level) will be used when rebuild navmesh */
+	/** When set to TRUE, equivalent navigation data will override the existing one in the system */
 	UPROPERTY(EditAnywhere, Category = Navigation)
-	uint32 bUseNavDataInAdditionalLevelWhenDuplicatedAgent : 1;
+	uint32 bAllowEquivalentNavDataOverride : 1;
 public:
 	UNavigationSystemModuleConfig(const FObjectInitializer& ObjectInitializer = FObjectInitializer::Get());
 
