Deprecating RequestRegistration and adding RequestRegistrationForNextTick since the optional parameter is no more used.

diff --git a/Engine/Source/Runtime/NavigationSystem/Private/NavigationSystem.cpp b/Engine/Source/Runtime/NavigationSystem/Private/NavigationSystem.cpp
--- a/Engine/Source/Runtime/NavigationSystem/Private/NavigationSystem.cpp
+++ b/Engine/Source/Runtime/NavigationSystem/Private/NavigationSystem.cpp
@@ -873,7 +873,7 @@ void UNavigationSystemV1::RegisterNavigationDataInstances()
 		ANavigationData* Nav = (*It);
 		if (Nav != NULL && Nav->IsPendingKill() == false && Nav->IsRegistered() == false)
 		{
-			RequestRegistration(Nav);
+			RequestRegistrationDeferred(*Nav);
 			bProcessRegistration = true;
 		}
 	}
@@ -1751,13 +1751,13 @@ void UNavigationSystemV1::ApplyWorldOffset(const FVector& InOffset, bool bWorldS
 //----------------------------------------------------------------------//
 // Bookkeeping 
 //----------------------------------------------------------------------//
-void UNavigationSystemV1::RequestRegistration(ANavigationData* NavData)
+void UNavigationSystemV1::RequestRegistrationDeferred(ANavigationData& NavData)
 {
 	FScopeLock RegistrationLock(&NavDataRegistrationSection);
 
 	if (NavDataRegistrationQueue.Num() < REGISTRATION_QUEUE_SIZE)
 	{
-		NavDataRegistrationQueue.AddUnique(NavData);
+		NavDataRegistrationQueue.AddUnique(&NavData);
 	}
 	else
 	{
@@ -3107,7 +3107,7 @@ void UNavigationSystemV1::SpawnMissingNavigationData()
 				ANavigationData* Instance = CreateNavigationDataInstanceInLevel(NavConfig, nullptr);
 				if (Instance)
 				{
-					RequestRegistration(Instance);
+					RequestRegistrationDeferred(*Instance);
 				}
 				else
 				{
@@ -3336,7 +3336,7 @@ void UNavigationSystemV1::OnLevelAddedToWorld(ULevel* InLevel, UWorld* InWorld)
 					ANavigationData* NavData = *It;
 					if (NavData->GetLevel() == InLevel)
 					{
-						RequestRegistration(NavData);
+						RequestRegistrationDeferred(*NavData);
 					}
 				}
 			}
@@ -3383,7 +3383,7 @@ void UNavigationSystemV1::OnLevelRemovedFromWorld(ULevel* InLevel, UWorld* InWor
                                     const FNavDataConfig& TestConfig = NavData->GetConfig();
                                     if (!CurrentNavData->IsRegistered() && CurrentNavData->GetLevel() != InLevel && GetTypeHash(CurrentNavData->GetConfig()) == GetTypeHash(NavData->GetConfig()))
                                     {
-                                        RequestRegistration(CurrentNavData);
+                                        RequestRegistrationDeferred(*CurrentNavData);
 										break;
                                     }
                                 }
@@ -4053,6 +4053,14 @@ void UNavigationSystemV1::UnregisterNavigationInvoker(AActor* Invoker)
 //----------------------------------------------------------------------//
 // DEPRECATED
 //----------------------------------------------------------------------//
+void UNavigationSystemV1::RequestRegistration(ANavigationData* NavData, bool bTriggerRegistrationProcessing)
+{
+	if (NavData)
+	{
+		RequestRegistrationDeferred(*NavData);
+	}	
+}
+
 FVector UNavigationSystemV1::ProjectPointToNavigation(UObject* WorldContextObject, const FVector& Point, ANavigationData* NavData, TSubclassOf<UNavigationQueryFilter> FilterClass, const FVector QueryExtent)
 {
 	FNavLocation ProjectedPoint(Point);
diff --git a/Engine/Source/Runtime/NavigationSystem/Public/NavigationSystem.h b/Engine/Source/Runtime/NavigationSystem/Public/NavigationSystem.h
--- a/Engine/Source/Runtime/NavigationSystem/Public/NavigationSystem.h
+++ b/Engine/Source/Runtime/NavigationSystem/Public/NavigationSystem.h
@@ -541,10 +541,10 @@ public:
 	// @todo document
 	virtual void UnregisterNavData(ANavigationData* NavData);
 
-	/** adds NavData to registration candidates queue - NavDataRegistrationQueue
-	 *	@return true if registration request was successful, false if given NavData 
-	 *	was deemed unsuitable for registration consideration */
-	virtual void RequestRegistration(ANavigationData* NavData);
+	/** Adds NavData to registration candidates queue - NavDataRegistrationQueue*/
+	virtual void RequestRegistrationDeferred(ANavigationData& NavData);
+	UE_DEPRECATED(4.24, "This version of RequestRegistration is deprecated. Please use the RequestRegistrationDeferred as the registration request is always queued now.")
+	virtual void RequestRegistration(ANavigationData* NavData, bool bTriggerRegistrationProcessing = true);
 
 protected:
 
