Fixed FPathFindingQuery constructors to use the Navigation Data's config as NavAgentConfig instead of defaulting it to a default one. This fixes issue when a path is being recalculated (RepathRequests) it was not using the same NavAgent config that was originally requested with.

diff --git a/Engine/Source/Runtime/NavigationSystem/Private/NavigationData.cpp b/Engine/Source/Runtime/NavigationSystem/Private/NavigationData.cpp
--- a/Engine/Source/Runtime/NavigationSystem/Private/NavigationData.cpp
+++ b/Engine/Source/Runtime/NavigationSystem/Private/NavigationData.cpp
@@ -21,7 +21,7 @@
 //----------------------------------------------------------------------//
 FPathFindingQuery::FPathFindingQuery(const UObject* InOwner, const ANavigationData& InNavData, const FVector& Start, const FVector& End, FSharedConstNavQueryFilter SourceQueryFilter, FNavPathSharedPtr InPathInstanceToFill) :
 	FPathFindingQueryData(InOwner, Start, End, SourceQueryFilter), 
-	NavData(&InNavData), PathInstanceToFill(InPathInstanceToFill), NavAgentProperties(FNavAgentProperties::DefaultProperties)
+	NavData(&InNavData), PathInstanceToFill(InPathInstanceToFill), NavAgentProperties(InNavData.GetConfig())
 {
 	if (!QueryFilter.IsValid() && NavData.IsValid())
 	{
@@ -71,9 +71,14 @@ FPathFindingQuery::FPathFindingQuery(FNavPathSharedRef PathToRecalculate, const
 		}
 	}
 
-	if (!QueryFilter.IsValid() && NavData.IsValid())
+	if (NavData.IsValid())
 	{
-		QueryFilter = NavData->GetDefaultQueryFilter();
+		if (!QueryFilter.IsValid())
+		{
+			QueryFilter = NavData->GetDefaultQueryFilter();
+		}
+
+		NavAgentProperties = NavData->GetConfig();
 	}
 }
 
@@ -319,7 +324,6 @@ void ANavigationData::TickActor(float DeltaTime, enum ELevelTick TickType, FActo
 			}
 
 			FPathFindingQuery Query(PinnedPath.ToSharedRef());
-			// @todo consider supplying NavAgentPropertied from path's querier
 			const FPathFindingResult Result = FindPath(Query.NavAgentProperties, Query.SetPathInstanceToUpdate(PinnedPath));
 
 			// update time stamp to give observers any means of telling if it has changed
