Forgotten file to previous deprecation method code in NavigationSystem

#rnx

diff --git a/Engine/Source/Runtime/NavigationSystem/Private/NavigationData.cpp b/Engine/Source/Runtime/NavigationSystem/Private/NavigationData.cpp
--- a/Engine/Source/Runtime/NavigationSystem/Private/NavigationData.cpp
+++ b/Engine/Source/Runtime/NavigationSystem/Private/NavigationData.cpp
@@ -220,7 +220,7 @@ void ANavigationData::RequestRegistration()
 		UNavigationSystemV1* NavSys = FNavigationSystem::GetCurrent<UNavigationSystemV1>(GetWorld());
 		if (NavSys)
 		{
-			NavSys->RequestRegistration(this);
+			NavSys->RequestRegistrationDeferred(*this);
 		}
 	}
 }
@@ -427,7 +427,7 @@ void ANavigationData::PostEditUndo()
 		}
 		else
 		{
-			NavSys->RequestRegistration(this);
+			NavSys->RequestRegistrationDeferred(*this);
 		}
 	}
 }
