Dynamic Obstacle support for Instanced Static Mesh or Hierarchical Instanced Static Mesh Component

FRecastTileGenerator & FRecastNavMeshGenerator
* Better support for transformation of per instance Convex shapes in MarkDynamicArea (replacing PartialTransformConvexHull)
* Added Per Instance modifier support in GrabDebugSnapshot for VisualLogger (both Geomerty export & AreaNavModifier)
* Added Per Instance modifier support in ExportNavigationData (both Geomerty export & AreaNavModifier)

FCompositeNavModifier
* Deprecated PerInstanceTransformDelegate code since it is unused and duplicated functionality handled by FNavigationRelevantData
* Added flag to mark as a PerInstance modifier

Fixed NavigationDataDump when default FBox is used (since GetVolume will be computed from uninitialized values)

diff --git a/Engine/Source/Runtime/Engine/Private/AI/Navigation/NavigationModifier.cpp b/Engine/Source/Runtime/Engine/Private/AI/Navigation/NavigationModifier.cpp
--- a/Engine/Source/Runtime/Engine/Private/AI/Navigation/NavigationModifier.cpp
+++ b/Engine/Source/Runtime/Engine/Private/AI/Navigation/NavigationModifier.cpp
@@ -326,6 +326,22 @@ FAreaNavModifier::FAreaNavModifier(const TNavStatArray<FVector>& InPoints, const
 	SetConvex(InPoints.GetData(), FirstIndex, LastIndex, CoordType, LocalToWorld);
 }
 
+void FAreaNavModifier::InitializeConvex(const TNavStatArray<FVector>& InPoints, const int32 FirstIndex, const int32 LastIndex, const FTransform& LocalToWorld, const TSubclassOf<UNavAreaBase> InAreaClass)
+{
+	check(InPoints.IsValidIndex(FirstIndex) && InPoints.IsValidIndex(LastIndex-1));
+
+	Init(InAreaClass);
+	SetConvex(InPoints.GetData(), FirstIndex, LastIndex, ENavigationCoordSystem::Unreal, LocalToWorld);
+}
+
+void FAreaNavModifier::InitializePerInstanceConvex(const TNavStatArray<FVector>& InPoints, const int32 FirstIndex, const int32 LastIndex, const TSubclassOf<UNavAreaBase> InAreaClass)
+{
+	check(InPoints.IsValidIndex(FirstIndex) && InPoints.IsValidIndex(LastIndex - 1));
+
+	Init(InAreaClass);
+	SetPerInstanceConvex(InPoints.GetData(), FirstIndex, LastIndex);
+}
+
 FAreaNavModifier::FAreaNavModifier(const UBrushComponent* BrushComponent, const TSubclassOf<UNavAreaBase> InAreaClass)
 {
 	check(BrushComponent != NULL);
@@ -371,6 +387,14 @@ void FAreaNavModifier::GetConvex(FConvexNavAreaData& Data) const
 	Data.MaxZ = LastPoint.Y;
 }
 
+void FAreaNavModifier::GetPerInstanceConvex(const FTransform& InLocalToWorld, FConvexNavAreaData& OutConvexData) const
+{
+	ensure(ShapeType == ENavigationShapeType::InstancedConvex);
+
+	FBox TmpBounds;
+	FillConvexNavAreaData(Points.GetData(), Points.Num(), InLocalToWorld, OutConvexData, TmpBounds);
+}
+
 void FAreaNavModifier::Init(const TSubclassOf<UNavAreaBase> InAreaClass)
 {
 	bIncludeAgentHeight = false;
@@ -458,30 +482,33 @@ void FAreaNavModifier::SetBox(const FBox& Box, const FTransform& LocalToWorld)
 	}
 }
 
-void FAreaNavModifier::SetConvex(const FVector* InPoints, const int32 FirstIndex, const int32 LastIndex, ENavigationCoordSystem::Type CoordType, const FTransform& LocalToWorld)
+void FAreaNavModifier::FillConvexNavAreaData(const FVector* InPoints, const int32 InNumPoints, const FTransform& InTotalTransform, FConvexNavAreaData& OutConvexData, FBox& OutBounds)
 {
-	FConvexNavAreaData ConvexData;
-	ConvexData.MinZ = MAX_FLT;
-	ConvexData.MaxZ = -MAX_FLT;
+	OutBounds = FBox(ForceInit);
+	OutConvexData.Points.Reset();
+	OutConvexData.MinZ = MAX_FLT;
+	OutConvexData.MaxZ = -MAX_FLT;
+
+	if (InNumPoints <= 0)
+	{
+		return;
+	}
 
 	const int MaxConvexPoints = 8;
 	TArray<FVector, TInlineAllocator<MaxConvexPoints>> HullVertices;
 	HullVertices.Empty(MaxConvexPoints);
 
-	// we store verts in unreal coords
-	const FTransform& TotalTransform = LocalToWorld * FNavigationSystem::GetCoordTransform(CoordType, ENavigationCoordSystem::Unreal);
-
-	for (int32 i = FirstIndex; i < LastIndex; i++)
+	for (int32 i = 0; i < InNumPoints; i++)
 	{
-		FVector TransformedPoint = TotalTransform.TransformPosition(InPoints[i]);
-		ConvexData.MinZ = FMath::Min( ConvexData.MinZ, TransformedPoint.Z );
-		ConvexData.MaxZ = FMath::Max( ConvexData.MaxZ, TransformedPoint.Z );
+		FVector TransformedPoint = InTotalTransform.TransformPosition(InPoints[i]);
+		OutConvexData.MinZ = FMath::Min(OutConvexData.MinZ, TransformedPoint.Z);
+		OutConvexData.MaxZ = FMath::Max(OutConvexData.MaxZ, TransformedPoint.Z);
 		TransformedPoint.Z = 0.f;
 
 		// check if there's a similar point already in HullVertices array
 		bool bUnique = true;
 		const FVector* RESTRICT Start = HullVertices.GetData();
-		for (const FVector* RESTRICT Data = Start, * RESTRICT DataEnd = Data + HullVertices.Num(); Data != DataEnd; ++Data)
+		for (const FVector* RESTRICT Data = Start, *RESTRICT DataEnd = Data + HullVertices.Num(); Data != DataEnd; ++Data)
 		{
 			if (FVector::DistSquared(*Data, TransformedPoint) < CONVEX_HULL_POINTS_MIN_DISTANCE_SQ)
 			{
@@ -489,7 +516,7 @@ void FAreaNavModifier::SetConvex(const FVector* InPoints, const int32 FirstIndex
 				break;
 			}
 		}
-		
+
 		if (bUnique)
 		{
 			HullVertices.Add(TransformedPoint);
@@ -500,19 +527,51 @@ void FAreaNavModifier::SetConvex(const FVector* InPoints, const int32 FirstIndex
 	HullIndices.Empty(MaxConvexPoints);
 
 	ConvexHull2D::ComputeConvexHull(HullVertices, HullIndices);
-	if (HullIndices.Num())
+	
+	// ConvexHull implementation requires at least 3 vertices  (i.e. GrowConvexHull)
+	const int32 MIN_NUM_POINTS = 3;
+
+	if (HullIndices.Num() >= MIN_NUM_POINTS)
 	{
-		Bounds = FBox(ForceInit);
-		for(int32 i = 0; i < HullIndices.Num(); ++i)
+		for (int32 i = 0; i < HullIndices.Num(); ++i)
 		{
 			const FVector& HullVert = HullVertices[HullIndices[i]];
-			ConvexData.Points.Add(HullVert);
-			Bounds += HullVert;
+			OutConvexData.Points.Add(HullVert);
+			OutBounds += HullVert;
+		}
+
+		OutBounds.Min.Z = OutConvexData.MinZ;
+		OutBounds.Max.Z = OutConvexData.MaxZ;
+	}
+}
+
+void FAreaNavModifier::SetPerInstanceConvex(const FVector* InPoints, const int32 InFirstIndex, const int32 InLastIndex)
+{
+	// Per Instance modifiers requires that we keep all unique points until we receive the instance transform for
+    // ConvexHull to be computed. Local Bounds must be computed right away.
+	Bounds = FBox(ForceInit);
+	for (int32 i = InFirstIndex; i < InLastIndex; i++)
+	{
+		const FVector& CurrentPoint = InPoints[i];
+		FVector* SamePoint = Points.FindByPredicate([&CurrentPoint](FVector& Point) { return FMath::IsNearlyZero(FVector::DistSquared(Point, CurrentPoint)); });
+		if (SamePoint == nullptr)
+		{
+			Points.Add(CurrentPoint);
+			Bounds += CurrentPoint;
 		}
+	}
+	ShapeType = ENavigationShapeType::InstancedConvex;
+}
 
-		Bounds.Min.Z = ConvexData.MinZ;
-		Bounds.Max.Z = ConvexData.MaxZ;
+void FAreaNavModifier::SetConvex(const FVector* InPoints, const int32 FirstIndex, const int32 LastIndex, ENavigationCoordSystem::Type CoordType, const FTransform& LocalToWorld)
+{
+	const FTransform& TotalTransform = LocalToWorld * FNavigationSystem::GetCoordTransform(CoordType, ENavigationCoordSystem::Unreal);
 
+	FConvexNavAreaData ConvexData;
+	FillConvexNavAreaData(InPoints + FirstIndex, LastIndex - FirstIndex, TotalTransform, ConvexData, Bounds);
+
+	if (ConvexData.Points.Num() > 0)
+	{
 		Points.Append(ConvexData.Points);
 		Points.Add(FVector(ConvexData.MinZ, ConvexData.MaxZ, 0));
 		ShapeType = ENavigationShapeType::Convex;
@@ -669,6 +728,7 @@ void FCompositeNavModifier::Reset()
 	CustomLinks.Reset();
 	bHasPotentialLinks = false;
 	bAdjustHeight = false;
+	bIsPerInstanceModifier = false;
 }
 
 void FCompositeNavModifier::Empty()
@@ -844,5 +904,7 @@ uint32 FCompositeNavModifier::GetAllocatedSize() const
 
 bool FCompositeNavModifier::HasPerInstanceTransforms() const
 {
+	PRAGMA_DISABLE_DEPRECATION_WARNINGS
 	return NavDataPerInstanceTransformDelegate.IsBound();
+	PRAGMA_ENABLE_DEPRECATION_WARNINGS
 }
diff --git a/Engine/Source/Runtime/Engine/Private/InstancedStaticMesh.cpp b/Engine/Source/Runtime/Engine/Private/InstancedStaticMesh.cpp
--- a/Engine/Source/Runtime/Engine/Private/InstancedStaticMesh.cpp
+++ b/Engine/Source/Runtime/Engine/Private/InstancedStaticMesh.cpp
@@ -2398,6 +2398,7 @@ void UInstancedStaticMeshComponent::GetNavigationData(FNavigationRelevantData& D
 		UNavCollisionBase* NavCollision = GetStaticMesh()->NavCollision;
 		if (NavCollision->IsDynamicObstacle())
 		{
+			Data.Modifiers.MarkAsPerInstanceModifier();
 			NavCollision->GetNavigationModifier(Data.Modifiers, FTransform::Identity);
 
 			// Hook per instance transform delegate
diff --git a/Engine/Source/Runtime/Engine/Public/AI/NavigationModifier.h b/Engine/Source/Runtime/Engine/Public/AI/NavigationModifier.h
--- a/Engine/Source/Runtime/Engine/Public/AI/NavigationModifier.h
+++ b/Engine/Source/Runtime/Engine/Public/AI/NavigationModifier.h
@@ -31,6 +31,7 @@ namespace ENavigationShapeType
 		Cylinder,
 		Box,
 		Convex,
+		InstancedConvex,
 	};
 }
 
@@ -103,6 +104,9 @@ struct ENGINE_API FAreaNavModifier : public FNavigationModifier
 	FAreaNavModifier(const TNavStatArray<FVector>& Points, const int32 FirstIndex, const int32 LastIndex, ENavigationCoordSystem::Type CoordType, const FTransform& LocalToWorld, const TSubclassOf<UNavAreaBase> AreaClass);
 	FAreaNavModifier(const UBrushComponent* BrushComponent, const TSubclassOf<UNavAreaBase> AreaClass);
 
+	void InitializePerInstanceConvex(const TNavStatArray<FVector>& Points, const int32 FirstIndex, const int32 LastIndex, const TSubclassOf<UNavAreaBase> AreaClass);
+	void InitializeConvex(const TNavStatArray<FVector>& Points, const int32 FirstIndex, const int32 LastIndex, const FTransform& LocalToWorld, const TSubclassOf<UNavAreaBase> AreaClass);
+
 	FORCEINLINE const FBox& GetBounds() const { return Bounds; }
 	FORCEINLINE ENavigationShapeType::Type GetShapeType() const { return ShapeType; }
 	FORCEINLINE ENavigationAreaMode::Type GetApplyMode() const { return ApplyMode; }
@@ -124,6 +128,7 @@ struct ENGINE_API FAreaNavModifier : public FNavigationModifier
 	void GetCylinder(FCylinderNavAreaData& Data) const;
 	void GetBox(FBoxNavAreaData& Data) const;
 	void GetConvex(FConvexNavAreaData& Data) const;
+	void GetPerInstanceConvex(const FTransform& InLocalToWorld, FConvexNavAreaData& OutData) const;
 
 protected:
 	/** this should take a value of a game specific navigation modifier	*/
@@ -144,7 +149,10 @@ protected:
 	void Init(const TSubclassOf<UNavAreaBase> InAreaClass);
 	/** @param CoordType specifies which coord system the input data is in */
 	void SetConvex(const FVector* InPoints, const int32 FirstIndex, const int32 LastIndex, ENavigationCoordSystem::Type CoordType, const FTransform& LocalToWorld);
+	void SetPerInstanceConvex(const FVector* InPoints, const int32 InFirstIndex, const int32 InLastIndex);
 	void SetBox(const FBox& Box, const FTransform& LocalToWorld);
+	
+	static void FillConvexNavAreaData(const FVector* InPoints, const int32 InNumPoints, const FTransform& InLocalToWorld, FConvexNavAreaData& OutConvexData, FBox& OutBounds);
 };
 
 /**
@@ -241,7 +249,7 @@ protected:
 
 struct ENGINE_API FCompositeNavModifier : public FNavigationModifier
 {
-	FCompositeNavModifier() : bHasPotentialLinks(false), bAdjustHeight(false), bHasLowAreaModifiers(false) {}
+	FCompositeNavModifier() : bHasPotentialLinks(false), bAdjustHeight(false), bHasLowAreaModifiers(false), bIsPerInstanceModifier(false) {}
 
 	void Shrink();
 	void Reset();
@@ -293,24 +301,27 @@ struct ENGINE_API FCompositeNavModifier : public FNavigationModifier
 	FORCEINLINE bool HasAgentHeightAdjust() const { return bAdjustHeight; }
 	FORCEINLINE bool HasAreas() const { return Areas.Num() > 0; }
 	FORCEINLINE bool HasLowAreaModifiers() const { return bHasLowAreaModifiers; }
+    FORCEINLINE bool IsPerInstanceModifier() const { return bIsPerInstanceModifier; }
 
 	FORCEINLINE void ReserveForAdditionalAreas(int32 AdditionalElementsCount) { Areas.Reserve(Areas.Num() + AdditionalElementsCount); }
 
 	void MarkPotentialLinks() { bHasPotentialLinks = true; }
+    void MarkAsPerInstanceModifier() { bIsPerInstanceModifier = true; }
 
 	/** returns a copy of Modifier */
 	FCompositeNavModifier GetInstantiatedMetaModifier(const struct FNavAgentProperties* NavAgent, TWeakObjectPtr<UObject> WeakOwnerPtr) const;
 	uint32 GetAllocatedSize() const;
 
+	UE_DEPRECATED(4.23, "This method will be removed in future versions. Use FNavigationRelevantData::HasPerInstanceTransforms instead.")
 	bool HasPerInstanceTransforms() const;
-	// Should be called only on game thread
-	void GetPerInstanceTransforms(const FBox& AreaBox, TArray<FTransform>& PerInstanceTransforms) const;
 
 	TArray<FAreaNavModifier>& GetMutableAreas() { return Areas; }
 	TArray<FSimpleLinkNavModifier>& GetSimpleLinks() { return SimpleLinks; }
 	TArray<FCustomLinkNavModifier>& GetCustomLinks() { return CustomLinks; }
 
 public:
+
+	UE_DEPRECATED(4.23, "This property will be removed in future versions. Use FNavigationRelevantData NavDataPerInstanceTransformDelegate instead.")
 	// Gathers per instance data for navigation area modifiers in a specified area box
 	FNavDataPerInstanceTransformDelegate NavDataPerInstanceTransformDelegate;
 
@@ -321,4 +332,5 @@ private:
 	uint32 bHasPotentialLinks : 1;
 	uint32 bAdjustHeight : 1;
 	uint32 bHasLowAreaModifiers : 1;
+    uint32 bIsPerInstanceModifier : 1;
 };
diff --git a/Engine/Source/Runtime/NavigationSystem/Private/NavCollision.cpp b/Engine/Source/Runtime/NavigationSystem/Private/NavCollision.cpp
--- a/Engine/Source/Runtime/NavigationSystem/Private/NavCollision.cpp
+++ b/Engine/Source/Runtime/NavigationSystem/Private/NavCollision.cpp
@@ -295,6 +295,21 @@ void UNavCollision::GetNavigationModifier(FCompositeNavModifier& Modifier, const
 	const int32 NumModifiers = (TriMeshCollision.VertexBuffer.Num() ? 1 : 0) + ConvexShapeIndices.Num();
 	Modifier.ReserveForAdditionalAreas(NumModifiers);
 
+	auto AddModFunc = [&](const TNavStatArray<FVector>& VertexBuffer, const int32 FirstVertIndex, const int32 LastVertIndex)
+	{
+		FAreaNavModifier AreaMod;
+		if (Modifier.IsPerInstanceModifier())
+		{
+			AreaMod.InitializePerInstanceConvex(VertexBuffer, FirstVertIndex, LastVertIndex, UseAreaClass);
+		}
+		else
+		{
+			AreaMod.InitializeConvex(VertexBuffer, FirstVertIndex, LastVertIndex, LocalToWorld, UseAreaClass);
+		}
+		AreaMod.SetIncludeAgentHeight(true);
+		Modifier.Add(AreaMod);
+	};
+
 	int32 LastVertIndex = 0;
 	for (int32 Idx = 0; Idx < ConvexShapeIndices.Num(); Idx++)
 	{
@@ -305,17 +320,13 @@ void UNavCollision::GetNavigationModifier(FCompositeNavModifier& Modifier, const
 		// contain any duplicates (which is the original cause of UE-52123)
 		if (FirstVertIndex < LastVertIndex)
 		{
-			FAreaNavModifier AreaMod(ConvexCollision.VertexBuffer, FirstVertIndex, LastVertIndex, ENavigationCoordSystem::Unreal, LocalToWorld, UseAreaClass);
-			AreaMod.SetIncludeAgentHeight(true);
-			Modifier.Add(AreaMod);
+			AddModFunc(ConvexCollision.VertexBuffer, FirstVertIndex, LastVertIndex);
 		}
 	}
 
 	if (TriMeshCollision.VertexBuffer.Num() > 0)
 	{
-		FAreaNavModifier AreaMod(TriMeshCollision.VertexBuffer, 0, TriMeshCollision.VertexBuffer.Num() - 1, ENavigationCoordSystem::Unreal, LocalToWorld, UseAreaClass);
-		AreaMod.SetIncludeAgentHeight(true);
-		Modifier.Add(AreaMod);
+		AddModFunc(TriMeshCollision.VertexBuffer, 0, TriMeshCollision.VertexBuffer.Num() - 1);
 	}
 }
 
diff --git a/Engine/Source/Runtime/NavigationSystem/Private/NavMesh/RecastNavMeshGenerator.cpp b/Engine/Source/Runtime/NavigationSystem/Private/NavMesh/RecastNavMeshGenerator.cpp
--- a/Engine/Source/Runtime/NavigationSystem/Private/NavMesh/RecastNavMeshGenerator.cpp
+++ b/Engine/Source/Runtime/NavigationSystem/Private/NavMesh/RecastNavMeshGenerator.cpp
@@ -1112,23 +1112,6 @@ void FRecastGeometryExport::SetNavDataPerInstanceTransformDelegate(const FNavDat
 	Data->NavDataPerInstanceTransformDelegate = InDelegate;
 }
 
-static void PartialTransformConvexHull(FConvexNavAreaData& ConvexData, const FTransform& LocalToWorld)
-{
-	FVector ScaleXY = LocalToWorld.GetScale3D().GetAbs();
-	ScaleXY.Z = 1.f;
-
-	FVector TranslationXY = LocalToWorld.GetLocation();
-	TranslationXY.Z = 0.f;
-
-	for (FVector& Point : ConvexData.Points)
-	{
-		Point = Point*ScaleXY + TranslationXY;
-	}
-
-	ConvexData.MaxZ+= LocalToWorld.GetLocation().Z;
-	ConvexData.MinZ+= LocalToWorld.GetLocation().Z;
-}
-
 FORCEINLINE void GrowConvexHull(const float ExpandBy, const TArray<FVector>& Verts, TArray<FVector>& OutResult)
 {
 	if (Verts.Num() < 3)
@@ -3090,11 +3073,17 @@ void FRecastTileGenerator::MarkDynamicArea(const FAreaNavModifier& Modifier, con
 		break;
 
 	case ENavigationShapeType::Convex:
+	case ENavigationShapeType::InstancedConvex:
 		{
 			FConvexNavAreaData ConvexData;
-			Modifier.GetConvex(ConvexData);
-			// Only scaling and translation
-			PartialTransformConvexHull(ConvexData, LocalToWorld);
+			if (Modifier.GetShapeType() == ENavigationShapeType::InstancedConvex)
+			{
+				Modifier.GetPerInstanceConvex(LocalToWorld, ConvexData);
+			} 
+			else
+			{
+				Modifier.GetConvex(ConvexData);
+			}
 
 			TArray<FVector> ConvexVerts;
 			GrowConvexHull(ExpandBy, ConvexData.Points, ConvexVerts);
@@ -4845,17 +4834,13 @@ void FRecastNavMeshGenerator::GrabDebugSnapshot(struct FVisualLogEntry* Snapshot
 				const FNavigationOctreeElement& Element = It.GetCurrentElement();
 				const bool bExportGeometry = Element.Data->HasGeometry() && Element.ShouldUseGeometry(DestNavMesh->GetConfig());
 
+				TArray<FTransform> InstanceTransforms;
+				Element.Data->NavDataPerInstanceTransformDelegate.ExecuteIfBound(Element.Bounds.GetBox(), InstanceTransforms);
+
 				if (bExportGeometry && Element.Data->CollisionData.Num())
 				{
 					FRecastGeometryCache CachedGeometry(Element.Data->CollisionData.GetData());
 					
-					const uint32 NumVerts = CachedGeometry.Header.NumVerts;
-					CoordBuffer.Reset(NumVerts);
-					for (uint32 VertIdx = 0; VertIdx < NumVerts * 3; VertIdx += 3)
-					{
-						CoordBuffer.Add(Recast2UnrealPoint(&CachedGeometry.Verts[VertIdx]));
-					}
-
 					const uint32 NumIndices = CachedGeometry.Header.NumFaces * 3;
 					Indices.SetNum(NumIndices, false);
 					for (uint32 IndicesIdx = 0; IndicesIdx < NumIndices; ++IndicesIdx)
@@ -4863,50 +4848,89 @@ void FRecastNavMeshGenerator::GrabDebugSnapshot(struct FVisualLogEntry* Snapshot
 						Indices[IndicesIdx] = CachedGeometry.Indices[IndicesIdx];
 					}
 
-					Snapshot->AddElement(CoordBuffer, Indices, CategoryName, LogVerbosity, FColorList::LightGrey.WithAlpha(255));
+					auto AddElementFunc = [&](const FTransform& Transform)
+					{
+						const uint32 NumVerts = CachedGeometry.Header.NumVerts;
+						CoordBuffer.Reset(NumVerts);
+						for (uint32 VertIdx = 0; VertIdx < NumVerts * 3; VertIdx += 3)
+						{
+							CoordBuffer.Add(Transform.TransformPosition(Recast2UnrealPoint(&CachedGeometry.Verts[VertIdx])));
+						}
+
+						Snapshot->AddElement(CoordBuffer, Indices, CategoryName, LogVerbosity, FColorList::LightGrey.WithAlpha(255));
+					};
+
+					if (InstanceTransforms.Num() == 0)
+					{
+						AddElementFunc(FTransform::Identity);
+					}
+					for (const FTransform& InstanceTransform : InstanceTransforms)
+					{
+						AddElementFunc(InstanceTransform);
+					}
 				}
 				else
 				{
 					TArray<FVector> Verts;
-					const TArray<FAreaNavModifier>& AreaMods = Element.Data->Modifiers.GetAreas();
-					for (int32 i = 0; i < AreaMods.Num(); i++)
+					for (const FAreaNavModifier& AreaMod : Element.Data->Modifiers.GetAreas())
 					{
-						if (AreaMods[i].GetShapeType() == ENavigationShapeType::Unknown)
+						ENavigationShapeType::Type ShapeType = AreaMod.GetShapeType();
+						if (ShapeType == ENavigationShapeType::Unknown)
 						{
 							continue;
 						}
 
-						const uint8 AreaId = NavData->GetAreaID(AreaMods[i].GetAreaClass());
+						const uint8 AreaId = NavData->GetAreaID(AreaMod.GetAreaClass());
 						const UClass* AreaClass = NavData->GetAreaClass(AreaId);
 						const UNavArea* DefArea = AreaClass ? ((UClass*)AreaClass)->GetDefaultObject<UNavArea>() : NULL;
 						const FColor PolygonColor = AreaClass != FNavigationSystem::GetDefaultWalkableArea() ? (DefArea ? DefArea->DrawColor : NavData->GetConfig().Color) : FColorList::Cyan;
 
-						if (AreaMods[i].GetShapeType() == ENavigationShapeType::Box)
+						if (ShapeType == ENavigationShapeType::Box)
 						{
 							FBoxNavAreaData Box;
-							AreaMods[i].GetBox(Box);
+							AreaMod.GetBox(Box);
 
 							Snapshot->AddElement(FBox::BuildAABB(Box.Origin, Box.Extent), FMatrix::Identity, CategoryName, NavAreaVerbosity, PolygonColor.WithAlpha(255));
 						}
-						else if (AreaMods[i].GetShapeType() == ENavigationShapeType::Cylinder)
+						else if (ShapeType == ENavigationShapeType::Cylinder)
 						{
 							FCylinderNavAreaData Cylinder;
-							AreaMods[i].GetCylinder(Cylinder);
+							AreaMod.GetCylinder(Cylinder);
 
 							Snapshot->AddElement(Cylinder.Origin, Cylinder.Origin + FVector(0, 0, Cylinder.Height), Cylinder.Radius, CategoryName, NavAreaVerbosity, PolygonColor.WithAlpha(255));
 						}
-						else
+						else if (ShapeType == ENavigationShapeType::Convex || ShapeType == ENavigationShapeType::InstancedConvex)
 						{
-							FConvexNavAreaData Convex;
-							AreaMods[i].GetConvex(Convex);
-							Verts.Reset();
-							GrowConvexHull(NavData->AgentRadius, Convex.Points, Verts);
-
-							Snapshot->AddElement(
-								Verts,
-								Convex.MinZ - NavData->CellHeight, 
-								Convex.MaxZ + NavData->CellHeight, 
-								CategoryName, NavAreaVerbosity, PolygonColor.WithAlpha(255));
+							auto AddElementFunc = [&](const FConvexNavAreaData& InConvexNavAreaData)
+							{
+								Verts.Reset();
+								GrowConvexHull(NavData->AgentRadius, InConvexNavAreaData.Points, Verts);
+
+								if (Verts.Num())
+								{
+									Snapshot->AddElement(
+										Verts,
+										InConvexNavAreaData.MinZ - NavData->CellHeight,
+										InConvexNavAreaData.MaxZ + NavData->CellHeight,
+										CategoryName, NavAreaVerbosity, PolygonColor.WithAlpha(255));
+								}
+							};
+
+							if (ShapeType == ENavigationShapeType::Convex)
+							{
+								FConvexNavAreaData Convex;
+								AreaMod.GetConvex(Convex);
+								AddElementFunc(Convex);
+							}
+							else // ShapeType == ENavigationShapeType::InstancedConvex
+							{
+								for (const FTransform& InstanceTransform : InstanceTransforms)
+								{
+									FConvexNavAreaData Convex;
+									AreaMod.GetPerInstanceConvex(InstanceTransform, Convex);
+									AddElementFunc(Convex);
+								}
+							}
 						}
 					}
 				}
@@ -4954,39 +4978,85 @@ void FRecastNavMeshGenerator::ExportNavigationData(const FString& FileName) cons
 				const FNavigationOctreeElement& Element = It.GetCurrentElement();
 				const bool bExportGeometry = Element.Data->HasGeometry() && Element.ShouldUseGeometry(DestNavMesh->GetConfig());
 
+				TArray<FTransform> InstanceTransforms;
+				Element.Data->NavDataPerInstanceTransformDelegate.ExecuteIfBound(Element.Bounds.GetBox(), InstanceTransforms);
+				
 				if (bExportGeometry && Element.Data->CollisionData.Num())
 				{
+					const int32 NumInstances = FMath::Max(InstanceTransforms.Num(), 1);
 					FRecastGeometryCache CachedGeometry(Element.Data->CollisionData.GetData());
-					IndexBuffer.Reserve( IndexBuffer.Num() + (CachedGeometry.Header.NumFaces * 3 ));
-					CoordBuffer.Reserve( CoordBuffer.Num() + (CachedGeometry.Header.NumVerts * 3 ));
-					for (int32 i = 0; i < CachedGeometry.Header.NumFaces * 3; i++)
+					IndexBuffer.Reserve( IndexBuffer.Num() + (CachedGeometry.Header.NumFaces * 3 ) * NumInstances );
+					CoordBuffer.Reserve( CoordBuffer.Num() + (CachedGeometry.Header.NumVerts * 3 ) * NumInstances );
+
+					if (InstanceTransforms.Num() == 0)
 					{
-						IndexBuffer.Add(CachedGeometry.Indices[i] + CoordBuffer.Num() / 3);
+						for (int32 i = 0; i < CachedGeometry.Header.NumFaces * 3; i++)
+						{
+							IndexBuffer.Add(CachedGeometry.Indices[i] + CoordBuffer.Num() / 3);
+						}
+						for (int32 i = 0; i < CachedGeometry.Header.NumVerts * 3; i++)
+						{
+							CoordBuffer.Add(CachedGeometry.Verts[i]);
+						}
 					}
-					for (int32 i = 0; i < CachedGeometry.Header.NumVerts * 3; i++)
+					for (const FTransform& InstanceTransform : InstanceTransforms)
 					{
-						CoordBuffer.Add(CachedGeometry.Verts[i]);
+						for (int32 i = 0; i < CachedGeometry.Header.NumFaces * 3; i++)
+						{
+							IndexBuffer.Add(CachedGeometry.Indices[i] + CoordBuffer.Num() / 3);
+						}
+
+						FMatrix LocalToRecastWorld = InstanceTransform.ToMatrixWithScale()*Unreal2RecastMatrix();
+
+						for (int32 i = 0; i < CachedGeometry.Header.NumVerts * 3; i += 3)
+						{
+							// collision cache stores coordinates in recast space, convert them to unreal and transform to recast world space
+							FVector WorldRecastCoord = LocalToRecastWorld.TransformPosition(Recast2UnrealPoint(&CachedGeometry.Verts[i]));
+
+							CoordBuffer.Add(WorldRecastCoord.X);
+							CoordBuffer.Add(WorldRecastCoord.Y);
+							CoordBuffer.Add(WorldRecastCoord.Z);
+						}
 					}
 				}
 				else
 				{
-					const TArray<FAreaNavModifier>& AreaMods = Element.Data->Modifiers.GetAreas();
-					for (int32 i = 0; i < AreaMods.Num(); i++)
+					for (const FAreaNavModifier& AreaMod : Element.Data->Modifiers.GetAreas())
 					{
-						FAreaExportData ExportInfo;
-						ExportInfo.AreaId = NavData->GetAreaID(AreaMods[i].GetAreaClass());
-
-						if (AreaMods[i].GetShapeType() == ENavigationShapeType::Convex)
+						ENavigationShapeType::Type ShapeType = AreaMod.GetShapeType();
+						
+						if (ShapeType == ENavigationShapeType::Convex || ShapeType == ENavigationShapeType::InstancedConvex)
 						{
-							AreaMods[i].GetConvex(ExportInfo.Convex);
+							FAreaExportData ExportInfo;
+							ExportInfo.AreaId = NavData->GetAreaID(AreaMod.GetAreaClass());
 
-							TArray<FVector> ConvexVerts;
-							GrowConvexHull(NavData->AgentRadius, ExportInfo.Convex.Points, ConvexVerts);
-							ExportInfo.Convex.MinZ -= NavData->CellHeight;
-							ExportInfo.Convex.MaxZ += NavData->CellHeight;
-							ExportInfo.Convex.Points = ConvexVerts;
+							auto AddAreaExportDataFunc = [&](const FConvexNavAreaData& InConvexNavAreaData)
+							{
+								TArray<FVector> ConvexVerts;
+								GrowConvexHull(NavData->AgentRadius, ExportInfo.Convex.Points, ConvexVerts);
+								if (ConvexVerts.Num())
+								{
+									ExportInfo.Convex.MinZ -= NavData->CellHeight;
+									ExportInfo.Convex.MaxZ += NavData->CellHeight;
+									ExportInfo.Convex.Points = ConvexVerts;
 
-							AreaExport.Add(ExportInfo);
+									AreaExport.Add(ExportInfo);
+								}								
+							};
+
+							if (ShapeType == ENavigationShapeType::Convex)
+							{
+								AreaMod.GetConvex(ExportInfo.Convex);
+								AddAreaExportDataFunc(ExportInfo.Convex);
+							}
+							else // ShapeType == ENavigationShapeType::InstancedConvex
+							{
+								for (const FTransform& InstanceTransform : InstanceTransforms)
+								{
+									AreaMod.GetPerInstanceConvex(InstanceTransform, ExportInfo.Convex);
+									AddAreaExportDataFunc(ExportInfo.Convex);
+								}
+							}
 						}
 					}
 				}
diff --git a/Engine/Source/Runtime/NavigationSystem/Private/NavigationSystem.cpp b/Engine/Source/Runtime/NavigationSystem/Private/NavigationSystem.cpp
--- a/Engine/Source/Runtime/NavigationSystem/Private/NavigationSystem.cpp
+++ b/Engine/Source/Runtime/NavigationSystem/Private/NavigationSystem.cpp
@@ -218,7 +218,7 @@ namespace
 		const FNavDataGenerator* Generator = MainNavData ? MainNavData->GetGenerator() : nullptr;
 		if (Generator)
 		{
-			Generator->GrabDebugSnapshot(&CurrentEntry, FMath::IsNearlyZero(Box.GetVolume()) ? MainNavData->GetBounds().ExpandBy(FVector(20, 20, 20)) : Box, CategoryName, Verbosity);
+			Generator->GrabDebugSnapshot(&CurrentEntry, (!Box.IsValid || FMath::IsNearlyZero(Box.GetVolume())) ? MainNavData->GetBounds().ExpandBy(FVector(20, 20, 20)) : Box, CategoryName, Verbosity);
 		}
 	}
 #endif // ENABLE_VISUAL_LOG
