Partial rollback of CL6604191.

RecastNavMeshGenerator initialization should be called right after object construction to ensure virtual methods are properly called and derived classes are fully constructed before rebuild navigation mesh.

#rnx

diff --git a/Engine/Source/Runtime/NavigationSystem/Private/NavMesh/RecastNavMesh.cpp b/Engine/Source/Runtime/NavigationSystem/Private/NavMesh/RecastNavMesh.cpp
--- a/Engine/Source/Runtime/NavigationSystem/Private/NavMesh/RecastNavMesh.cpp
+++ b/Engine/Source/Runtime/NavigationSystem/Private/NavMesh/RecastNavMesh.cpp
@@ -2418,6 +2418,7 @@ void ARecastNavMesh::ConditionalConstructGenerator()
 		if (Generator)
 		{
 			NavDataGenerator = MakeShareable(Generator);
+			Generator->Init();
 		}
 
 		UNavigationSystemV1* NavSys = FNavigationSystem::GetCurrent<UNavigationSystemV1>(World);
diff --git a/Engine/Source/Runtime/NavigationSystem/Private/NavMesh/RecastNavMeshGenerator.cpp b/Engine/Source/Runtime/NavigationSystem/Private/NavMesh/RecastNavMeshGenerator.cpp
--- a/Engine/Source/Runtime/NavigationSystem/Private/NavMesh/RecastNavMeshGenerator.cpp
+++ b/Engine/Source/Runtime/NavigationSystem/Private/NavMesh/RecastNavMeshGenerator.cpp
@@ -3220,79 +3220,6 @@
 #endif
 
 	INC_DWORD_STAT_BY(STAT_NavigationMemory, sizeof(*this));
-
-	Init();
-
-	int32 MaxTiles = 0;
-	int32 MaxPolysPerTile = 0;
-
-	// recreate navmesh if no data was loaded, or when loaded data doesn't match current grid layout
-	bool bRecreateNavmesh = true;
-	if (DestNavMesh->HasValidNavmesh())
-	{
-		const bool bGameStaticNavMesh = IsGameStaticNavMesh(DestNavMesh);
-		const dtNavMeshParams* SavedNavParams = DestNavMesh->GetRecastNavMeshImpl()->DetourNavMesh->getParams();
-		if (SavedNavParams)
-		{
-			if (bGameStaticNavMesh)
-			{
-				bRecreateNavmesh = false;
-				MaxTiles = SavedNavParams->maxTiles;
-				MaxPolysPerTile = SavedNavParams->maxPolys;
-			}
-			else
-			{
-				const float TileDim = Config.tileSize * Config.cs;
-				if (SavedNavParams->tileHeight == TileDim && SavedNavParams->tileWidth == TileDim)
-				{
-					const FVector Orig = Recast2UnrealPoint(SavedNavParams->orig);
-					const FVector OrigError(FMath::Fmod(Orig.X, TileDim), FMath::Fmod(Orig.Y, TileDim), FMath::Fmod(Orig.Z, TileDim));
-					if (OrigError.IsNearlyZero())
-					{
-						bRecreateNavmesh = false;
-					}
-					else
-					{
-						UE_LOG(LogNavigation, Warning, TEXT("Recreating dtNavMesh instance due to saved navmesh origin (%s, usually the RecastNavMesh location) not being aligned with tile size (%d uu) ")
-							, *Orig.ToString(), int(TileDim));
-					}
-				}
-
-				// if new navmesh needs more tiles, force recreating
-				if (!bRecreateNavmesh)
-				{
-					CalcNavMeshProperties(MaxTiles, MaxPolysPerTile);
-					if (FMath::Log2(MaxTiles) != FMath::Log2(SavedNavParams->maxTiles))
-					{
-						bRecreateNavmesh = true;
-						UE_LOG(LogNavigation, Warning, TEXT("Recreating dtNavMesh instance due mismatch in number of bytes required to store serialized maxTiles (%d, %d bits) vs calculated maxtiles (%d, %d bits)")
-							, SavedNavParams->maxTiles, FMath::CeilToInt(FMath::Log2(SavedNavParams->maxTiles))
-							, MaxTiles, FMath::CeilToInt(FMath::Log2(MaxTiles)));
-					}
-				}
-			}
-		};
-	}
-
-	if (bRecreateNavmesh)
-	{
-		// recreate navmesh from scratch if no data was loaded
-		ConstructTiledNavMesh();
-
-		// mark all the areas we need to update, which is the whole (known) navigable space if not restricted to active tiles
-		const UNavigationSystemV1* NavSys = FNavigationSystem::GetCurrent<UNavigationSystemV1>(GetWorld());
-		if (NavSys)
-		{
-			bRestrictBuildingToActiveTiles = NavSys->IsActiveTilesGenerationEnabled();
-		}
-		MarkNavBoundsDirty();
-	}
-	else
-	{
-		// otherwise just update generator params
-		Config.MaxPolysPerTile = MaxPolysPerTile;
-		NumActiveTiles = GetTilesCountHelper(DestNavMesh->GetRecastNavMeshImpl()->DetourNavMesh);
-	}
 }
 
 FRecastNavMeshGenerator::~FRecastNavMeshGenerator()
@@ -3393,6 +3320,78 @@
 	}
 
 	bInitialized = true;
+
+
+	int32 MaxTiles = 0;
+	int32 MaxPolysPerTile = 0;
+
+	// recreate navmesh if no data was loaded, or when loaded data doesn't match current grid layout
+	bool bRecreateNavmesh = true;
+	if (DestNavMesh->HasValidNavmesh())
+	{
+		const bool bGameStaticNavMesh = IsGameStaticNavMesh(DestNavMesh);
+		const dtNavMeshParams* SavedNavParams = DestNavMesh->GetRecastNavMeshImpl()->DetourNavMesh->getParams();
+		if (SavedNavParams)
+		{
+			if (bGameStaticNavMesh)
+			{
+				bRecreateNavmesh = false;
+				MaxTiles = SavedNavParams->maxTiles;
+				MaxPolysPerTile = SavedNavParams->maxPolys;
+			}
+			else
+			{
+				const float TileDim = Config.tileSize * Config.cs;
+				if (SavedNavParams->tileHeight == TileDim && SavedNavParams->tileWidth == TileDim)
+				{
+					const FVector Orig = Recast2UnrealPoint(SavedNavParams->orig);
+					const FVector OrigError(FMath::Fmod(Orig.X, TileDim), FMath::Fmod(Orig.Y, TileDim), FMath::Fmod(Orig.Z, TileDim));
+					if (OrigError.IsNearlyZero())
+					{
+						bRecreateNavmesh = false;
+					}
+					else
+					{
+						UE_LOG(LogNavigation, Warning, TEXT("Recreating dtNavMesh instance due to saved navmesh origin (%s, usually the RecastNavMesh location) not being aligned with tile size (%d uu) ")
+							, *Orig.ToString(), int(TileDim));
+					}
+				}
+
+				// if new navmesh needs more tiles, force recreating
+				if (!bRecreateNavmesh)
+				{
+					CalcNavMeshProperties(MaxTiles, MaxPolysPerTile);
+					if (FMath::Log2(MaxTiles) != FMath::Log2(SavedNavParams->maxTiles))
+					{
+						bRecreateNavmesh = true;
+						UE_LOG(LogNavigation, Warning, TEXT("Recreating dtNavMesh instance due mismatch in number of bytes required to store serialized maxTiles (%d, %d bits) vs calculated maxtiles (%d, %d bits)")
+							, SavedNavParams->maxTiles, FMath::CeilToInt(FMath::Log2(SavedNavParams->maxTiles))
+							, MaxTiles, FMath::CeilToInt(FMath::Log2(MaxTiles)));
+					}
+				}
+			}
+		};
+	}
+
+	if (bRecreateNavmesh)
+	{
+		// recreate navmesh from scratch if no data was loaded
+		ConstructTiledNavMesh();
+
+		// mark all the areas we need to update, which is the whole (known) navigable space if not restricted to active tiles
+		const UNavigationSystemV1* NavSys = FNavigationSystem::GetCurrent<UNavigationSystemV1>(GetWorld());
+		if (NavSys)
+		{
+			bRestrictBuildingToActiveTiles = NavSys->IsActiveTilesGenerationEnabled();
+		}
+		MarkNavBoundsDirty();
+	}
+	else
+	{
+		// otherwise just update generator params
+		Config.MaxPolysPerTile = MaxPolysPerTile;
+		NumActiveTiles = GetTilesCountHelper(DestNavMesh->GetRecastNavMeshImpl()->DetourNavMesh);
+	}
 }
 
 void FRecastNavMeshGenerator::UpdateNavigationBounds()
