NavSystem: Prevent AbstractNavData from being rejected during NavigationData registration. Problem was exposed by recent change to the order in which NavDataRegistrationQueue is processed. Since AbstractNavData uses the default config, it was rejected since another navigation data was already registered for the same agent config  (which was not the case when AbstractNavData was processed first). With this change, we should no longer be dependant on the registration order.

diff --git a/Engine/Source/Runtime/NavigationSystem/Private/NavigationSystem.cpp b/Engine/Source/Runtime/NavigationSystem/Private/NavigationSystem.cpp
--- a/Engine/Source/Runtime/NavigationSystem/Private/NavigationSystem.cpp
+++ b/Engine/Source/Runtime/NavigationSystem/Private/NavigationSystem.cpp
@@ -1874,13 +1874,31 @@ UNavigationSystemV1::ERegistrationResult UNavigationSystemV1::RegisterNavData(AN
 
 	if (NavConfig.IsValid() == true)
 	{
-		// check if this kind of agent has already its navigation implemented
-		TWeakObjectPtr<ANavigationData>* NavDataForAgent = AgentToNavDataMap.Find(NavConfig);
-		ANavigationData* NavDataInstanceForAgent = NavDataForAgent ? NavDataForAgent->Get() : nullptr;
+		if (NavData->IsA(AAbstractNavData::StaticClass()))
+		{
+			if (AbstractNavData == nullptr || AbstractNavData == NavData)
+			{
+				// fake registration since it's a special navigation data type 
+				// and it would get discarded for not implementing any particular
+				// navigation agent
+				// Node that we don't add abstract navigation data to NavDataSet
+				NavData->OnRegistered();
 
-		if (NavDataInstanceForAgent == nullptr || ShouldAllowEquivalentNavDataOverride())
+				Result = RegistrationSuccessful;
+			}
+			else
+			{
+				// otherwise specified agent type already has its navmesh implemented, fail redundant instance
+				Result = RegistrationFailed_AgentAlreadySupported;
+			}
+		}
+		else
 		{
-			if (NavData->IsA(AAbstractNavData::StaticClass()) == false)
+			// check if this kind of agent has already its navigation implemented
+			TWeakObjectPtr<ANavigationData>* NavDataForAgent = AgentToNavDataMap.Find(NavConfig);
+			ANavigationData* NavDataInstanceForAgent = NavDataForAgent ? NavDataForAgent->Get() : nullptr;
+
+			if (NavDataInstanceForAgent == nullptr || ShouldAllowEquivalentNavDataOverride())
 			{
 				if (NavDataInstanceForAgent != nullptr && ShouldAllowEquivalentNavDataOverride())
 				{
@@ -1918,27 +1936,17 @@ UNavigationSystemV1::ERegistrationResult UNavigationSystemV1::RegisterNavData(AN
 				}
 				Result = bAgentSupported == true ? RegistrationSuccessful : RegistrationFailed_AgentNotValid;
 			}
-			else
+			else if (NavDataInstanceForAgent == NavData)
 			{
-				// fake registration since it's a special navigation data type 
-				// and it would get discarded for not implementing any particular
-				// navigation agent
-				// Node that we don't add abstract navigation data to NavDataSet
-				NavData->OnRegistered();
-
+				ensure(NavDataSet.Find(NavData) != INDEX_NONE);
+				// let's treat double registration of the same nav data with the same agent as a success
 				Result = RegistrationSuccessful;
 			}
-		}
-		else if (NavDataInstanceForAgent == NavData)
-		{
-			ensure(NavDataSet.Find(NavData) != INDEX_NONE);
-			// let's treat double registration of the same nav data with the same agent as a success
-			Result = RegistrationSuccessful;
-		}
-		else
-		{
-			// otherwise specified agent type already has its navmesh implemented, fail redundant instance
-			Result = RegistrationFailed_AgentAlreadySupported;
+			else
+			{
+				// otherwise specified agent type already has its navmesh implemented, fail redundant instance
+				Result = RegistrationFailed_AgentAlreadySupported;
+			}
 		}
 	}
 	else
