Added an option to RecastNavMeshGenerator that makes sorting of tiles to be generated optional.

Yoan.StAmant, Maxime.Mercier

diff --git a/Engine/Source/Runtime/NavigationSystem/Private/NavMesh/RecastNavMeshGenerator.cpp b/Engine/Source/Runtime/NavigationSystem/Private/NavMesh/RecastNavMeshGenerator.cpp
--- a/Engine/Source/Runtime/NavigationSystem/Private/NavMesh/RecastNavMeshGenerator.cpp
+++ b/Engine/Source/Runtime/NavigationSystem/Private/NavMesh/RecastNavMeshGenerator.cpp
@@ -45,6 +45,8 @@
 
 #define TEXT_WEAKOBJ_NAME(obj) (obj.IsValid(false) ? *obj->GetName() : (obj.IsValid(false, true)) ? TEXT("MT-Unreachable") : TEXT("INVALID"))
 
+PRAGMA_DISABLE_OPTIMIZATION
+
 struct dtTileCacheAlloc;
 
 FORCEINLINE bool DoesBoxContainOrOverlapVector(const FBox& BigBox, const FVector& In)
@@ -3213,6 +3215,7 @@ FRecastNavMeshGenerator::FRecastNavMeshGenerator(ARecastNavMesh& InDestNavMesh)
 	, DestNavMesh(&InDestNavMesh)
 	, bInitialized(false)
 	, bRestrictBuildingToActiveTiles(false)
+	, bSortTilesWithSeedLocations(true)
 	, Version(0)
 {
 #if TIME_SLICE_NAV_REGEN
@@ -4296,6 +4299,11 @@ void FRecastNavMeshGenerator::MarkDirtyTiles(const TArray<FNavigationDirtyArea>&
 
 void FRecastNavMeshGenerator::SortPendingBuildTiles()
 {
+	if (bSortTilesWithSeedLocations == false)
+	{
+		return;
+	}
+
 	TArray<FVector2D> SeedLocations;
 	UWorld* CurWorld = GetWorld();
 	if (CurWorld == nullptr)
@@ -4317,7 +4325,7 @@ void FRecastNavMeshGenerator::SortPendingBuildTiles()
 	if (SeedLocations.Num() == 0)
 	{
 		// Use navmesh origin for sorting
-		SeedLocations.Add(FVector2D::ZeroVector);
+		SeedLocations.Add(FVector2D(TotalNavBounds.GetCenter()));
 	}
 
 	if (SeedLocations.Num() > 0)
@@ -4331,11 +4339,7 @@ void FRecastNavMeshGenerator::SortPendingBuildTiles()
 			FVector2D TileCenter2D = FVector2D(TileBox.GetCenter());
 			for (FVector2D SeedLocation : SeedLocations)
 			{
-				const float DistSq = FVector2D::DistSquared(TileCenter2D, SeedLocation);
-				if (DistSq < Element.SeedDistance)
-				{
-					Element.SeedDistance = DistSq;
-				}
+				Element.SeedDistance = FMath::Min(Element.SeedDistance, FVector2D::DistSquared(TileCenter2D, SeedLocation));
 			}
 		}
 
diff --git a/Engine/Source/Runtime/NavigationSystem/Public/NavMesh/RecastNavMeshGenerator.h b/Engine/Source/Runtime/NavigationSystem/Public/NavMesh/RecastNavMeshGenerator.h
--- a/Engine/Source/Runtime/NavigationSystem/Public/NavMesh/RecastNavMeshGenerator.h
+++ b/Engine/Source/Runtime/NavigationSystem/Public/NavMesh/RecastNavMeshGenerator.h
@@ -631,6 +631,8 @@ protected:
 
 	uint32 bRestrictBuildingToActiveTiles:1;
 
+	uint32 bSortTilesWithSeedLocations:1;
+
 	/** Runtime generator's version, increased every time all tile generators get invalidated
 	 *	like when navmesh size changes */
 	uint32 Version;
