Removed PRAGMA_DISABLE_OPTIMIZATION left behind accidentally

diff --git a/Engine/Source/Runtime/NavigationSystem/Private/NavMesh/RecastNavMeshGenerator.cpp b/Engine/Source/Runtime/NavigationSystem/Private/NavMesh/RecastNavMeshGenerator.cpp
--- a/Engine/Source/Runtime/NavigationSystem/Private/NavMesh/RecastNavMeshGenerator.cpp
+++ b/Engine/Source/Runtime/NavigationSystem/Private/NavMesh/RecastNavMeshGenerator.cpp
@@ -45,8 +45,6 @@
 
 #define TEXT_WEAKOBJ_NAME(obj) (obj.IsValid(false) ? *obj->GetName() : (obj.IsValid(false, true)) ? TEXT("MT-Unreachable") : TEXT("INVALID"))
 
-PRAGMA_DISABLE_OPTIMIZATION
-
 struct dtTileCacheAlloc;
 
 FORCEINLINE bool DoesBoxContainOrOverlapVector(const FBox& BigBox, const FVector& In)
