Reorganized how we configure SupportedAgents. From now on every NavigationSystem instance has data on the full set of SupportedAgents (as configured via ProjectSettings) but also has a SupportedAgentsMask, that marks agents supported by the instance. It has been done this way to keep filtered per-agent NavMeshBoundVolumes work regardless of which supported agents a nav sys instance wants to support.

This approach allowed for easy NavigationSystem's config update with NavSystemConfigOverride's properties. NavSystemConfigOverride can now specify if it wants to fully override the pre-existing nav sys instance or if it just wants to append new information (like supported agents) to the existing navigation system instance. There's also an option to do nothing if there's already a navigation system present.

This CL rolls-back a bunch of temp fixes done in past couple of days.

Yoan.StAmant

diff --git a/Engine/Source/Runtime/Engine/Classes/AI/NavigationSystemBase.h b/Engine/Source/Runtime/Engine/Classes/AI/NavigationSystemBase.h
--- a/Engine/Source/Runtime/Engine/Classes/AI/NavigationSystemBase.h
+++ b/Engine/Source/Runtime/Engine/Classes/AI/NavigationSystemBase.h
@@ -95,7 +95,7 @@ namespace FNavigationSystem
 	 *	@param NavigationSystemConfig is used to pick the navigation system's class and set it up. If null
 	 *		then WorldOwner.WorldSettings.NavigationSystemConfig will be used
 	 */
-	ENGINE_API void AddNavigationSystemToWorld(UWorld& WorldOwner, const FNavigationSystemRunMode RunMode = FNavigationSystemRunMode::InvalidMode, UNavigationSystemConfig* NavigationSystemConfig = nullptr, const bool bInitializeForWorld = true);
+	ENGINE_API void AddNavigationSystemToWorld(UWorld& WorldOwner, const FNavigationSystemRunMode RunMode = FNavigationSystemRunMode::InvalidMode, UNavigationSystemConfig* NavigationSystemConfig = nullptr, const bool bInitializeForWorld = true, const bool bOverridePreviousNavSys = false);
 
 	/** Discards all navigation data chunks in all sub-levels */
 	ENGINE_API void DiscardNavigationDataChunks(UWorld& InWorld);
@@ -212,6 +212,8 @@ public:
 	virtual void Tick(float DeltaSeconds) PURE_VIRTUAL(UNavigationSystemBase::Tick, );
 	virtual void CleanUp(const FNavigationSystem::ECleanupMode Mode) PURE_VIRTUAL(UNavigationSystemBase::CleanUp, );
 	virtual void Configure(const UNavigationSystemConfig& Config) PURE_VIRTUAL(UNavigationSystemBase::Configure, );
+	/** Called when there's a need to extend current navigation system's config with information in NewConfig */
+	virtual void AppendConfig(const UNavigationSystemConfig& NewConfig) PURE_VIRTUAL(UNavigationSystemBase::AppendConfig, );
 
 	/**
 	*	Called when owner-UWorld initializes actors
diff --git a/Engine/Source/Runtime/Engine/Classes/AI/NavigationSystemConfig.h b/Engine/Source/Runtime/Engine/Classes/AI/NavigationSystemConfig.h
--- a/Engine/Source/Runtime/Engine/Classes/AI/NavigationSystemConfig.h
+++ b/Engine/Source/Runtime/Engine/Classes/AI/NavigationSystemConfig.h
@@ -20,6 +20,18 @@ public:
 	UPROPERTY(EditAnywhere, Category=Navigation, meta = (MetaClass = "NavigationSystemBase", NoResetToDefault))
 	FSoftClassPath NavigationSystemClass;
 
+	/** NavigationSystem's properties in Project Settings define all possible supported agents,
+	 *	but a specific navigation system can choose to support only a subset of agents.*/
+	UPROPERTY(EditAnywhere, Category = Navigation)
+	FNavAgentSelector SupportedAgentsMask;
+
+	/** If not None indicates which of navigation datas and supported agents are 
+	 * going to be used as the default ones. If navigation agent of this type does 
+	 * not exist or is not enabled then the first available nav data will be used 
+	 * as the default one */
+	UPROPERTY(EditAnywhere, Category = Navigation)
+	FName DefaultAgentName;
+
 protected:
 	/** If true it means the navigation system settings are overridden from another source (like a NavConfigOverrideActor) */
 	UPROPERTY(VisibleAnywhere, Category = Navigation)
diff --git a/Engine/Source/Runtime/Engine/Private/AI/NavigationSystemBase.cpp b/Engine/Source/Runtime/Engine/Private/AI/NavigationSystemBase.cpp
--- a/Engine/Source/Runtime/Engine/Private/AI/NavigationSystemBase.cpp
+++ b/Engine/Source/Runtime/Engine/Private/AI/NavigationSystemBase.cpp
@@ -36,10 +36,9 @@ namespace FNavigationSystem
 		}
 	}
 
-
-	void AddNavigationSystemToWorld(UWorld& WorldOwner, const FNavigationSystemRunMode RunMode, UNavigationSystemConfig* NavigationSystemConfig, const bool bInitializeForWorld)
+	void AddNavigationSystemToWorld(UWorld& WorldOwner, const FNavigationSystemRunMode RunMode, UNavigationSystemConfig* NavigationSystemConfig, const bool bInitializeForWorld, const bool bOverridePreviousNavSys)
 	{
-		if (WorldOwner.GetNavigationSystem() == nullptr)
+		if (WorldOwner.GetNavigationSystem() == nullptr || bOverridePreviousNavSys)
 		{
 			if (NavigationSystemConfig == nullptr)
 			{
@@ -50,11 +49,11 @@ namespace FNavigationSystem
 				}
 			}
 
-			if (NavigationSystemConfig)
-			{
-				UNavigationSystemBase* NavSysInstance = NavigationSystemConfig->CreateAndConfigureNavigationSystem(WorldOwner);
-				WorldOwner.SetNavigationSystem(NavSysInstance);
-			}
+			UNavigationSystemBase* NavSysInstance = NavigationSystemConfig 
+				? NavigationSystemConfig->CreateAndConfigureNavigationSystem(WorldOwner)
+				: nullptr;
+			// we're setting to an instance or null, both are correct
+			WorldOwner.SetNavigationSystem(NavSysInstance);			
 		}
 
 		if (bInitializeForWorld)
diff --git a/Engine/Source/Runtime/NavigationSystem/Private/NavSystemConfigOverride.cpp b/Engine/Source/Runtime/NavigationSystem/Private/NavSystemConfigOverride.cpp
--- a/Engine/Source/Runtime/NavigationSystem/Private/NavSystemConfigOverride.cpp
+++ b/Engine/Source/Runtime/NavigationSystem/Private/NavSystemConfigOverride.cpp
@@ -81,55 +81,21 @@ void ANavSystemConfigOverride::PostLoad()
 void ANavSystemConfigOverride::ApplyConfig()
 {
 	UWorld* World = GetWorld();
-	if (World && NavigationSystemConfig)
+	if (World)
 	{
-		AWorldSettings* WorldSetting = World->GetWorldSettings();
-		if (WorldSetting)
+		UNavigationSystemBase* PrevNavSys = World->GetNavigationSystem();
+
+		if (PrevNavSys == nullptr || OverridePolicy == ENavSystemOverridePolicy::Override)
 		{
-			WorldSetting->SetNavigationSystemConfigOverride(NavigationSystemConfig);
+			OverrideNavSystem();
 		}
-
-		if (World->bIsWorldInitialized
-			&& NavigationSystemConfig
-#if WITH_EDITOR
-			&& (GIsEditorLoadingPackage == false || World->IsGameWorld())
-#endif // WITH_EDITOR
-			)
+		// PrevNavSys != null at this point
+		else if (OverridePolicy == ENavSystemOverridePolicy::Append)
 		{
-			
-			UNavigationSystemV1* PrevNavSys = Cast<UNavigationSystemV1>(World->GetNavigationSystem());
-			World->SetNavigationSystem(nullptr);
-
-			const FNavigationSystemRunMode RunMode = World->WorldType == EWorldType::Editor
-				? FNavigationSystemRunMode::EditorMode
-				: (World->WorldType == EWorldType::PIE
-					? FNavigationSystemRunMode::PIEMode
-					: FNavigationSystemRunMode::GameMode)
-				;
-
-			if (RunMode == FNavigationSystemRunMode::EditorMode)
-			{
-				FNavigationSystem::AddNavigationSystemToWorld(*World, RunMode, NavigationSystemConfig, /*bInitializeForWorld=*/false);
-#if WITH_EDITOR
-				UNavigationSystemBase* NewNavSys = World->GetNavigationSystem();
-				if (NewNavSys)
-				{
-					GEditor->GetTimerManager()->SetTimerForNextTick(FTimerDelegate::CreateUObject(this
-						, &ANavSystemConfigOverride::InitializeForWorld, NewNavSys, World, RunMode));
-				}
-#endif // WITH_EDITOR
-			}
-			else
-			{
-				FNavigationSystem::AddNavigationSystemToWorld(*World, RunMode, NavigationSystemConfig, /*bInitializeForWorld=*/true);
-			}
-
-			UNavigationSystemV1* NewNavSys = Cast<UNavigationSystemV1>(World->GetNavigationSystem());
-			if (NewNavSys && PrevNavSys)
-			{
-				NewNavSys->FinalizeOverridingNavSystem(*PrevNavSys);
-			}
+			// take the prev nav system and append data to it
+			AppendToNavSystem(*PrevNavSys);
 		}
+		// else PrevNavSys != null AND OverridePolicy == ENavSystemOverridePolicy::Skip, so ignoring the override
 	}
 }
 
@@ -138,6 +104,57 @@ void ANavSystemConfigOverride::PostInitProperties()
 	Super::PostInitProperties();
 }
 
+void ANavSystemConfigOverride::AppendToNavSystem(UNavigationSystemBase& PrevNavSys)
+{
+	if (NavigationSystemConfig)
+	{
+		PrevNavSys.AppendConfig(*NavigationSystemConfig);
+	}
+}
+
+void ANavSystemConfigOverride::OverrideNavSystem()
+{
+	UWorld* World = GetWorld();
+	if (World == nullptr)
+	{
+		return;
+	}
+
+	AWorldSettings* WorldSetting = World->GetWorldSettings();
+	if (WorldSetting)
+	{
+		WorldSetting->SetNavigationSystemConfigOverride(NavigationSystemConfig);
+	}
+
+	if (World->bIsWorldInitialized
+		&& NavigationSystemConfig)
+	{
+		const FNavigationSystemRunMode RunMode = World->WorldType == EWorldType::Editor
+			? FNavigationSystemRunMode::EditorMode
+			: (World->WorldType == EWorldType::PIE
+				? FNavigationSystemRunMode::PIEMode
+				: FNavigationSystemRunMode::GameMode)
+			;
+
+		if (RunMode == FNavigationSystemRunMode::EditorMode)
+		{
+			FNavigationSystem::AddNavigationSystemToWorld(*World, RunMode, NavigationSystemConfig, /*bInitializeForWorld=*/false, /*bOverridePreviousNavSys=*/true);
+#if WITH_EDITOR
+			UNavigationSystemBase* NewNavSys = World->GetNavigationSystem();
+			if (NewNavSys)
+			{
+				GEditor->GetTimerManager()->SetTimerForNextTick(FTimerDelegate::CreateUObject(this
+					, &ANavSystemConfigOverride::InitializeForWorld, NewNavSys, World, RunMode));
+			}
+#endif // WITH_EDITOR
+		}
+		else
+		{
+			FNavigationSystem::AddNavigationSystemToWorld(*World, RunMode, NavigationSystemConfig, /*bInitializeForWorld=*/true, /*bOverridePreviousNavSys=*/true);
+		}
+	}
+}
+
 #if WITH_EDITOR
 void ANavSystemConfigOverride::InitializeForWorld(UNavigationSystemBase* NewNavSys, UWorld* World, const FNavigationSystemRunMode RunMode)
 {
diff --git a/Engine/Source/Runtime/NavigationSystem/Private/NavigationSystem.cpp b/Engine/Source/Runtime/NavigationSystem/Private/NavigationSystem.cpp
--- a/Engine/Source/Runtime/NavigationSystem/Private/NavigationSystem.cpp
+++ b/Engine/Source/Runtime/NavigationSystem/Private/NavigationSystem.cpp
@@ -122,6 +122,12 @@ DEFINE_STAT(STAT_Navigation_RecastMemory);
 //----------------------------------------------------------------------//
 namespace FNavigationSystem
 {
+	const FNavDataConfig& GetFallbackNavDataConfig()
+	{
+		static FNavDataConfig FallbackNavDataConfig(FNavigationSystem::FallbackAgentRadius, FNavigationSystem::FallbackAgentHeight);
+		return FallbackNavDataConfig;
+	}
+
 	FORCEINLINE bool IsValidExtent(const FVector& Extent)
 	{
 		return Extent != INVALID_NAVEXTENT;
@@ -245,8 +251,6 @@ bool UNavigationSystemV1::bIsPIEActive = false;
 UNavigationSystemV1::UNavigationSystemV1(const FObjectInitializer& ObjectInitializer)
 	: Super(ObjectInitializer)
 	, bTickWhilePaused(false)
-	, bShouldGenerateNavDataOnlyWhenCompatibleNavBound(false)
-	, bAllowEquivalentNavDataOverride(false)
 	, bWholeWorldNavigable(false)
 	, bSkipAgentHeightCheckWhenPickingNavData(false)
 	, OperationMode(FNavigationSystemRunMode::InvalidMode)
@@ -503,7 +507,19 @@ void UNavigationSystemV1::UpdateAbstractNavData()
 
 void UNavigationSystemV1::SetSupportedAgentsNavigationClass(int32 AgentIndex, TSubclassOf<ANavigationData> NavigationDataClass)
 {
-	check(SupportedAgents.IsValidIndex(AgentIndex));
+	const bool bCDOInEditor =
+#if WITH_EDITOR
+		// the CDO will have 0 supported agents if none are defined which is fine in the editor
+		(GIsEditor && HasAnyFlags(RF_ClassDefaultObject))
+#else
+		false
+#endif // WITH_EDITOR
+		;
+
+	check(SupportedAgents.IsValidIndex(AgentIndex) || bCDOInEditor);
+
+	if (SupportedAgents.IsValidIndex(AgentIndex))
+	{
 	SupportedAgents[AgentIndex].SetNavDataClass(NavigationDataClass);
 
 	// keep preferred navigation data class in sync with actual class
@@ -512,6 +528,7 @@ void UNavigationSystemV1::SetSupportedAgentsNavigationClass(int32 AgentIndex, TS
 	//
 	// "Any" navigation data preference is valid only for instanced agents
 	SupportedAgents[AgentIndex].SetPreferredNavData(NavigationDataClass);
+	}
 
 #if WITH_EDITOR
 	if (GIsEditor)
@@ -543,19 +560,12 @@ void UNavigationSystemV1::PostInitProperties()
 			RegisterNavAreaClass(NavAreaClass);
 		}
 
-		// make sure there's at least one supported navigation agent size
-		if (SupportedAgents.Num() == 0)
-		{
-			SupportedAgents.Add(FNavDataConfig(FNavigationSystem::FallbackAgentRadius, FNavigationSystem::FallbackAgentHeight));
-		}
-		else
-		{
+		ApplySupportedAgentsFilter();
 			for (int32 AgentIndex = 0; AgentIndex < SupportedAgents.Num(); ++AgentIndex)
 			{
 				FNavDataConfig& SupportedAgentConfig = SupportedAgents[AgentIndex];
 				SetSupportedAgentsNavigationClass(AgentIndex, SupportedAgentConfig.GetNavDataClass<ANavigationData>());
 			}
-		}
 	
 		if (bInitialBuildingLocked)
 		{
@@ -1532,15 +1542,16 @@ ANavigationData* UNavigationSystemV1::GetDefaultNavDataInstance(FNavigationSyste
 {
 	checkSlow(IsInGameThread() == true);
 
-	if (MainNavData == NULL || MainNavData->IsPendingKill())
+	if (MainNavData == nullptr || MainNavData->IsPendingKill())
 	{
-		MainNavData = NULL;
+		MainNavData = nullptr;
 
 		// @TODO this should be done a differently. There should be specified a "default agent"
 		for (int32 NavDataIndex = 0; NavDataIndex < NavDataSet.Num(); ++NavDataIndex)
 		{
 			ANavigationData* NavData = NavDataSet[NavDataIndex];
-			if (NavData && !NavData->IsPendingKill() && NavData->CanBeMainNavData())
+			if (NavData && !NavData->IsPendingKill() && NavData->CanBeMainNavData()
+				&& (DefaultAgentName == NAME_None || NavData->GetConfig().Name == DefaultAgentName))
 			{
 				MainNavData = NavData;
 				break;
@@ -1548,7 +1559,7 @@ ANavigationData* UNavigationSystemV1::GetDefaultNavDataInstance(FNavigationSyste
 		}
 
 #if WITH_RECAST
-		if ( /*GIsEditor && */(MainNavData == NULL) && CreateNewIfNoneFound == FNavigationSystem::Create )
+		if (/*GIsEditor && */(MainNavData == nullptr) && CreateNewIfNoneFound == FNavigationSystem::Create)
 		{
 			// Spawn a new one if we're in the editor.  In-game, either we loaded one or we don't get one.
 			MainNavData = GetWorld()->SpawnActor<ANavigationData>(ARecastNavMesh::StaticClass());
@@ -1844,15 +1855,6 @@ UNavigationSystemV1::ERegistrationResult UNavigationSystemV1::RegisterNavData(AN
 	{
 		return RegistrationSuccessful;
 	}
-    // Don't allow nav data to register
-	else if (ShouldGenerateNavigationEverywhere() == false && ShouldGenerateNavDataOnlyWhenCompatibleNavBound() == true && NavData->IsA(AAbstractNavData::StaticClass()) == false)
-	{
-		TArray<FBox> SupportedBounds;
-		if (GetNavigationBoundsForNavData(*NavData, SupportedBounds) <= 0)
-		{
-			return RegistrationFailed_NotSuitable;
-		}
-	}
 
 	FScopeLock Lock(&NavDataRegistration);
 
@@ -1894,21 +1896,12 @@ UNavigationSystemV1::ERegistrationResult UNavigationSystemV1::RegisterNavData(AN
 		}
 		else
 		{
-			// check if this kind of agent has already its navigation implemented
-			TWeakObjectPtr<ANavigationData>* NavDataForAgent = AgentToNavDataMap.Find(NavConfig);
-			ANavigationData* NavDataInstanceForAgent = NavDataForAgent ? NavDataForAgent->Get() : nullptr;
+		// check if this kind of agent has already its navigation implemented
+		TWeakObjectPtr<ANavigationData>* NavDataForAgent = AgentToNavDataMap.Find(NavConfig);
+		ANavigationData* NavDataInstanceForAgent = NavDataForAgent ? NavDataForAgent->Get() : nullptr;
 
-			if (NavDataInstanceForAgent == nullptr || ShouldAllowEquivalentNavDataOverride())
-			{
-				if (NavDataInstanceForAgent != nullptr && ShouldAllowEquivalentNavDataOverride())
-				{
-					if (MainNavData == NavDataInstanceForAgent)
-					{
-						MainNavData = nullptr;
-					}
-
-					UnregisterNavData(NavDataInstanceForAgent);
-				}
+		if (NavDataInstanceForAgent == nullptr)
+		{
 				// ok, so this navigation agent doesn't have its navmesh registered yet, but do we want to support it?
 				bool bAgentSupported = false;
 
@@ -1923,7 +1916,7 @@ UNavigationSystemV1::ERegistrationResult UNavigationSystemV1::RegisterNavData(AN
 
 						NavData->SetConfig(SupportedAgents[AgentIndex]);
 						AgentToNavDataMap.Add(SupportedAgents[AgentIndex], NavData);
-						NavData->SetSupportsDefaultAgent(AgentIndex == 0);
+						NavData->SetSupportsDefaultAgent(SupportedAgents[AgentIndex].Name == DefaultAgentName);
 						NavData->ProcessNavAreas(NavAreaClasses, AgentIndex);
 
 						OnNavDataRegisteredEvent.Broadcast(NavData);
@@ -1936,19 +1929,19 @@ UNavigationSystemV1::ERegistrationResult UNavigationSystemV1::RegisterNavData(AN
 				}
 				Result = bAgentSupported == true ? RegistrationSuccessful : RegistrationFailed_AgentNotValid;
 			}
-			else if (NavDataInstanceForAgent == NavData)
-			{
-				ensure(NavDataSet.Find(NavData) != INDEX_NONE);
-				// let's treat double registration of the same nav data with the same agent as a success
-				Result = RegistrationSuccessful;
-			}
-			else
-			{
-				// otherwise specified agent type already has its navmesh implemented, fail redundant instance
-				Result = RegistrationFailed_AgentAlreadySupported;
-			}
+		else if (NavDataInstanceForAgent == NavData)
+		{
+			ensure(NavDataSet.Find(NavData) != INDEX_NONE);
+			// let's treat double registration of the same nav data with the same agent as a success
+			Result = RegistrationSuccessful;
+		}
+		else
+		{
+			// otherwise specified agent type already has its navmesh implemented, fail redundant instance
+			Result = RegistrationFailed_AgentAlreadySupported;
 		}
 	}
+	}
 	else
 	{
 		Result = RegistrationFailed_AgentNotValid;
@@ -2164,7 +2157,7 @@ void UNavigationSystemV1::OnNavigationAreaEvent(UClass* AreaClass, ENavAreaEvent
 
 int32 UNavigationSystemV1::GetSupportedAgentIndex(const ANavigationData* NavData) const
 {
-	if (SupportedAgents.Num() < 2)
+	if (SupportedAgents.Num() == 1)
 	{
 		return 0;
 	}
@@ -2172,7 +2165,7 @@ int32 UNavigationSystemV1::GetSupportedAgentIndex(const ANavigationData* NavData
 	const FNavDataConfig& TestConfig = NavData->GetConfig();
 	for (int32 AgentIndex = 0; AgentIndex < SupportedAgents.Num(); AgentIndex++)
 	{
-		if (SupportedAgents[AgentIndex].IsEquivalent(TestConfig))
+		if (SupportedAgents[AgentIndex].IsValid() && SupportedAgents[AgentIndex].IsEquivalent(TestConfig))
 		{
 			return AgentIndex;
 		}
@@ -2183,14 +2176,14 @@ int32 UNavigationSystemV1::GetSupportedAgentIndex(const ANavigationData* NavData
 
 int32 UNavigationSystemV1::GetSupportedAgentIndex(const FNavAgentProperties& NavAgent) const
 {
-	if (SupportedAgents.Num() < 2)
+	if (SupportedAgents.Num() == 1)
 	{
 		return 0;
 	}
 
 	for (int32 AgentIndex = 0; AgentIndex < SupportedAgents.Num(); AgentIndex++)
 	{
-		if (SupportedAgents[AgentIndex].IsEquivalent(NavAgent))
+		if (SupportedAgents[AgentIndex].IsValid() && SupportedAgents[AgentIndex].IsEquivalent(NavAgent))
 		{
 			return AgentIndex;
 		}
@@ -2308,15 +2301,6 @@ void UNavigationSystemV1::InitializeForWorld(UWorld& World, FNavigationSystemRun
 	OnWorldInitDone(Mode);
 }
 
-void UNavigationSystemV1::FinalizeOverridingNavSystem(UNavigationSystemV1& PreviousNavSystem)
-{
-    // When a Navigation Override Config loads it may happen that we have some nav data queued to be registered, keep them
-	if (PreviousNavSystem.NavDataRegistrationQueue.Num() > 0)
-	{
-		NavDataRegistrationQueue = PreviousNavSystem.NavDataRegistrationQueue;
-	}
-}
-
 UNavigationSystemV1* UNavigationSystemV1::GetCurrent(UWorld* World)
 {
 	return FNavigationSystem::GetCurrent<UNavigationSystemV1>(World);
@@ -2443,9 +2427,9 @@ void UNavigationSystemV1::UpdateComponentInNavOctree(UActorComponent& Comp)
 		AActor* OwnerActor = Comp.GetOwner();
 		if (OwnerActor)
 		{
-		INavRelevantInterface* NavElement = Cast<INavRelevantInterface>(&Comp);
-		if (NavElement)
-		{
+			INavRelevantInterface* NavElement = Cast<INavRelevantInterface>(&Comp);
+			if (NavElement)
+			{
 				UNavigationSystemV1* NavSys = FNavigationSystem::GetCurrent<UNavigationSystemV1>(OwnerActor->GetWorld());
 				if (NavSys)
 				{
@@ -2498,17 +2482,17 @@ void UNavigationSystemV1::UpdateActorAndComponentsInNavOctree(AActor& Actor, boo
 			DataHandler.UpdateActorAndComponentsInNavOctree(Actor);
 		}
 		else
-	{
+		{
 			// update just the actor itself 
 			INavRelevantInterface* NavElement = Cast<INavRelevantInterface>(&Actor);
 			if (NavElement)
-		{
+			{
 				DataHandler.UpdateNavOctreeElement(Actor, *NavElement, FNavigationOctreeController::OctreeUpdate_Default);
+			}
 		}
-	}
 
-	if (bUpdateAttachedActors)
-	{
+		if (bUpdateAttachedActors)
+		{
 			TArray<AActor*> UniqueAttachedActors;
 			if (GetAllAttachedActors(Actor, UniqueAttachedActors) > 0)
 			{
@@ -2597,7 +2581,7 @@ void UNavigationSystemV1::UpdateNavOctreeElement(UObject* ElementOwner, INavRele
 	if (ElementOwner && ElementInterface)
 	{
 		FNavigationDataHandler(DefaultOctreeController, DefaultDirtyAreasController).UpdateNavOctreeElement(*ElementOwner, *ElementInterface, UpdateFlags);
-		}
+	}
 }
 
 void UNavigationSystemV1::UpdateNavOctreeParentChain(UObject* ElementOwner, bool bSkipElementOwnerUpdate)
@@ -3055,16 +3039,24 @@ void UNavigationSystemV1::CancelBuild()
 
 void UNavigationSystemV1::SpawnMissingNavigationData()
 {
-	const int32 SupportedAgentsCount = SupportedAgents.Num();
-	check(SupportedAgentsCount >= 0);
+	const int32 AllSupportedAgentsCount = SupportedAgents.Num();
+	check(AllSupportedAgentsCount >= 0);
+	int32 ValidSupportedAgentsCount = 0;
+	for (int32 AgentIndex = 0; AgentIndex < AllSupportedAgentsCount; ++AgentIndex)
+	{
+		if (SupportedAgentsMask.Contains(AgentIndex))
+		{
+			++ValidSupportedAgentsCount;
+		}
+	}
 	
 	// Bit array might be a bit of an overkill here, but this function will be called very rarely
-	TBitArray<> AlreadyInstantiated(false, SupportedAgentsCount);
+	TBitArray<> AlreadyInstantiated(false, AllSupportedAgentsCount);
 	uint8 NumberFound = 0;
 	UWorld* NavWorld = GetWorld();
 
 	// 1. check whether any of required navigation data has already been instantiated
-	for (TActorIterator<ANavigationData> It(NavWorld); It && NumberFound < SupportedAgentsCount; ++It)
+	for (TActorIterator<ANavigationData> It(NavWorld); It && NumberFound < AllSupportedAgentsCount; ++It)
 	{
 		ANavigationData* Nav = (*It);
 		if (Nav != nullptr 
@@ -3073,7 +3065,7 @@ void UNavigationSystemV1::SpawnMissingNavigationData()
 			&& (Nav->GetTypedOuter<UWorld>() == NavWorld || NavWorld->GetLevels().Contains(Nav->GetLevel())))
 		{
 			// find out which one it is
-			for (int32 AgentIndex = 0; AgentIndex < SupportedAgentsCount; ++AgentIndex)
+			for (int32 AgentIndex = 0; AgentIndex < AllSupportedAgentsCount; ++AgentIndex)
 			{
 				if (AlreadyInstantiated[AgentIndex] == false
 					&& Nav->GetClass() == SupportedAgents[AgentIndex].GetNavDataClass<ANavigationData>()
@@ -3088,12 +3080,14 @@ void UNavigationSystemV1::SpawnMissingNavigationData()
 	}
 
 	// 2. for any not already instantiated navigation data call creator functions
-	if (NumberFound < SupportedAgentsCount)
+	if (NumberFound < ValidSupportedAgentsCount)
 	{
-		for (int32 AgentIndex = 0; AgentIndex < SupportedAgentsCount; ++AgentIndex)
+		for (int32 AgentIndex = 0; AgentIndex < AllSupportedAgentsCount; ++AgentIndex)
 		{
 			const FNavDataConfig& NavConfig = SupportedAgents[AgentIndex];
-			if (AlreadyInstantiated[AgentIndex] == false && NavConfig.GetNavDataClass<ANavigationData>() != nullptr)
+			if (AlreadyInstantiated[AgentIndex] == false 
+				&& SupportedAgentsMask.Contains(AgentIndex)
+				&& NavConfig.GetNavDataClass<ANavigationData>() != nullptr)
 			{
 				bool bHandled = false;
 
@@ -3137,12 +3131,22 @@ ANavigationData* UNavigationSystemV1::CreateNavigationDataInstanceInLevel(const
 	UWorld* World = GetWorld();
 	check(World);
 
+	const int32 NavSupportedAgents = GetSupportedAgentIndex(NavConfig);
+
+	// not creating new NavData instance if the agent it's representing is not supported
+	// with the exception of AbstractNavData
+	if (NavSupportedAgents == INDEX_NONE
+		&& NavConfig.GetNavDataClass<AAbstractNavData>() == nullptr)
+	{
+		UE_LOG(LogNavigation, Warning, TEXT("Unable to create NavigationData instance for config \'%s\' as this agent is not supported by current NavigationSystem instance")
+			, *NavConfig.GetDescription());
+		return nullptr;
+	}
+
 	FActorSpawnParameters SpawnInfo;
 	SpawnInfo.OverrideLevel = SpawnLevel;
 	if (bSpawnNavDataInNavBoundsLevel && SpawnLevel == nullptr && RegisteredNavBounds.Num() > 0)
 	{
-		const int32 NavSupportedAgents = GetSupportedAgentIndex(NavConfig);
-
 		// pick the first valid level that supports these agents
 		for (const FNavigationBounds& Bounds : RegisteredNavBounds)
 		{
@@ -3175,7 +3179,6 @@ ANavigationData* UNavigationSystemV1::CreateNavigationDataInstanceInLevel(const
 				if (ExistingNavigationData)
 				{
 					UnregisterNavData(ExistingNavigationData);
-					AgentToNavDataMap.Remove(ExistingNavigationData->GetConfig());
 				}
 
 				ExistingObject->Rename(NULL, NULL, REN_DontCreateRedirectors | REN_ForceGlobalUnique | REN_DoNotDirty | REN_NonTransactional | REN_ForceNoResetLoaders);
@@ -3329,26 +3332,17 @@ int32 UNavigationSystemV1::GetNumRunningBuildTasks() const
 
 void UNavigationSystemV1::OnLevelAddedToWorld(ULevel* InLevel, UWorld* InWorld)
 {
-	if ((IsNavigationSystemStatic() == false) && (InWorld == GetWorld()) && (InLevel != nullptr))
+	if ((InWorld != GetWorld()) || (InLevel == nullptr))
+	{
+		return;
+	}
+
+	if ((IsNavigationSystemStatic() == false))
 	{
 		AddLevelCollisionToOctree(InLevel);
 
 		if (!InLevel->IsPersistentLevel())
 		{
-#if WITH_EDITOR
-            // Registering NavigationData in its PostLoad() isn't enough because when adding/removing level in editor, it calls Remove/Load to refresh streaming levels
-			if (GIsEditor)
-			{
-				for (TActorIterator<ANavigationData> It(InWorld); It; ++It)
-				{
-					ANavigationData* NavData = *It;
-					if (NavData->GetLevel() == InLevel)
-					{
-						RequestRegistrationDeferred(*NavData);
-					}
-				}
-			}
-#endif
 			for (ANavigationData* NavData : NavDataSet)
 			{
 				if (NavData)
@@ -3358,6 +3352,29 @@ void UNavigationSystemV1::OnLevelAddedToWorld(ULevel* InLevel, UWorld* InWorld)
 			}
 		}
 	}
+
+#if WITH_EDITOR
+	if (OperationMode == FNavigationSystemRunMode::EditorMode)
+	{
+		// see if there are any unregistered yet valid nav data instances
+			// In general we register navdata on its PostLoad, but in some cases
+			// levels get removed from world and readded and in that case we might
+			// miss registering them
+		for (AActor* Actor : InLevel->Actors)
+		{
+			ANavigationData* NavData = Cast<ANavigationData>(Actor);
+			if (NavData != nullptr && NavData->IsRegistered() == false)
+			{
+				RequestRegistrationDeferred(*NavData);
+			}
+		}
+	}
+	else
+#endif WITH_EDITOR
+		if (NavDataRegistrationQueue.Num() > 0)
+	{
+		ProcessRegistrationCandidates();
+	}
 }
 
 void UNavigationSystemV1::OnLevelRemovedFromWorld(ULevel* InLevel, UWorld* InWorld)
@@ -3379,28 +3396,9 @@ void UNavigationSystemV1::OnLevelRemovedFromWorld(ULevel* InLevel, UWorld* InWor
 					}
 					else
 					{
-                        if (NavData->IsRegistered())
-                        {
+						// removing manually first so that UnregisterNavData won't mess with NavDataSet
+						NavDataSet.RemoveAt(DataIndex, 1, /*bAllowShrinking=*/false);
                             UnregisterNavData(NavData);
-
-                            if (ShouldAllowEquivalentNavDataOverride())
-                            {
-                                for (TActorIterator<ANavigationData> It(InWorld); It; ++It)
-                                {
-                                    ANavigationData* CurrentNavData = *It;
-                                    const FNavDataConfig& TestConfig = NavData->GetConfig();
-                                    if (!CurrentNavData->IsRegistered() && CurrentNavData->GetLevel() != InLevel && GetTypeHash(CurrentNavData->GetConfig()) == GetTypeHash(NavData->GetConfig()))
-                                    {
-                                        RequestRegistrationDeferred(*CurrentNavData);
-										break;
-                                    }
-                                }
-                            }
-                        }
-                        else
-                        {
-                            NavDataSet.RemoveAt(DataIndex, 1, /*bAllowShrinking=*/false);
-                        }
                     }
                 }
             }
@@ -3529,15 +3527,16 @@ void UNavigationSystemV1::CleanUp(FNavigationSystem::ECleanupMode Mode)
 			UnregisterNavData(NavData);
 		}
 	}	
+
 	if (NavDataSet.Num())
 	{
-        UE_LOG(LogNavigation, Error, TEXT("UNavigationSystemV1::CleanUp still have an agent in NavDataSet loaded after unregister them all!!"));
+		UE_LOG(LogNavigation, Error, TEXT("UNavigationSystemV1::CleanUp still has data in NavDataSet after unregister them all"));
         NavDataSet.Reset();
     }
-	
+
 	if (AgentToNavDataMap.Num())
 	{
-		UE_LOG(LogNavigation, Error, TEXT("UNavigationSystemV1::CleanUp still have an agent mapped to navigation data after clean up!"));
+		UE_LOG(LogNavigation, Error, TEXT("UNavigationSystemV1::CleanUp still has agents mapped to navigation data after clean up"));
 		AgentToNavDataMap.Reset();
 	}
 	
@@ -3550,9 +3549,9 @@ void UNavigationSystemV1::CleanUp(FNavigationSystem::ECleanupMode Mode)
 		MyWorld->OnBeginTearingDown().RemoveAll(this);
 
 		if (MyWorld->WorldType == EWorldType::Game || MyWorld->WorldType == EWorldType::Editor)
-	{
-		INavLinkCustomInterface::NextUniqueId = 1;
-	}
+		{
+			INavLinkCustomInterface::NextUniqueId = 1;
+		}
 	}
 
 	bCleanUpDone = true;
@@ -3821,6 +3820,16 @@ bool FNavigationSystemExec::Exec(UWorld* InWorld, const TCHAR* Cmd, FOutputDevic
 		{
 			NavSys->Build();
 		}
+		else if (FParse::Command(&Cmd, TEXT("RedrawNav")) || FParse::Command(&Cmd, TEXT("RedrawNavigation")))
+		{
+			for (ANavigationData* NavData : NavSys->NavDataSet)
+			{
+				if (NavData)
+				{
+					NavData->MarkComponentsRenderStateDirty();
+				}
+			}
+		}
 	}
 
 	return false;
@@ -4212,6 +4221,8 @@ int UNavigationSystemV1::GetNavigationBoundsForNavData(const ANavigationData& Na
 	OutBounds.Reserve(InitialBoundsCount + RegisteredNavBounds.Num());
 	const int32 AgentIndex = GetSupportedAgentIndex(&NavData);
 
+	if (AgentIndex != INDEX_NONE)
+	{
 	for (const FNavigationBounds& NavigationBounds : RegisteredNavBounds)
 	{
 		if ((InLevel == nullptr || NavigationBounds.Level == InLevel)
@@ -4220,6 +4231,7 @@ int UNavigationSystemV1::GetNavigationBoundsForNavData(const ANavigationData& Na
 			OutBounds.Add(NavigationBounds.AreaBox);
 		}
 	}
+	}
 
 	return OutBounds.Num() - InitialBoundsCount;
 }
@@ -4234,20 +4246,117 @@ const FNavDataConfig& UNavigationSystemV1::GetDefaultSupportedAgent()
 		: DefaultAgent;
 }
 
+const FNavDataConfig& UNavigationSystemV1::GetDefaultSupportedAgentConfig() const 
+{ 
+	static const FNavDataConfig DefaultAgent;
+
+	int32 FirstValidIndex = INDEX_NONE;
+	for (int32 AgentIndex = 0; AgentIndex < SupportedAgents.Num(); ++AgentIndex)
+	{
+		if (SupportedAgentsMask.Contains(AgentIndex))
+		{
+			if ((DefaultAgentName == NAME_None || SupportedAgents[AgentIndex].Name == DefaultAgentName))
+			{
+				return SupportedAgents[AgentIndex];
+			}
+			FirstValidIndex = (FirstValidIndex == INDEX_NONE) ? AgentIndex : FirstValidIndex;			
+		}
+	}
+
+	// if not found, get the first one allowed
+	return FirstValidIndex != INDEX_NONE ? SupportedAgents[FirstValidIndex] : DefaultAgent;;
+}
+
 void UNavigationSystemV1::OverrideSupportedAgents(const TArray<FNavDataConfig>& NewSupportedAgents)
 {
 	UE_CLOG(bWorldInitDone, LogNavigation, Warning, TEXT("Trying to override NavigationSystem\'s SupportedAgents past the World\'s initialization"));
 
-	SupportedAgents = NewSupportedAgents;
+	SupportedAgentsMask.Empty();
+
+	// reset the SupportedAgents 
+	const UNavigationSystemV1* NavSysCDO = GetClass()->GetDefaultObject<UNavigationSystemV1>();
+	SupportedAgents = NavSysCDO->SupportedAgents;
+
+	for (const FNavDataConfig& Agent : NewSupportedAgents)
+	{
+		for(int32 AgentIndex = 0; AgentIndex < SupportedAgents.Num(); AgentIndex++)
+		{
+			if (SupportedAgents[AgentIndex].IsEquivalent(Agent))
+			{
+				SupportedAgentsMask.Set(AgentIndex);
+				break;
+			}
+		}
+	}
+
+	SupportedAgentsMask.MarkInitialized();
+
+	ApplySupportedAgentsFilter();
+}
+
+void UNavigationSystemV1::ApplySupportedAgentsFilter()
+{
+	// reset the SupportedAgents 
+	const UNavigationSystemV1* NavSysCDO = GetClass()->GetDefaultObject<UNavigationSystemV1>();
+	SupportedAgents = NavSysCDO->SupportedAgents;
+	// make sure there's at least one supported navigation agent size
 	if (SupportedAgents.Num() == 0)
 	{
-		SupportedAgents.Add(FNavDataConfig(FNavigationSystem::FallbackAgentRadius, FNavigationSystem::FallbackAgentHeight));
+		SupportedAgents.Add(FNavigationSystem::GetFallbackNavDataConfig());
+	}
+
+	// make all SupportedAgents filtered out by SupportedAgentsMask invalid by
+	// clearing out their NavDataClass
+	for (int32 AgentIndex = 0; AgentIndex < SupportedAgents.Num(); AgentIndex++)
+	{
+		if (SupportedAgentsMask.Contains(AgentIndex) == false)
+	{
+			SupportedAgents[AgentIndex].Invalidate();
+		}
 	}
 }
 
+void UNavigationSystemV1::SetSupportedAgentsMask(const FNavAgentSelector& InSupportedAgentsMask)
+{
+	SupportedAgentsMask = InSupportedAgentsMask;
+	ApplySupportedAgentsFilter();
+}
+
 void UNavigationSystemV1::Configure(const UNavigationSystemConfig& Config)
 {
+	if (Config.DefaultAgentName != NAME_None)
+	{
+		DefaultAgentName = Config.DefaultAgentName;
+	}
+	SetSupportedAgentsMask(Config.SupportedAgentsMask);
+}
+
+void UNavigationSystemV1::AppendConfig(const UNavigationSystemConfig& NewConfig)
+{
+	if (NewConfig.SupportedAgentsMask.IsSame(SupportedAgentsMask) == false)
+	{
+		bool bAgentsAdded = false;
+		for (int AgentIndex = 0; AgentIndex < SupportedAgents.Num(); ++AgentIndex)
+		{
+			if (NewConfig.SupportedAgentsMask.Contains(AgentIndex) == true
+				&& SupportedAgentsMask.Contains(AgentIndex) == false)
+			{
+				SupportedAgentsMask.Set(AgentIndex);
+				bAgentsAdded = true;
+			}
+		}
+
+		if (bAgentsAdded)
+		{
+			ApplySupportedAgentsFilter();
+			// @todo consider updating the octree, it might be missing data for the new agent(s)
+		}
 
+		if (DefaultAgentName == NAME_None)
+		{
+			DefaultAgentName = NewConfig.DefaultAgentName;
+		}
+	}
 }
 
 //----------------------------------------------------------------------//
@@ -4278,8 +4387,6 @@ void UNavigationSystemModuleConfig::UpdateWithNavSysCDO(const UNavigationSystemV
 		bCreateOnClient = NavSysCDO.bAllowClientSideNavigation;
 		bAutoSpawnMissingNavData = NavSysCDO.bAutoCreateNavigationData;
 		bSpawnNavDataInNavBoundsLevel = NavSysCDO.bSpawnNavDataInNavBoundsLevel;
-		bShouldGenerateNavDataOnlyWhenCompatibleNavBound = NavSysCDO.bShouldGenerateNavDataOnlyWhenCompatibleNavBound;
-		bAllowEquivalentNavDataOverride = NavSysCDO.bAllowEquivalentNavDataOverride;
 	}
 }
 
@@ -4301,8 +4408,6 @@ UNavigationSystemBase* UNavigationSystemModuleConfig::CreateAndConfigureNavigati
 	{
 		NavSysInstance->bAutoCreateNavigationData = bAutoSpawnMissingNavData;
 		NavSysInstance->bSpawnNavDataInNavBoundsLevel = bSpawnNavDataInNavBoundsLevel;
-		NavSysInstance->bShouldGenerateNavDataOnlyWhenCompatibleNavBound = bShouldGenerateNavDataOnlyWhenCompatibleNavBound;
-		NavSysInstance->bAllowEquivalentNavDataOverride  = bAllowEquivalentNavDataOverride;
 		if (bStrictlyStatic)
 		{
 			NavSysInstance->ConfigureAsStatic();
diff --git a/Engine/Source/Runtime/NavigationSystem/Public/NavSystemConfigOverride.h b/Engine/Source/Runtime/NavigationSystem/Public/NavSystemConfigOverride.h
--- a/Engine/Source/Runtime/NavigationSystem/Public/NavSystemConfigOverride.h
+++ b/Engine/Source/Runtime/NavigationSystem/Public/NavSystemConfigOverride.h
@@ -11,6 +11,14 @@
 
 class UNavigationSystemConfig;
 
+UENUM()
+enum class ENavSystemOverridePolicy : uint8
+{
+	Override, // the pre-exising nav system instance will be destroyed.
+	Append, // config information will be added to pre-existing nav system instance
+	Skip	// if there's already a NavigationSystem in the world then the overriding config will be ignored
+};
+
 
 UCLASS(hidecategories = (Input, Rendering, Actor, LOD, Cooking))
 class NAVIGATIONSYSTEM_API ANavSystemConfigOverride : public AActor
@@ -27,6 +35,10 @@ protected:
 	UPROPERTY(EditAnywhere, BlueprintReadOnly, Category = Navigation, Instanced,  meta = (NoResetToDefault))
 	UNavigationSystemConfig* NavigationSystemConfig;
 
+	/** If there's already a NavigationSystem instance in the world how should this nav override behave */
+	UPROPERTY(EditAnywhere, BlueprintReadOnly, Category = Navigation)
+	ENavSystemOverridePolicy OverridePolicy;
+
 	UPROPERTY(EditAnywhere, BlueprintReadWrite, Category = Navigation, AdvancedDisplay)
 	uint8 bLoadOnClient : 1;
 
@@ -49,6 +61,14 @@ public:
 #endif
 
 protected:
+	/** Creates a new navigation system and plugs it into the world. If there's a
+	 *	nav system instance already in place it gets destroyed. */
+	virtual void OverrideNavSystem();
+
+	/** Appends non-conflicting information (like supported agents) to a pre-existing 
+	 *	nav system instance */
+	virtual void AppendToNavSystem(UNavigationSystemBase& PrevNavSys);
+
 #if WITH_EDITOR
 	/** Called only in the editor mode*/
 	void InitializeForWorld(UNavigationSystemBase* NewNavSys, UWorld* World, const FNavigationSystemRunMode RunMode);
diff --git a/Engine/Source/Runtime/NavigationSystem/Public/NavigationOctreeController.h b/Engine/Source/Runtime/NavigationSystem/Public/NavigationOctreeController.h
--- a/Engine/Source/Runtime/NavigationSystem/Public/NavigationOctreeController.h
+++ b/Engine/Source/Runtime/NavigationSystem/Public/NavigationOctreeController.h
@@ -40,6 +40,7 @@ struct NAVIGATIONSYSTEM_API FNavigationOctreeController
 	/** basically says if navoctree has been created already */
 	bool IsValid() const { return NavOctree.IsValid(); }
 	bool IsValidElement(const FOctreeElementId& ElementId) const;
+	bool IsEmpty() const { return (IsValid() == false) || NavOctree->GetSizeBytes() == 0; }
 
 private:
 	static uint32 HashObject(const UObject& Object);
diff --git a/Engine/Source/Runtime/NavigationSystem/Public/NavigationSystem.h b/Engine/Source/Runtime/NavigationSystem/Public/NavigationSystem.h
--- a/Engine/Source/Runtime/NavigationSystem/Public/NavigationSystem.h
+++ b/Engine/Source/Runtime/NavigationSystem/Public/NavigationSystem.h
@@ -38,10 +38,6 @@ class UNavigationSystemModuleConfig;
 struct FNavigationRelevantData;
 struct FNavigationOctreeElement;
 
-#if WITH_EDITOR
-class FEdMode;
-#endif // WITH_EDITOR
-
 /** delegate to let interested parties know that new nav area class has been registered */
 DECLARE_MULTICAST_DELEGATE_OneParam(FOnNavAreaChanged, const UClass* /*AreaClass*/);
 
@@ -117,6 +113,13 @@ public:
 	ANavigationData* AbstractNavData;
 
 protected:
+	/** If not None indicates which of navigation datas and supported agents are
+	 * going to be used as the default ones. If navigation agent of this type does
+	 * not exist or is not enabled then the first available nav data will be used
+	 * as the default one */
+	UPROPERTY(config, EditAnywhere, BlueprintReadOnly, Category = Navigation)
+	FName DefaultAgentName;
+
 	UPROPERTY(config, EditAnywhere, BlueprintReadOnly, Category = Navigation)
 	TSoftClassPtr<UCrowdManagerBase> CrowdManagerClass;
 
@@ -142,13 +145,6 @@ protected:
 	UPROPERTY()
 	uint32 bSupportRebuilding : 1; 
 
-	/** When set to TRUE, navigation data will be generated only when at least one navigation bounds is compatible for building it */
-	UPROPERTY(EditAnywhere, Category = Navigation)
-	uint32 bShouldGenerateNavDataOnlyWhenCompatibleNavBound : 1;
-
-	/** When set to TRUE, equivalent navigation data will override the existing one in the system */
-	UPROPERTY(EditAnywhere, Category = Navigation)
-	uint32 bAllowEquivalentNavDataOverride  : 1;
 public:
 	/** if set to true will result navigation system not rebuild navigation until 
 	 *	a call to ReleaseInitialBuildingLock() is called. Does not influence 
@@ -192,6 +188,12 @@ protected:
 	UPROPERTY(config, EditAnywhere, Category = Agents)
 	TArray<FNavDataConfig> SupportedAgents;
 
+	/** NavigationSystem's properties in Project Settings define all possible supported agents,
+	 *	but a specific navigation system can choose to support only a subset of agents. Set via 
+	 *	NavigationSystemConfig */
+	UPROPERTY(config, EditAnywhere, Category = Agents)
+	FNavAgentSelector SupportedAgentsMask;
+
 public:
 
 	UPROPERTY()
@@ -324,7 +326,7 @@ public:
 		RegistrationFailed_DataPendingKill,			// means navigation data being registered is marked as pending kill
 		RegistrationFailed_AgentAlreadySupported,	// this means that navigation agent supported by given nav data is already handled by some other, previously registered instance
 		RegistrationFailed_AgentNotValid,			// given instance contains navmesh that doesn't support any of expected agent types, or instance doesn't specify any agent
-		RegistrationFailed_NotSuitable,				// given instance had been considered unsuitable by current navigation system instance itself
+		RegistrationFailed_NotSuitable,				// given instance had been considered unsuitable by current navigation system instance itself. NOTE: this value is not currently being used by the engine-supplied navigation system classes
 		RegistrationSuccessful,
 	};
 
@@ -494,8 +496,6 @@ public:
 	bool ShouldAllowClientSideNavigation() const { return bAllowClientSideNavigation; }
 	virtual bool ShouldLoadNavigationOnClient(ANavigationData* NavData = nullptr) const { return bAllowClientSideNavigation; }
 	virtual bool ShouldDiscardSubLevelNavData(ANavigationData* NavData = nullptr) const { return bShouldDiscardSubLevelNavData; }
-	bool ShouldGenerateNavDataOnlyWhenCompatibleNavBound() const { return bShouldGenerateNavDataOnlyWhenCompatibleNavBound; }
-	bool ShouldAllowEquivalentNavDataOverride() const { return bAllowEquivalentNavDataOverride ; }
 
 	FBox GetWorldBounds() const;
 	
@@ -506,9 +506,10 @@ public:
 	const TSet<FNavigationBounds>& GetNavigationBounds() const;
 
 	static const FNavDataConfig& GetDefaultSupportedAgent();
-	FORCEINLINE const FNavDataConfig& GetDefaultSupportedAgentConfig() const { check(SupportedAgents.Num() > 0);  return SupportedAgents[0]; }
+	const FNavDataConfig& GetDefaultSupportedAgentConfig() const;
 	FORCEINLINE const TArray<FNavDataConfig>& GetSupportedAgents() const { return SupportedAgents; }
 	void OverrideSupportedAgents(const TArray<FNavDataConfig>& NewSupportedAgents);
+	void SetSupportedAgentsMask(const FNavAgentSelector& InSupportedAgentsMask);
 
 	virtual void ApplyWorldOffset(const FVector& InOffset, bool bWorldShift) override;
 
@@ -547,6 +548,7 @@ public:
 	virtual void RequestRegistration(ANavigationData* NavData, bool bTriggerRegistrationProcessing = true);
 
 protected:
+	void ApplySupportedAgentsFilter();
 
 	/** Processes all NavigationData instances in UWorld owning navigation system instance, and registers
 	 *	all previously unregistered */
@@ -800,9 +802,6 @@ public:
 	
 	virtual void InitializeForWorld(UWorld& World, FNavigationSystemRunMode Mode) override;
 
-	/** Called after a NavSystemConfigOverride has finished overriding the old NavSystem */
-	virtual void FinalizeOverridingNavSystem(UNavigationSystemV1& PreviousNavSystem);
-
 	// Fetch the array of all nav-agent properties.
 	void GetNavAgentPropertiesArray(TArray<FNavAgentProperties>& OutNavAgentProperties) const;
 
@@ -1001,7 +1000,8 @@ public:
 	/** @param if InLevel is given then only navigation bounds from that level will be considered*/
 	virtual int GetNavigationBoundsForNavData(const ANavigationData& NavData, TArray<FBox>& OutBounds, ULevel* InLevel = nullptr) const;
 	static INavigationDataInterface* GetNavDataForActor(const AActor& Actor);
-	virtual void Configure(const UNavigationSystemConfig& Config);
+	virtual void Configure(const UNavigationSystemConfig& Config) override;
+	virtual void AppendConfig(const UNavigationSystemConfig& NewConfig) override;
 	
 #if !UE_BUILD_SHIPPING
 	void GetOnScreenMessages(TMultiMap<FCoreDelegates::EOnScreenMessageSeverity, FText>& OutMessages);
@@ -1087,13 +1087,6 @@ protected:
 	UPROPERTY(EditAnywhere, Category = Navigation)
 	uint32 bSpawnNavDataInNavBoundsLevel : 1;
 
-	/** When set to TRUE, navigation data will be generated only when at least one navigation bounds is compatible for building it */
-	UPROPERTY(EditAnywhere, Category = Navigation)
-	uint32 bShouldGenerateNavDataOnlyWhenCompatibleNavBound : 1;
-
-	/** When set to TRUE, equivalent navigation data will override the existing one in the system */
-	UPROPERTY(EditAnywhere, Category = Navigation)
-	uint32 bAllowEquivalentNavDataOverride : 1;
 public:
 	UNavigationSystemModuleConfig(const FObjectInitializer& ObjectInitializer = FObjectInitializer::Get());
 
