CIS fix

diff --git a/Engine/Source/Runtime/NavigationSystem/Private/NavigationSystem.cpp b/Engine/Source/Runtime/NavigationSystem/Private/NavigationSystem.cpp
--- a/Engine/Source/Runtime/NavigationSystem/Private/NavigationSystem.cpp
+++ b/Engine/Source/Runtime/NavigationSystem/Private/NavigationSystem.cpp
@@ -3370,7 +3370,7 @@ void UNavigationSystemV1::OnLevelAddedToWorld(ULevel* InLevel, UWorld* InWorld)
 		}
 	}
 	else
-#endif WITH_EDITOR
+#endif // WITH_EDITOR
 		if (NavDataRegistrationQueue.Num() > 0)
 	{
 		ProcessRegistrationCandidates();
