Extended FNavMeshSceneProxyData with a capability to store arbitrary FBoxes and given UNavMeshRenderingComponent's subclassess a way to populate FNavMeshSceneProxyData with data.

Yoan.StAmant

diff --git a/Engine/Source/Runtime/NavigationSystem/Private/NavMesh/NavMeshRenderingComponent.cpp b/Engine/Source/Runtime/NavigationSystem/Private/NavMesh/NavMeshRenderingComponent.cpp
--- a/Engine/Source/Runtime/NavigationSystem/Private/NavMesh/NavMeshRenderingComponent.cpp
+++ b/Engine/Source/Runtime/NavigationSystem/Private/NavMesh/NavMeshRenderingComponent.cpp
@@ -234,6 +234,7 @@ void FNavMeshSceneProxyData::Reset()
 	NavMeshEdgeLines.Reset();
 	NavLinkLines.Reset();
 	ClusterLinkLines.Reset();
+	AuxBoxes.Reset();
 	DebugLabels.Reset();
 	OctreeBounds.Reset();
 	Bounds.Init();
@@ -298,6 +299,31 @@ void FNavMeshSceneProxyData::Serialize(FArchive& Ar)
 		}
 	}
 
+	int32 NumBoxes = AuxBoxes.Num();
+	Ar << NumBoxes;
+	if (Ar.IsLoading())
+	{
+		FDebugRenderSceneProxy::FDebugBox TmpBox = FDebugRenderSceneProxy::FDebugBox(FBox(), FColor());
+		AuxBoxes.Reserve(NumBoxes);
+		for (int32 Idx = 0; Idx < NumBoxes; Idx++)
+		{	
+			Ar << TmpBox.Box;
+			Ar << TmpBox.Color;
+			Ar << TmpBox.Transform;
+
+			AuxBoxes.Add(TmpBox);
+		}
+	}
+	else
+	{
+		for (int32 Idx = 0; Idx < NumBoxes; Idx++)
+		{
+			Ar << AuxBoxes[Idx].Box;
+			Ar << AuxBoxes[Idx].Color;
+			Ar << AuxBoxes[Idx].Transform;
+		}
+	}
+
 	int32 NumLabels = DebugLabels.Num();
 	Ar << NumLabels;
 	if (Ar.IsLoading())
@@ -342,6 +368,7 @@ uint32 FNavMeshSceneProxyData::GetAllocatedSize() const
 		NavMeshEdgeLines.GetAllocatedSize() +
 		NavLinkLines.GetAllocatedSize() +
 		ClusterLinkLines.GetAllocatedSize() +
+		AuxBoxes.GetAllocatedSize() +
 		DebugLabels.GetAllocatedSize() +
 		OctreeBounds.GetAllocatedSize();
 }
@@ -824,9 +851,12 @@ FNavMeshSceneProxy::FNavMeshSceneProxy(const UPrimitiveComponent* InComponent, F
 	, bRequestedData(false)
 	, bForceRendering(ForceToRender)
 {
+	DrawType = EDrawType::SolidAndWireMeshes;
+
 	if (InProxyData)
 	{
 		ProxyData = *InProxyData;
+		Boxes.Append(InProxyData->AuxBoxes);
 	}
 
 	RenderingComponent = MakeWeakObjectPtr(const_cast<UNavMeshRenderingComponent*>(Cast<UNavMeshRenderingComponent>(InComponent)));
@@ -945,6 +975,8 @@ void FNavMeshSceneProxy::GetDynamicMeshElements(const TArray<const FSceneView*>&
 {
 	QUICK_SCOPE_CYCLE_COUNTER(STAT_RecastRenderingSceneProxy_GetDynamicMeshElements);
 
+	FDebugRenderSceneProxy::GetDynamicMeshElements(Views, ViewFamily, VisibilityMap, Collector);
+
 	for (int32 ViewIndex = 0; ViewIndex < Views.Num(); ViewIndex++)
 	{
 		if (VisibilityMap & (1 << ViewIndex))
@@ -1270,6 +1302,13 @@ void UNavMeshRenderingComponent::OnUnregister()
 	Super::OnUnregister();
 }
 
+void UNavMeshRenderingComponent::GatherData(const ARecastNavMesh& NavMesh, FNavMeshSceneProxyData& OutProxyData) const
+{
+	const int32 DetailFlags = OutProxyData.GetDetailFlags(&NavMesh);
+	TArray<int32> EmptyTileSet;
+	OutProxyData.GatherData(&NavMesh, DetailFlags, EmptyTileSet);
+}
+
 FPrimitiveSceneProxy* UNavMeshRenderingComponent::CreateSceneProxy()
 {
 #if WITH_RECAST && !UE_BUILD_SHIPPING && !UE_BUILD_TEST
@@ -1285,10 +1324,7 @@ FPrimitiveSceneProxy* UNavMeshRenderingComponent::CreateSceneProxy()
 		if (NavMesh && NavMesh->IsDrawingEnabled())
 		{
 			FNavMeshSceneProxyData ProxyData;
-
-			const int32 DetailFlags = ProxyData.GetDetailFlags(NavMesh);
-			TArray<int32> EmptyTileSet;
-			ProxyData.GatherData(NavMesh, DetailFlags, EmptyTileSet);
+			GatherData(*NavMesh, ProxyData);
 
 			NavMeshSceneProxy = new FNavMeshSceneProxy(this, &ProxyData);
 		}
diff --git a/Engine/Source/Runtime/NavigationSystem/Public/NavMesh/NavMeshRenderingComponent.h b/Engine/Source/Runtime/NavigationSystem/Public/NavMesh/NavMeshRenderingComponent.h
--- a/Engine/Source/Runtime/NavigationSystem/Public/NavMesh/NavMeshRenderingComponent.h
+++ b/Engine/Source/Runtime/NavigationSystem/Public/NavMesh/NavMeshRenderingComponent.h
@@ -60,6 +60,7 @@ struct NAVIGATIONSYSTEM_API FNavMeshSceneProxyData : public TSharedFromThis<FNav
 	TArray<FDebugRenderSceneProxy::FDebugLine> NavMeshEdgeLines;
 	TArray<FDebugRenderSceneProxy::FDebugLine> NavLinkLines;
 	TArray<FDebugRenderSceneProxy::FDebugLine> ClusterLinkLines;
+	TArray<FDebugRenderSceneProxy::FDebugBox> AuxBoxes;
 
 	struct FDebugText
 	{
@@ -97,7 +98,7 @@ class NAVIGATIONSYSTEM_API FNavMeshSceneProxy final : public FDebugRenderScenePr
 {
 	friend class FNavMeshDebugDrawDelegateHelper;
 public:
-	SIZE_T GetTypeHash() const override;
+	virtual SIZE_T GetTypeHash() const override;
 
 	FNavMeshSceneProxy(const UPrimitiveComponent* InComponent, FNavMeshSceneProxyData* InProxyData, bool ForceToRender = false);
 	virtual ~FNavMeshSceneProxy();
@@ -198,6 +199,10 @@ public:
 	static bool IsNavigationShowFlagSet(const UWorld* World);
 
 protected:
+	/** Gathers drawable information from NavMesh and puts it in OutProxyData. 
+	 *	Override to add additional information to OutProxyData.*/
+	virtual void GatherData(const ARecastNavMesh& NavMesh, FNavMeshSceneProxyData& OutProxyData) const;
+
 	void TimerFunction();
 
 protected:
