NavigationOctree: added additional data in FNavigationOctreeElement to recover from dangling entries.

#rnx

diff --git a/Engine/Source/Runtime/NavigationSystem/Private/NavigationOctree.cpp b/Engine/Source/Runtime/NavigationSystem/Private/NavigationOctree.cpp
--- a/Engine/Source/Runtime/NavigationSystem/Private/NavigationOctree.cpp
+++ b/Engine/Source/Runtime/NavigationSystem/Private/NavigationOctree.cpp
@@ -136,14 +136,6 @@ void FNavigationOctree::AddNode(UObject* ElementOb, INavRelevantInterface* NavEl
 	NodesMemory += ElementMemory;
 	INC_MEMORY_STAT_BY(STAT_Navigation_CollisionTreeMemory, ElementMemory);
 
-	if (ElementOb != nullptr)
-	{
-		ensure(ObjectToOctreeId.Find(HashObject(*ElementOb)) == nullptr);
-
-		// Reserve entry in our map before octree calls us back with the assigned Id. This is required since we ensure that all set id already exists
-		ObjectToOctreeId.Add(HashObject(*ElementOb), FOctreeElementId());
-	}	
-
 	AddElement(Element);
 }
 
@@ -206,17 +198,9 @@ const FNavigationRelevantData* FNavigationOctree::GetDataForID(const FOctreeElem
 	return &*OctreeElement.Data;
 }
 
-void FNavigationOctree::SetElementIdImpl(const UObject& Object, FOctreeElementId Id)
+void FNavigationOctree::SetElementIdImpl(const uint32 OwnerUniqueId, FOctreeElementId Id)
 {
-	ensure(IsValidElementId(Id));
-
-	FOctreeElementId* ExistingElementId = ObjectToOctreeId.Find(HashObject(Object));
-	// We expect ObjectToOctreeId to already contain an entry for Object since it should have already 
-	// been added in FNavigationOctree::AddNode
-	if (ensure(ExistingElementId))
-	{
-		*ExistingElementId = Id;
-	}
+	ObjectToOctreeId.Add(OwnerUniqueId, Id);
 }
 
 //----------------------------------------------------------------------//
@@ -227,10 +211,5 @@ FORCENOINLINE
 #endif // NAVSYS_DEBUG
 void FNavigationOctreeSemantics::SetElementId(FNavigationOctreeSemantics::FOctree& OctreeOwner, const FNavigationOctreeElement& Element, FOctreeElementId Id)
 {
-	const bool bEvenIfPendingKill = true;
-	UObject* ElementOwner = Element.GetOwner(bEvenIfPendingKill);
-	if (ensure(ElementOwner != nullptr))
-	{
-		((FNavigationOctree&)OctreeOwner).SetElementIdImpl(*ElementOwner, Id);
-	}
+	((FNavigationOctree&)OctreeOwner).SetElementIdImpl(Element.OwnerUniqueId, Id);
 }
diff --git a/Engine/Source/Runtime/NavigationSystem/Private/NavigationOctreeController.cpp b/Engine/Source/Runtime/NavigationSystem/Private/NavigationOctreeController.cpp
--- a/Engine/Source/Runtime/NavigationSystem/Private/NavigationOctreeController.cpp
+++ b/Engine/Source/Runtime/NavigationSystem/Private/NavigationOctreeController.cpp
@@ -27,7 +27,7 @@ void FNavigationOctreeController::SetNavigableGeometryStoringMode(FNavigationOct
 bool FNavigationOctreeController::GetNavOctreeElementData(const UObject& NodeOwner, int32& DirtyFlags, FBox& DirtyBounds)
 {
 	const FOctreeElementId* ElementId = GetObjectsNavOctreeId(NodeOwner);
-	if (ElementId != nullptr && ensure(IsValidElement(*ElementId)))
+	if (ElementId != nullptr && IsValidElement(*ElementId))
 	{
 		// mark area occupied by given actor as dirty
 		FNavigationOctreeElement& ElementData = NavOctree->GetElementById(*ElementId);
@@ -43,7 +43,7 @@ const FNavigationRelevantData* FNavigationOctreeController::GetDataForObject(con
 {
 	const FOctreeElementId* ElementId = GetObjectsNavOctreeId(Object);
 
-	if (ElementId != nullptr && ensure(IsValidElement(*ElementId)))
+	if (ElementId != nullptr && IsValidElement(*ElementId))
 	{
 		return NavOctree->GetDataForID(*ElementId);
 	}
diff --git a/Engine/Source/Runtime/NavigationSystem/Public/NavigationOctree.h b/Engine/Source/Runtime/NavigationSystem/Public/NavigationOctree.h
--- a/Engine/Source/Runtime/NavigationSystem/Public/NavigationOctree.h
+++ b/Engine/Source/Runtime/NavigationSystem/Public/NavigationOctree.h
@@ -20,15 +20,18 @@ struct NAVIGATIONSYSTEM_API FNavigationOctreeElement
 {
 	FBoxSphereBounds Bounds;
 	TSharedRef<FNavigationRelevantData, ESPMode::ThreadSafe> Data;
+	uint32 OwnerUniqueId = INDEX_NONE;	
 
 public:
 	explicit FNavigationOctreeElement(UObject& SourceObject)
 		: Data(new FNavigationRelevantData(SourceObject))
+		, OwnerUniqueId(SourceObject.GetUniqueID())
 	{}
 
 	FNavigationOctreeElement(const FNavigationOctreeElement& Other)
 		: Bounds(Other.Bounds)
 		, Data(Other.Data)
+		, OwnerUniqueId(Other.OwnerUniqueId)
 	{}
 
 	FNavigationOctreeElement& operator=(const FNavigationOctreeElement& Other)
@@ -163,7 +166,7 @@ protected:
 	friend struct FNavigationOctreeController;
 	friend struct FNavigationOctreeSemantics;
 
-	void SetElementIdImpl(const UObject& Object, FOctreeElementId Id);
+	void SetElementIdImpl(const uint32 OwnerUniqueId, FOctreeElementId Id);
 
 	TMap<uint32, FOctreeElementId> ObjectToOctreeId;
 	ENavDataGatheringMode DefaultGeometryGatheringMode;
