RecastNavMeshGenerator: Reworked fix from CL 7419221 to not mark NavBounds dirty for newly created navmesh when active tiles generation is enabled
in the navigation sytem without setting flag bRestrictBuildingToActiveTiles yet. This is required for the initial virtual call to RestrictBuildingToActiveTiles to perform the ActiveTiles initialization.

#rnx

diff --git a/Engine/Source/Runtime/NavigationSystem/Private/NavMesh/RecastNavMeshGenerator.cpp b/Engine/Source/Runtime/NavigationSystem/Private/NavMesh/RecastNavMeshGenerator.cpp
--- a/Engine/Source/Runtime/NavigationSystem/Private/NavMesh/RecastNavMeshGenerator.cpp
+++ b/Engine/Source/Runtime/NavigationSystem/Private/NavMesh/RecastNavMeshGenerator.cpp
@@ -3381,11 +3381,10 @@ void FRecastNavMeshGenerator::Init()
 
 		// mark all the areas we need to update, which is the whole (known) navigable space if not restricted to active tiles
 		const UNavigationSystemV1* NavSys = FNavigationSystem::GetCurrent<UNavigationSystemV1>(GetWorld());
-		if (NavSys)
+		if (NavSys && NavSys->IsActiveTilesGenerationEnabled() == false)
 		{
-			bRestrictBuildingToActiveTiles = NavSys->IsActiveTilesGenerationEnabled();
+			MarkNavBoundsDirty();
 		}
-		MarkNavBoundsDirty();
 	}
 	else
 	{
