NavigationDataHandler: prevent registration of an element that should not be registered when processing the parent chain.
One problematic scenario is when all components from an actor are unregistered and one of them is the navigation parent of the others. In this case, the parent may be unregistered before some of its children and then calls to UpdateNavOctreeParentChain from a child may queue back the unregistered parent. UpdateNavOctreeParentChain will now consider the actual status of the parent (Added or queued for add) before registering it or its children.

diff --git a/Engine/Source/Runtime/NavigationSystem/Private/NavigationDataHandler.cpp b/Engine/Source/Runtime/NavigationSystem/Private/NavigationDataHandler.cpp
--- a/Engine/Source/Runtime/NavigationSystem/Private/NavigationDataHandler.cpp
+++ b/Engine/Source/Runtime/NavigationSystem/Private/NavigationDataHandler.cpp
@@ -153,13 +153,13 @@ void FNavigationDataHandler::AddElementToNavOctree(const FNavigationDirtyElement
 	}
 }
 
-void FNavigationDataHandler::UnregisterNavOctreeElement(UObject& ElementOwner, INavRelevantInterface& ElementInterface, int32 UpdateFlags)
+bool FNavigationDataHandler::UnregisterNavOctreeElement(UObject& ElementOwner, INavRelevantInterface& ElementInterface, int32 UpdateFlags)
 {
 	SCOPE_CYCLE_COUNTER(STAT_Navigation_UnregisterNavOctreeElement);
 
 	if (OctreeController.NavOctree.IsValid() == false)
 	{
-		return;
+		return false;
 	}
 
 	if (OctreeController.IsNavigationOctreeLocked())
@@ -167,9 +167,10 @@ void FNavigationDataHandler::UnregisterNavOctreeElement(UObject& ElementOwner, I
 #if !WITH_EDITOR
 		UE_LOG(LogNavOctree, Log, TEXT("IGNORE(UnregisterNavOctreeElement) %s"), *GetPathNameSafe(&ElementOwner));
 #endif // WITH_EDITOR
-		return;
+		return false;
 	}
 
+	bool bUnregistered = false;
 	const FOctreeElementId* ElementId = OctreeController.GetObjectsNavOctreeId(ElementOwner);
 	UE_LOG(LogNavOctree, Log, TEXT("UNREG %s %s"), *ElementOwner.GetName(), ElementId ? TEXT("[exists]") : TEXT("[does\'t exist]"));
 
@@ -177,6 +178,7 @@ void FNavigationDataHandler::UnregisterNavOctreeElement(UObject& ElementOwner, I
 	{
 		RemoveNavOctreeElementId(*ElementId, UpdateFlags);
 		OctreeController.RemoveObjectsNavOctreeId(ElementOwner);
+		bUnregistered = true;
 	}
 	else
 	{
@@ -200,9 +202,17 @@ void FNavigationDataHandler::UnregisterNavOctreeElement(UObject& ElementOwner, I
 		const FSetElementId RequestId = OctreeController.PendingOctreeUpdates.FindId(FNavigationDirtyElement(&ElementOwner));
 		if (RequestId.IsValidId())
 		{
-			OctreeController.PendingOctreeUpdates[RequestId].bInvalidRequest = true;
+			FNavigationDirtyElement& DirtyElement = OctreeController.PendingOctreeUpdates[RequestId];
+
+			// Only consider as unregistered when pending update was not already invalidated since return value must indicate
+			// that ElementOwner was fully added or about to be added (valid pending update).
+			bUnregistered |= (DirtyElement.bInvalidRequest == false);
+
+			DirtyElement.bInvalidRequest = true;
 		}
 	}
+
+	return bUnregistered;
 }
 
 void FNavigationDataHandler::UpdateNavOctreeElement(UObject& ElementOwner, INavRelevantInterface& ElementInterface, int32 UpdateFlags)
@@ -258,20 +268,37 @@ void FNavigationDataHandler::UpdateNavOctreeParentChain(UObject& ElementOwner, b
 	TArray<FWeakObjectPtr> ChildNodes;
 	OctreeController.OctreeChildNodesMap.MultiFind(&ElementOwner, ChildNodes);
 
-	if (ChildNodes.Num() == 0)
+	auto ElementOwnerUpdateFunc = [&]()->bool
 	{
+		bool bShouldRegisterChildren = true;
 		if (bSkipElementOwnerUpdate == false)
 		{
 			INavRelevantInterface* ElementInterface = Cast<INavRelevantInterface>(&ElementOwner);
-			if (ElementInterface)
+			if (ElementInterface != nullptr)
 			{
-				UpdateNavOctreeElement(ElementOwner, *ElementInterface, UpdateFlags);
+				// We don't want to register NavOctreeElement if owner was not already registered or queued
+				// so we use Unregister/Register combo instead of UpdateNavOctreeElement
+				if (UnregisterNavOctreeElement(ElementOwner, *ElementInterface, UpdateFlags))
+				{
+					FSetElementId NewId = RegisterNavOctreeElement(ElementOwner, *ElementInterface, UpdateFlags);
+					bShouldRegisterChildren = NewId.IsValidId();
+				}
+				else
+				{
+					bShouldRegisterChildren = false;
+				}
 			}
 		}
+		return bShouldRegisterChildren;		
+	};
+
+	if (ChildNodes.Num() == 0)
+	{
+		// Last child was removed, only need to rebuild owner's NavOctreeElement 
+		ElementOwnerUpdateFunc();
 		return;
 	}
 
-	INavRelevantInterface* ElementInterface = Cast<INavRelevantInterface>(&ElementOwner);
 	TArray<INavRelevantInterface*> ChildNavInterfaces;
 	ChildNavInterfaces.AddZeroed(ChildNodes.Num());
 
@@ -288,18 +315,17 @@ void FNavigationDataHandler::UpdateNavOctreeParentChain(UObject& ElementOwner, b
 		}
 	}
 
-	if (bSkipElementOwnerUpdate == false && ElementInterface)
-	{
-		UnregisterNavOctreeElement(ElementOwner, *ElementInterface, UpdateFlags);
-		RegisterNavOctreeElement(ElementOwner, *ElementInterface, UpdateFlags);
-	}
-
-	for (int32 Idx = 0; Idx < ChildNodes.Num(); Idx++)
+	const bool bShouldRegisterChildren = ElementOwnerUpdateFunc();
+	
+	if (bShouldRegisterChildren)
 	{
-		UObject* ChildElement = ChildNodes[Idx].Get();
-		if (ChildElement && ChildNavInterfaces[Idx])
+		for (int32 Idx = 0; Idx < ChildNodes.Num(); Idx++)
 		{
-			RegisterNavOctreeElement(*ChildElement, *ChildNavInterfaces[Idx], UpdateFlags);
+			UObject* ChildElement = ChildNodes[Idx].Get();
+			if (ChildElement && ChildNavInterfaces[Idx])
+			{
+				RegisterNavOctreeElement(*ChildElement, *ChildNavInterfaces[Idx], UpdateFlags);
+			}
 		}
 	}
 }
diff --git a/Engine/Source/Runtime/NavigationSystem/Public/NavigationDataHandler.h b/Engine/Source/Runtime/NavigationSystem/Public/NavigationDataHandler.h
--- a/Engine/Source/Runtime/NavigationSystem/Public/NavigationDataHandler.h
+++ b/Engine/Source/Runtime/NavigationSystem/Public/NavigationDataHandler.h
@@ -15,7 +15,16 @@ struct NAVIGATIONSYSTEM_API FNavigationDataHandler
 	void RemoveNavOctreeElementId(const FOctreeElementId& ElementId, int32 UpdateFlags);
 	FSetElementId RegisterNavOctreeElement(UObject& ElementOwner, INavRelevantInterface& ElementInterface, int32 UpdateFlags);
 	void AddElementToNavOctree(const FNavigationDirtyElement& DirtyElement);
-	void UnregisterNavOctreeElement(UObject& ElementOwner, INavRelevantInterface& ElementInterface, int32 UpdateFlags);
+
+	/** Removes associated NavOctreeElement and invalidates associated pending updates. Also removes object from the list of children
+	* of the NavigationParent, if any.
+	* @param ElementOwner		Object for which we must remove associated NavOctreeElement
+	* @param ElementInterface	Object associated NavRelevantInterface to access NavigationParent when registered as children
+	* @param UpdateFlags		Flags indicating in which context the method is called to allow/forbid certain operations
+	*
+	* @return True if associated NavOctreeElement has been removed or pending update has been invalidated; false otherwise.
+	*/
+	bool UnregisterNavOctreeElement(UObject& ElementOwner, INavRelevantInterface& ElementInterface, int32 UpdateFlags);
 	void UpdateNavOctreeElement(UObject& ElementOwner, INavRelevantInterface& ElementInterface, int32 UpdateFlags);
 	void UpdateNavOctreeParentChain(UObject& ElementOwner, bool bSkipElementOwnerUpdate = false);
 	bool UpdateNavOctreeElementBounds(UActorComponent& Comp, const FBox& NewBounds, const FBox& DirtyArea);	
