Hotfix from CL 8696491 for UE-79923: Show Navigation does not work

diff --git a/Engine/Source/Runtime/NavigationSystem/Private/NavigationSystem.cpp b/Engine/Source/Runtime/NavigationSystem/Private/NavigationSystem.cpp
--- a/Engine/Source/Runtime/NavigationSystem/Private/NavigationSystem.cpp
+++ b/Engine/Source/Runtime/NavigationSystem/Private/NavigationSystem.cpp
@@ -561,11 +561,11 @@ void UNavigationSystemV1::PostInitProperties()
 		}
 
 		ApplySupportedAgentsFilter();
-			for (int32 AgentIndex = 0; AgentIndex < SupportedAgents.Num(); ++AgentIndex)
-			{
-				FNavDataConfig& SupportedAgentConfig = SupportedAgents[AgentIndex];
-				SetSupportedAgentsNavigationClass(AgentIndex, SupportedAgentConfig.GetNavDataClass<ANavigationData>());
-			}
+		for (int32 AgentIndex = 0; AgentIndex < SupportedAgents.Num(); ++AgentIndex)
+		{
+			FNavDataConfig& SupportedAgentConfig = SupportedAgents[AgentIndex];
+			SetSupportedAgentsNavigationClass(AgentIndex, SupportedAgentConfig.GetNavDataClass<ANavigationData>());
+		}
 	
 		if (bInitialBuildingLocked)
 		{
@@ -4335,6 +4335,25 @@ void UNavigationSystemV1::Configure(const UNavigationSystemConfig& Config)
 		DefaultAgentName = Config.DefaultAgentName;
 	}
 	SetSupportedAgentsMask(Config.SupportedAgentsMask);
+
+	if (DefaultAgentName == NAME_None)
+	{
+		if (SupportedAgents.Num() == 1)
+		{
+			DefaultAgentName = SupportedAgents[0].Name;
+		}
+		else // pick the first available one
+		{
+			for (int32 AgentIndex = 0; AgentIndex < SupportedAgents.Num(); ++AgentIndex)
+			{
+				if (SupportedAgents[AgentIndex].IsValid())
+				{
+					DefaultAgentName = SupportedAgents[AgentIndex].Name;
+					break;
+				}
+			}
+		}
+	}
 }
 
 void UNavigationSystemV1::AppendConfig(const UNavigationSystemConfig& NewConfig)
