Added a way to query NavigationSystemV1 for NavigationData corresponding to a given agent name

While modifying NavigationSystemV1 I've extended IsNavigationBuildingLocked with a parameter allowing more precise queries.

diff --git a/Engine/Source/Runtime/NavigationSystem/Private/NavigationSystem.cpp b/Engine/Source/Runtime/NavigationSystem/Private/NavigationSystem.cpp
--- a/Engine/Source/Runtime/NavigationSystem/Private/NavigationSystem.cpp
+++ b/Engine/Source/Runtime/NavigationSystem/Private/NavigationSystem.cpp
@@ -1544,6 +1544,22 @@ const ANavigationData* UNavigationSystemV1::GetNavDataForProps(const FNavAgentPr
 	return NavDataInstance ? NavDataInstance : MainNavData;
 }
 
+ANavigationData* UNavigationSystemV1::GetNavDataForAgentName(const FName AgentName) const
+{
+	ANavigationData* Result = nullptr;
+
+	for (ANavigationData* NavData : NavDataSet)
+	{
+		if (NavData && !NavData->IsPendingKill() && NavData->GetConfig().Name == AgentName)
+		{
+			Result = NavData;
+			break;
+		}
+	}
+
+	return Result;
+}
+
 ANavigationData* UNavigationSystemV1::GetDefaultNavDataInstance(FNavigationSystem::ECreateIfMissing CreateNewIfNoneFound)
 {
 	checkSlow(IsInGameThread() == true);
diff --git a/Engine/Source/Runtime/NavigationSystem/Public/NavigationSystem.h b/Engine/Source/Runtime/NavigationSystem/Public/NavigationSystem.h
--- a/Engine/Source/Runtime/NavigationSystem/Public/NavigationSystem.h
+++ b/Engine/Source/Runtime/NavigationSystem/Public/NavigationSystem.h
@@ -457,6 +457,10 @@ public:
 	 */
 	virtual const ANavigationData* GetNavDataForProps(const FNavAgentProperties& AgentProperties) const;
 
+	/** Goes through all registered NavigationData instances and retrieves the one 
+	 *	supporting agent named AgentName */
+	virtual ANavigationData* GetNavDataForAgentName(const FName AgentName) const;
+
 	/**
 	 * Looks up NavData appropriate for specified movement properties and returns it. NULL if not found;
 	 * This is the encouraged way of querying for the appropriate NavData. It makes no difference for NavigationSystemV1
@@ -709,7 +713,7 @@ public:
 	void OnPIEEnd();
 	
 	// @todo document
-	FORCEINLINE bool IsNavigationBuildingLocked() const { return NavBuildingLockFlags != 0; }
+	FORCEINLINE bool IsNavigationBuildingLocked(uint8 Flags = ~0) const { return (NavBuildingLockFlags & Flags) != 0; }
 
 	/** check if building is permanently locked to avoid showing navmesh building notify (due to queued dirty areas) */
 	FORCEINLINE bool IsNavigationBuildingPermanentlyLocked() const
