Added a way to suspend navigation generation of a single NavigationData instance (as opposed to locking the whole NavigationSystem).

While modifying NavigationData.cpp I've refactored all UE_LOG into UE_VLOG_UELOG for better vislog coverage

diff --git a/Engine/Source/Runtime/NavigationSystem/Private/NavigationData.cpp b/Engine/Source/Runtime/NavigationSystem/Private/NavigationData.cpp
--- a/Engine/Source/Runtime/NavigationSystem/Private/NavigationData.cpp
+++ b/Engine/Source/Runtime/NavigationSystem/Private/NavigationData.cpp
@@ -134,6 +134,7 @@ ANavigationData::ANavigationData(const FObjectInitializer& ObjectInitializer)
 	, FindPathImplementation(NULL)
 	, FindHierarchicalPathImplementation(NULL)
 	, bRegistered(false)
+	, bRebuildingSuspended(false)
 	, NavDataUniqueID(GetNextUniqueID())
 {
 	PrimaryActorTick.bCanEverTick = true;
@@ -190,7 +191,7 @@ void ANavigationData::PostInitializeComponents()
 		(MyWorld->GetNetMode() != NM_Client && NavSys == nullptr) ||
 		(MyWorld->GetNetMode() == NM_Client && !bNetLoadOnClient))
 	{
-		UE_LOG(LogNavigation, Log, TEXT("Marking %s as PendingKill due to %s"), *GetName()
+		UE_VLOG_UELOG(this, LogNavigation, Log, TEXT("Marking %s as PendingKill due to %s"), *GetName()
 			, !MyWorld ? TEXT("No World") : (MyWorld->GetNetMode() == NM_Client ? TEXT("not creating navigation on clients") : TEXT("missing navigation system")));
 		CleanUpAndMarkPendingKill();
 	}
@@ -549,9 +550,35 @@ void ANavigationData::TickAsyncBuild(float DeltaSeconds)
 
 void ANavigationData::RebuildDirtyAreas(const TArray<FNavigationDirtyArea>& DirtyAreas)
 {
-	if (NavDataGenerator.IsValid())
+	if (bRebuildingSuspended)
+	{
+		SuspendedDirtyAreas.Append(DirtyAreas);
+	}
+	else
+	{
+		if (NavDataGenerator.IsValid())
+		{
+			NavDataGenerator->RebuildDirtyAreas(DirtyAreas);
+		}
+	}
+}
+
+void ANavigationData::SetRebuildingSuspended(const bool bNewSuspend)
+{
+	if (bRebuildingSuspended != bNewSuspend)
 	{
-		NavDataGenerator->RebuildDirtyAreas(DirtyAreas);
+		bRebuildingSuspended = bNewSuspend;
+		UE_VLOG_UELOG(this, LogNavigation, Verbose, TEXT("%s nav generation %s")
+			, *GetName(), bNewSuspend ? TEXT("SUSPENDED") : TEXT("ACTIVE"));
+
+		if (bNewSuspend == false && SuspendedDirtyAreas.Num() > 0)
+		{
+			UE_VLOG_UELOG(this, LogNavigation, Verbose, TEXT("%s resuming nav generation with %d dirty areas")
+				, *GetName(), SuspendedDirtyAreas.Num());
+			// resuming the generation so we need to utilize SuspendedDirtyAreas and clean it
+			RebuildDirtyAreas(SuspendedDirtyAreas);
+			SuspendedDirtyAreas.Empty();
+		}
 	}
 }
 
@@ -602,7 +629,7 @@ void ANavigationData::OnNavAreaAdded(const UClass* NavAreaClass, int32 AgentInde
 	const bool bIsMetaArea = DefArea != nullptr && DefArea->IsMetaArea();
 	if (!DefArea || bIsMetaArea || !DefArea->IsSupportingAgent(AgentIndex))
 	{
-		UE_LOG(LogNavigation, Verbose, TEXT("%s discarded area %s (valid:%s meta:%s validAgent[%d]:%s)"),
+		UE_VLOG_UELOG(this, LogNavigation, Verbose, TEXT("%s discarded area %s (valid:%s meta:%s validAgent[%d]:%s)"),
 			*GetName(), *GetNameSafe(NavAreaClass),
 			DefArea ? TEXT("yes") : TEXT("NO"),
 			bIsMetaArea ? TEXT("YES") : TEXT("no"),
@@ -618,7 +645,7 @@ void ANavigationData::OnNavAreaAdded(const UClass* NavAreaClass, int32 AgentInde
 		{
 			SupportedAreas[i].AreaClass = NavAreaClass;
 			AreaClassToIdMap.Add(NavAreaClass, SupportedAreas[i].AreaID);
-			UE_LOG(LogNavigation, Verbose, TEXT("%s updated area %s with ID %d"), *GetName(), *AreaClassName, SupportedAreas[i].AreaID);
+			UE_VLOG_UELOG(this, LogNavigation, Verbose, TEXT("%s updated area %s with ID %d"), *GetName(), *AreaClassName, SupportedAreas[i].AreaID);
 			return;
 		}
 	}
@@ -627,7 +654,7 @@ void ANavigationData::OnNavAreaAdded(const UClass* NavAreaClass, int32 AgentInde
 	const int32 MaxSupported = GetMaxSupportedAreas();
 	if (SupportedAreas.Num() >= MaxSupported)
 	{
-		UE_LOG(LogNavigation, Error, TEXT("%s can't support area %s - limit reached! (%d)"), *GetName(), *AreaClassName, MaxSupported);
+		UE_VLOG_UELOG(this, LogNavigation, Error, TEXT("%s can't support area %s - limit reached! (%d)"), *GetName(), *AreaClassName, MaxSupported);
 		return;
 	}
 
@@ -638,7 +665,7 @@ void ANavigationData::OnNavAreaAdded(const UClass* NavAreaClass, int32 AgentInde
 	SupportedAreas.Add(NewAgentData);
 	AreaClassToIdMap.Add(NavAreaClass, NewAgentData.AreaID);
 
-	UE_LOG(LogNavigation, Verbose, TEXT("%s registered area %s with ID %d"), *GetName(), *AreaClassName, NewAgentData.AreaID);
+	UE_VLOG_UELOG(this, LogNavigation, Verbose, TEXT("%s registered area %s with ID %d"), *GetName(), *AreaClassName, NewAgentData.AreaID);
 }
 
 void ANavigationData::OnNavAreaEvent(const UClass* NavAreaClass, ENavAreaEvent::Type Event)
@@ -771,7 +798,7 @@ uint32 ANavigationData::LogMemUsed() const
 	const uint32 MemUsed = ActivePaths.GetAllocatedSize() + SupportedAreas.GetAllocatedSize() +
 		QueryFilters.GetAllocatedSize() + AreaClassToIdMap.GetAllocatedSize();
 
-	UE_LOG(LogNavigation, Display, TEXT("%s: ANavigationData: %u\n    self: %d"), *GetName(), MemUsed, sizeof(ANavigationData));	
+	UE_VLOG_UELOG(this, LogNavigation, Display, TEXT("%s: ANavigationData: %u\n    self: %d"), *GetName(), MemUsed, sizeof(ANavigationData));
 
 	if (NavDataGenerator.IsValid())
 	{
diff --git a/Engine/Source/Runtime/NavigationSystem/Public/NavigationData.h b/Engine/Source/Runtime/NavigationSystem/Public/NavigationData.h
--- a/Engine/Source/Runtime/NavigationSystem/Public/NavigationData.h
+++ b/Engine/Source/Runtime/NavigationSystem/Public/NavigationData.h
@@ -625,6 +625,16 @@ public:
 
 	/** Request navigation data update after changes in nav octree */
 	virtual void RebuildDirtyAreas(const TArray<FNavigationDirtyArea>& DirtyAreas);
+
+	/** Configures this NavData instance's navigation generation to be suspended 
+	 *	or active. It's active by default. If Suspended then all calls to 
+	 *	RebuildDirtyAreas will result in caching the request in SuspendedDirtyAreas 
+	 *	until SetRebuildingSuspended(false) gets call at which time all the contents 
+	 *	of SuspendedDirtyAreas will get pushed to the nav generator and SuspendedDirtyAreas 
+	 *	will be cleaned out. 
+	 *	Note that calling SetRebuildingSuspended(true) won't suspend the nav generation 
+	 *	already in progress. */
+	virtual void SetRebuildingSuspended(const bool bNewSuspend);
 	
 	/** releases navigation generator if any has been created */
 protected:
@@ -917,6 +927,13 @@ protected:
 
 protected:
 	TSharedPtr<FNavDataGenerator, ESPMode::ThreadSafe> NavDataGenerator;
+
+	/** caches requests to rebuild dirty areas while nav rebuilding is suspended 
+	 *	via SetRebuildingSuspended(true) call. Calling SetRebuildingSuspended(false) 
+	 *	will result in pushing SuspendedDirtyAreas contents to the nav generator 
+	 *	and clearing out of SuspendedDirtyAreas */
+	TArray<FNavigationDirtyArea> SuspendedDirtyAreas;
+
 	/** 
 	 *	Container for all path objects generated with this Navigation Data instance. 
 	 *	Is meant to be added to only on GameThread, and in fact should user should never 
@@ -954,6 +971,11 @@ protected:
 	/** was it generated for default agent (SupportedAgents[0]) */
 	uint32 bSupportsDefaultAgent : 1;
 
+	/** Set via SetRebuildingSuspended and controlling if RebuildDirtyAreas get 
+	 *	passed over to the generator instantly or cached in SuspendedDirtyAreas 
+	 *	to be applied at later date with SetRebuildingSuspended(false) call */
+	uint32 bRebuildingSuspended : 1;
+
 private:
 	uint16 NavDataUniqueID;
 
