Added a getter for UNavigationSystemV1::SupportedAgentsMask

diff --git a/Engine/Source/Runtime/NavigationSystem/Public/NavigationSystem.h b/Engine/Source/Runtime/NavigationSystem/Public/NavigationSystem.h
--- a/Engine/Source/Runtime/NavigationSystem/Public/NavigationSystem.h
+++ b/Engine/Source/Runtime/NavigationSystem/Public/NavigationSystem.h
@@ -514,6 +514,7 @@ public:
 	FORCEINLINE const TArray<FNavDataConfig>& GetSupportedAgents() const { return SupportedAgents; }
 	void OverrideSupportedAgents(const TArray<FNavDataConfig>& NewSupportedAgents);
 	void SetSupportedAgentsMask(const FNavAgentSelector& InSupportedAgentsMask);
+	FNavAgentSelector GetSupportedAgentsMask() const { return SupportedAgentsMask; }
 
 	virtual void ApplyWorldOffset(const FVector& InOffset, bool bWorldShift) override;
 
