Added a way to mark a NavigationData instance to not get auto-removed when there's no navigation system present

diff --git a/Engine/Source/Runtime/NavigationSystem/Private/NavigationData.cpp b/Engine/Source/Runtime/NavigationSystem/Private/NavigationData.cpp
--- a/Engine/Source/Runtime/NavigationSystem/Private/NavigationData.cpp
+++ b/Engine/Source/Runtime/NavigationSystem/Private/NavigationData.cpp
@@ -127,6 +127,7 @@ ANavigationData::ANavigationData(const FObjectInitializer& ObjectInitializer)
 	: Super(ObjectInitializer)
 	, bEnableDrawing(false)
 	, bForceRebuildOnLoad(false)
+	, bAutoDestroyWhenNoNavigation(true)
 	, bCanBeMainNavData(true)
 	, bCanSpawnOnRebuild(true)
 	, RuntimeGeneration(ERuntimeGenerationType::LegacyGeneration) //TODO: set to a valid value once bRebuildAtRuntime_DEPRECATED is removed
@@ -187,9 +188,9 @@ void ANavigationData::PostInitializeComponents()
 	UWorld* MyWorld = GetWorld();
 	UNavigationSystemV1* NavSys = FNavigationSystem::GetCurrent<UNavigationSystemV1>(MyWorld);
 
-	if (MyWorld == nullptr ||
+	if (bAutoDestroyWhenNoNavigation && (MyWorld == nullptr ||
 		(MyWorld->GetNetMode() != NM_Client && NavSys == nullptr) ||
-		(MyWorld->GetNetMode() == NM_Client && !bNetLoadOnClient))
+		(MyWorld->GetNetMode() == NM_Client && !bNetLoadOnClient)))
 	{
 		UE_VLOG_UELOG(this, LogNavigation, Log, TEXT("Marking %s as PendingKill due to %s"), *GetName()
 			, !MyWorld ? TEXT("No World") : (MyWorld->GetNetMode() == NM_Client ? TEXT("not creating navigation on clients") : TEXT("missing navigation system")));
diff --git a/Engine/Source/Runtime/NavigationSystem/Public/NavigationData.h b/Engine/Source/Runtime/NavigationSystem/Public/NavigationData.h
--- a/Engine/Source/Runtime/NavigationSystem/Public/NavigationData.h
+++ b/Engine/Source/Runtime/NavigationSystem/Public/NavigationData.h
@@ -515,6 +515,11 @@ protected:
 	UPROPERTY(config, EditAnywhere, Category = Runtime)
 	uint32 bForceRebuildOnLoad : 1;
 
+	/** Should this instance auto-destroy when there's no navigation system on
+	 *	world when it gets created/loaded */
+	UPROPERTY(config, EditAnywhere, Category = Runtime)
+	uint32 bAutoDestroyWhenNoNavigation : 1;
+
 	/** If set, navigation data can act as default one in navigation system's queries */
 	UPROPERTY(config, EditAnywhere, Category = Runtime, AdvancedDisplay)
 	uint32 bCanBeMainNavData : 1;
@@ -633,7 +638,10 @@ public:
 	 *	of SuspendedDirtyAreas will get pushed to the nav generator and SuspendedDirtyAreas 
 	 *	will be cleaned out. 
 	 *	Note that calling SetRebuildingSuspended(true) won't suspend the nav generation 
-	 *	already in progress. */
+	 *	already in progress.
+	 *	Note2: due to all areas dirtied during generation suspension ending up in 
+	 *	SuspendedDirtyAreas care needs to be taken to not use this feature for 
+	 *	extended periods of time - otherwise SuspendedDirtyAreas can get very large. */
 	virtual void SetRebuildingSuspended(const bool bNewSuspend);
 	
 	/** releases navigation generator if any has been created */
