Added a function to UNavigationSystemV1 that unregisteres all NavigationData instnces no longer supported by the nav system (as contolled by SupportedAgentsMask)

diff --git a/Engine/Source/Runtime/NavigationSystem/Private/NavigationSystem.cpp b/Engine/Source/Runtime/NavigationSystem/Private/NavigationSystem.cpp
--- a/Engine/Source/Runtime/NavigationSystem/Private/NavigationSystem.cpp
+++ b/Engine/Source/Runtime/NavigationSystem/Private/NavigationSystem.cpp
@@ -4332,12 +4332,28 @@ void UNavigationSystemV1::ApplySupportedAgentsFilter()
 	for (int32 AgentIndex = 0; AgentIndex < SupportedAgents.Num(); AgentIndex++)
 	{
 		if (SupportedAgentsMask.Contains(AgentIndex) == false)
-	{
+		{
 			SupportedAgents[AgentIndex].Invalidate();
 		}
 	}
 }
 
+void UNavigationSystemV1::UnregisterUnusedNavData()
+{
+	for (int32 AgentIndex = 0; AgentIndex < SupportedAgents.Num(); AgentIndex++)
+	{
+		if (SupportedAgentsMask.Contains(AgentIndex) == false)
+		{
+			// if we already have navdata for this agent we need to remove it
+			ANavigationData* NavData = GetNavDataForAgentName(SupportedAgents[AgentIndex].Name);
+			if (NavData)
+			{
+				UnregisterNavData(NavData);
+			}
+		}
+	}
+}
+
 void UNavigationSystemV1::SetSupportedAgentsMask(const FNavAgentSelector& InSupportedAgentsMask)
 {
 	SupportedAgentsMask = InSupportedAgentsMask;
diff --git a/Engine/Source/Runtime/NavigationSystem/Public/NavigationSystem.h b/Engine/Source/Runtime/NavigationSystem/Public/NavigationSystem.h
--- a/Engine/Source/Runtime/NavigationSystem/Public/NavigationSystem.h
+++ b/Engine/Source/Runtime/NavigationSystem/Public/NavigationSystem.h
@@ -547,6 +547,11 @@ public:
 	// @todo document
 	virtual void UnregisterNavData(ANavigationData* NavData);
 
+	/** Traverses SupportedAgents and for all agents not supported (i.e. filtered
+	 *	out by SupportedAgentsMask) checks if there's a currently registered
+	 *	NavigationData instance for that agent, and if so it unregisters that agent */
+	virtual void UnregisterUnusedNavData();
+
 	/** Adds NavData to registration candidates queue - NavDataRegistrationQueue*/
 	virtual void RequestRegistrationDeferred(ANavigationData& NavData);
 	UE_DEPRECATED(4.24, "This version of RequestRegistration is deprecated. Please use the RequestRegistrationDeferred as the registration request is always queued now.")
