Changed when NavigationSystemConfigOverride initiates nav system creation at runtime

This change is needed in order to ensure that we create navigation system via the nav config override while outside of async-loading logic. Otherwise we severly limit navigation system instances' capabilities of loading assets (like NavArea BP classes).

diff --git a/Engine/Source/Runtime/NavigationSystem/Private/NavSystemConfigOverride.cpp b/Engine/Source/Runtime/NavigationSystem/Private/NavSystemConfigOverride.cpp
--- a/Engine/Source/Runtime/NavigationSystem/Private/NavSystemConfigOverride.cpp
+++ b/Engine/Source/Runtime/NavigationSystem/Private/NavSystemConfigOverride.cpp
@@ -64,8 +64,9 @@ void ANavSystemConfigOverride::PostLoad()
 {
 	Super::PostLoad();
 	
+#if WITH_EDITOR
 	UWorld* World = GetWorld();
-	if (World != nullptr)
+	if ((World != nullptr) && (World->IsGameWorld() == false))
 	{
 		if (World->bInTick)
 		{
@@ -76,6 +77,13 @@ void ANavSystemConfigOverride::PostLoad()
 			ApplyConfig();
 		}
 	}
+#endif // WITH_EDITOR
+}
+
+void ANavSystemConfigOverride::BeginPlay()
+{
+	Super::BeginPlay();
+	ApplyConfig();
 }
 
 void ANavSystemConfigOverride::ApplyConfig()
diff --git a/Engine/Source/Runtime/NavigationSystem/Public/NavSystemConfigOverride.h b/Engine/Source/Runtime/NavigationSystem/Public/NavSystemConfigOverride.h
--- a/Engine/Source/Runtime/NavigationSystem/Public/NavSystemConfigOverride.h
+++ b/Engine/Source/Runtime/NavigationSystem/Public/NavSystemConfigOverride.h
@@ -53,6 +53,10 @@ public:
 #endif // WITH_EDITOR
 	//~ End UObject Interface
 
+	//~ Begin AActor Interface
+	virtual void BeginPlay() override;
+	//~ End AActor Interface
+
 #if WITH_EDITOR
 	/** made an explicit function since rebuilding navigation system can be expensive */
 	UFUNCTION(Category = Navigation, meta = (CallInEditor = "true"))
