NavSystemConfigOverride: when loading config for editor world, we apply config from timer next tick to avoid warning when loading objects from PostLoad.

#rnx

diff --git a/Engine/Source/Runtime/NavigationSystem/Private/NavSystemConfigOverride.cpp b/Engine/Source/Runtime/NavigationSystem/Private/NavSystemConfigOverride.cpp
--- a/Engine/Source/Runtime/NavigationSystem/Private/NavSystemConfigOverride.cpp
+++ b/Engine/Source/Runtime/NavigationSystem/Private/NavSystemConfigOverride.cpp
@@ -68,14 +68,7 @@ void ANavSystemConfigOverride::PostLoad()
 	UWorld* World = GetWorld();
 	if ((World != nullptr) && (World->IsGameWorld() == false))
 	{
-		if (World->bInTick)
-		{
-			World->GetTimerManager().SetTimerForNextTick(this, &ANavSystemConfigOverride::ApplyConfig);
-		}
-		else
-		{
-			ApplyConfig();
-		}
+		GEditor->GetTimerManager()->SetTimerForNextTick(this, &ANavSystemConfigOverride::ApplyConfig);
 	}
 #endif // WITH_EDITOR
 }
@@ -195,4 +188,4 @@ void ANavSystemConfigOverride::PostEditChangeProperty(FPropertyChangedEvent& Pro
 
 	bNetLoadOnClient = bLoadOnClient;
 }
-#endif // WITH_EDITOR
\ No newline at end of file
+#endif // WITH_EDITOR
