Fixed a crash in navigation system when renaming NavigationData instance while there's already another on of that name. We've already had a safety feature for that but it turns out there might be multiple objects of the same name if Undo/Redo comes to play.

diff --git a/Engine/Source/Runtime/NavigationSystem/Private/NavigationSystem.cpp b/Engine/Source/Runtime/NavigationSystem/Private/NavigationSystem.cpp
--- a/Engine/Source/Runtime/NavigationSystem/Private/NavigationSystem.cpp
+++ b/Engine/Source/Runtime/NavigationSystem/Private/NavigationSystem.cpp
@@ -3216,7 +3216,7 @@ ANavigationData* UNavigationSystemV1::CreateNavigationDataInstanceInLevel(const
 			// temporary solution to make sure we don't try to change name while there's already
 			// an object with this name
 			UObject* ExistingObject = StaticFindObject(/*Class=*/ NULL, Instance->GetOuter(), *StrName, true);
-			if (ExistingObject != NULL)
+			while (ExistingObject != NULL)
 			{
 				ANavigationData* ExistingNavigationData = Cast<ANavigationData>(ExistingObject);
 				if (ExistingNavigationData)
@@ -3224,7 +3224,11 @@ ANavigationData* UNavigationSystemV1::CreateNavigationDataInstanceInLevel(const
 					UnregisterNavData(ExistingNavigationData);
 				}
 
+				// Reset the existing object's name
 				ExistingObject->Rename(NULL, NULL, REN_DontCreateRedirectors | REN_ForceGlobalUnique | REN_DoNotDirty | REN_NonTransactional | REN_ForceNoResetLoaders);
+				// see if there's another one, it does happen when undo/redoing 
+				// nav instance deletion in the editor
+				ExistingObject = StaticFindObject(/*Class=*/ NULL, Instance->GetOuter(), *StrName, true);
 			}
 
 			// Set descriptive name
