Fixed null-area nav links not being filtered out by dtQueryFilter::passInlineFilter

diff --git a/Engine/Source/Runtime/NavigationSystem/Private/NavAreas/NavArea_Null.cpp b/Engine/Source/Runtime/NavigationSystem/Private/NavAreas/NavArea_Null.cpp
--- a/Engine/Source/Runtime/NavigationSystem/Private/NavAreas/NavArea_Null.cpp
+++ b/Engine/Source/Runtime/NavigationSystem/Private/NavAreas/NavArea_Null.cpp
@@ -1,9 +1,10 @@
 // Copyright 1998-2019 Epic Games, Inc. All Rights Reserved.
 
 #include "NavAreas/NavArea_Null.h"
+#include "NavMesh/RecastNavMesh.h"
 
 UNavArea_Null::UNavArea_Null(const FObjectInitializer& ObjectInitializer) : Super(ObjectInitializer)
 {
-	DefaultCost = BIG_NUMBER;
+	DefaultCost = FLT_MAX;
 	AreaFlags = 0;
 }
diff --git a/Engine/Source/Runtime/NavigationSystem/Private/NavMesh/RecastNavMesh.cpp b/Engine/Source/Runtime/NavigationSystem/Private/NavMesh/RecastNavMesh.cpp
--- a/Engine/Source/Runtime/NavigationSystem/Private/NavMesh/RecastNavMesh.cpp
+++ b/Engine/Source/Runtime/NavigationSystem/Private/NavMesh/RecastNavMesh.cpp
@@ -488,6 +488,10 @@ void ARecastNavMesh::RecreateDefaultFilter()
 	FRecastQueryFilter* DetourFilter = static_cast<FRecastQueryFilter*>(DefaultQueryFilter->GetImplementation());
 	DetourFilter->SetIsVirtual(bUseVirtualFilters);
 	DetourFilter->setHeuristicScale(HeuristicScale);
+	// clearing out the 'navlink flag' from included flags since it would make 
+	// dtQueryFilter::passInlineFilter pass navlinks of area classes with
+	// AreaFlags == 0 (like NavArea_Null), which should mean 'unwalkable'
+	DetourFilter->setIncludeFlags(DetourFilter->getIncludeFlags() & (~ARecastNavMesh::GetNavLinkFlag()));
 
 	for (int32 Idx = 0; Idx < SupportedAreas.Num(); Idx++)
 	{
