NavigationTestingActor -> Addition of a MaxCostFactor as a parameter for path searchs
NavigationTestingActor -> Fix path reset issue when toggling bSearchStart

diff --git a/Engine/Source/Runtime/NavigationSystem/Private/NavigationTestingActor.cpp b/Engine/Source/Runtime/NavigationSystem/Private/NavigationTestingActor.cpp
--- a/Engine/Source/Runtime/NavigationSystem/Private/NavigationTestingActor.cpp
+++ b/Engine/Source/Runtime/NavigationSystem/Private/NavigationTestingActor.cpp
@@ -41,6 +41,7 @@ ANavigationTestingActor::ANavigationTestingActor(const FObjectInitializer& Objec
 	bIsEditorOnlyActor = true;
 	NavAgentProps.AgentRadius = 34.f;
 	NavAgentProps.AgentHeight = 144.f;
+	MaxCostFactor = FLT_MAX;
 	ShowStepIndex = -1;
 	bShowNodePool = true;
 	bShowBestPath = true;
@@ -200,6 +201,17 @@ void ANavigationTestingActor::PostEditChangeProperty(FPropertyChangedEvent& Prop
 				if (OtherActor != NULL)
 				{
 					OtherActor->bSearchStart = !bSearchStart;
+					OtherActor->LastPath.Reset();
+					if (OtherActor->EdRenderComp)
+					{
+						OtherActor->EdRenderComp->MarkRenderStateDirty();
+					}
+				}
+
+				LastPath.Reset();
+				if (EdRenderComp)
+				{
+					EdRenderComp->MarkRenderStateDirty();
 				}
 			}
 
@@ -411,14 +423,16 @@ void ANavigationTestingActor::SearchPathTo(ANavigationTestingActor* Goal)
 
 	if (bBacktracking)
 	{
-		FSharedConstNavQueryFilter NavQueryFilter = Query.QueryFilter
-			? Query.QueryFilter
-			: NavData->GetDefaultQueryFilter();
-
+		FSharedConstNavQueryFilter NavQueryFilter = Query.QueryFilter ? Query.QueryFilter : NavData->GetDefaultQueryFilter();
 		FSharedNavQueryFilter NavigationFilterCopy = NavQueryFilter->GetCopy();
 		NavigationFilterCopy->SetBacktrackingEnabled(true);
 		Query.QueryFilter = NavigationFilterCopy;
 	}
+	
+	//Apply max cost factor
+	FSharedConstNavQueryFilter NavQueryFilter = Query.QueryFilter ? Query.QueryFilter : NavData->GetDefaultQueryFilter();
+	const float HeuristicScale = NavQueryFilter->GetHeuristicScale();
+	Query.MaxCost = Query.ComputeMaxCostFromHeuristic(Query.StartLocation, Query.EndLocation, HeuristicScale, MaxCostFactor);
 
 	EPathFindingMode::Type Mode = bUseHierarchicalPathfinding ? EPathFindingMode::Hierarchical : EPathFindingMode::Regular;
 	FPathFindingResult Result = NavSys->FindPathSync(NavAgentProps, Query, Mode);
diff --git a/Engine/Source/Runtime/NavigationSystem/Public/NavigationTestingActor.h b/Engine/Source/Runtime/NavigationSystem/Public/NavigationTestingActor.h
--- a/Engine/Source/Runtime/NavigationSystem/Public/NavigationTestingActor.h
+++ b/Engine/Source/Runtime/NavigationSystem/Public/NavigationTestingActor.h
@@ -79,9 +79,15 @@ public:
 	UPROPERTY(VisibleAnywhere, BlueprintReadOnly, Category=AgentStatus)
 	uint32 bProjectedLocationValid : 1;
 
+	/** if set, start the search from this actor, else start the search from the other actor */
 	UPROPERTY(EditAnywhere, Category=Pathfinding)
 	uint32 bSearchStart : 1;
 
+	/** this multiplier is used to compute a max node cost allowed to the open list
+	*	(max cost = MaxCostFacotr*InitialHeuristicEstimate) */
+	UPROPERTY(EditAnywhere, Category=Pathfinding, meta = (ClampMin = "0", UIMin = "0"))
+	float MaxCostFactor;
+
 	/** Instead of regular pathfinding from source to target location do
 	 *	a 'backwards' search that searches from the source, but as if the allowed
 	 *	movement direction was coming from the target. Meaningful only for paths
