Addition of a minimum CostLimit value to FPathFindingQuery::ComputeCostLimitFromHeuristic()
Exposing the minimum CostLimit parameter in NavigationTestingActor
Renaming MaxCost to CostLimit

diff --git a/Engine/Source/Runtime/NavigationSystem/Private/NavMesh/PImplRecastNavMesh.cpp b/Engine/Source/Runtime/NavigationSystem/Private/NavMesh/PImplRecastNavMesh.cpp
--- a/Engine/Source/Runtime/NavigationSystem/Private/NavMesh/PImplRecastNavMesh.cpp
+++ b/Engine/Source/Runtime/NavigationSystem/Private/NavMesh/PImplRecastNavMesh.cpp
@@ -788,7 +788,7 @@ ENavigationQueryResult::Type FPImplRecastNavMesh::FindPath(const FVector& StartL
 }
 
 // @TODONAV
-ENavigationQueryResult::Type FPImplRecastNavMesh::FindPath(const FVector& StartLoc, const FVector& EndLoc, const float MaxCost, FNavMeshPath& Path, const FNavigationQueryFilter& InQueryFilter, const UObject* Owner) const
+ENavigationQueryResult::Type FPImplRecastNavMesh::FindPath(const FVector& StartLoc, const FVector& EndLoc, const float CostLimit, FNavMeshPath& Path, const FNavigationQueryFilter& InQueryFilter, const UObject* Owner) const
 {
 	// temporarily disabling this check due to it causing too much "crashes"
 	// @todo but it needs to be back at some point since it realy checks for a buggy setup
@@ -826,7 +826,7 @@ ENavigationQueryResult::Type FPImplRecastNavMesh::FindPath(const FVector& StartL
 
 	// get path corridor
 	dtQueryResult PathResult;
-	const dtStatus FindPathStatus = NavQuery.findPath(StartPolyID, EndPolyID, &RecastStartPos.X, &RecastEndPos.X, MaxCost, QueryFilter, PathResult, 0);
+	const dtStatus FindPathStatus = NavQuery.findPath(StartPolyID, EndPolyID, &RecastStartPos.X, &RecastEndPos.X, CostLimit, QueryFilter, PathResult, 0);
 
 	// check for special case, where path has not been found, and starting polygon
 	// was the one closest to the target
@@ -897,9 +897,9 @@ ENavigationQueryResult::Type FPImplRecastNavMesh::TestPath(const FVector& StartL
 
 	// get path corridor
 	dtQueryResult PathResult;
-	const float MaxCost = FLT_MAX;
+	const float CostLimit = FLT_MAX;
 	const dtStatus FindPathStatus = NavQuery.findPath(StartPolyID, EndPolyID,
-		&RecastStartPos.X, &RecastEndPos.X, MaxCost, QueryFilter, PathResult, 0);
+		&RecastStartPos.X, &RecastEndPos.X, CostLimit, QueryFilter, PathResult, 0);
 
 	if (NumVisitedNodes)
 	{
@@ -1277,7 +1277,7 @@ static void StorePathfindingDebugStep(const dtNavMeshQuery& NavQuery, const dtNa
 	}
 }
 
-int32 FPImplRecastNavMesh::DebugPathfinding(const FVector& StartLoc, const FVector& EndLoc, const float MaxCost, const FNavigationQueryFilter& Filter, const UObject* Owner, TArray<FRecastDebugPathfindingData>& Steps)
+int32 FPImplRecastNavMesh::DebugPathfinding(const FVector& StartLoc, const FVector& EndLoc, const float CostLimit, const FNavigationQueryFilter& Filter, const UObject* Owner, TArray<FRecastDebugPathfindingData>& Steps)
 {
 	int32 NumSteps = 0;
 
@@ -1299,7 +1299,7 @@ int32 FPImplRecastNavMesh::DebugPathfinding(const FVector& StartLoc, const FVect
 		return NumSteps;
 	}
 
-	dtStatus status = NavQuery.initSlicedFindPath(StartPolyID, EndPolyID, &RecastStartPos.X, &RecastEndPos.X, MaxCost, QueryFilter);
+	dtStatus status = NavQuery.initSlicedFindPath(StartPolyID, EndPolyID, &RecastStartPos.X, &RecastEndPos.X, CostLimit, QueryFilter);
 	while (dtStatusInProgress(status))
 	{
 		StorePathfindingDebugStep(NavQuery, DetourNavMesh, Steps);
diff --git a/Engine/Source/Runtime/NavigationSystem/Private/NavMesh/RecastNavMesh.cpp b/Engine/Source/Runtime/NavigationSystem/Private/NavMesh/RecastNavMesh.cpp
--- a/Engine/Source/Runtime/NavigationSystem/Private/NavMesh/RecastNavMesh.cpp
+++ b/Engine/Source/Runtime/NavigationSystem/Private/NavMesh/RecastNavMesh.cpp
@@ -1310,8 +1310,8 @@ ENavigationQueryResult::Type ARecastNavMesh::CalcPathLengthAndCost(const FVector
 			Path->SetWantsStringPulling(false);
 			Path->SetWantsPathCorridor(true);
 			
-			const float MaxCost = FLT_MAX;
-			Result = RecastNavMeshImpl->FindPath(PathStart, PathEnd, MaxCost, Path.Get(), GetRightFilterRef(QueryFilter), QueryOwner);
+			const float CostLimit = FLT_MAX;
+			Result = RecastNavMeshImpl->FindPath(PathStart, PathEnd, CostLimit, Path.Get(), GetRightFilterRef(QueryFilter), QueryOwner);
 
 			if (Result == ENavigationQueryResult::Success || (Result == ENavigationQueryResult::Fail && Path->IsPartial()))
 			{
@@ -2095,7 +2095,7 @@ FPathFindingResult ARecastNavMesh::FindPath(const FNavAgentProperties& AgentProp
 		}
 		else
 		{
-			Result.Result = RecastNavMesh->RecastNavMeshImpl->FindPath(Query.StartLocation, AdjustedEndLocation, Query.MaxCost, *NavMeshPath, *NavFilter, Query.Owner.Get());
+			Result.Result = RecastNavMesh->RecastNavMeshImpl->FindPath(Query.StartLocation, AdjustedEndLocation, Query.CostLimit, *NavMeshPath, *NavFilter, Query.Owner.Get());
 
 			const bool bPartialPath = Result.IsPartial();
 			if (bPartialPath)
@@ -2309,7 +2309,7 @@ int32 ARecastNavMesh::DebugPathfinding(const FPathFindingQuery& Query, TArray<FR
 
 	if ((Query.StartLocation - Query.EndLocation).IsNearlyZero() == false)
 	{
-		NumSteps = RecastNavMesh->RecastNavMeshImpl->DebugPathfinding(Query.StartLocation, Query.EndLocation, Query.MaxCost, *(Query.QueryFilter.Get()), Query.Owner.Get(), Steps);
+		NumSteps = RecastNavMesh->RecastNavMeshImpl->DebugPathfinding(Query.StartLocation, Query.EndLocation, Query.CostLimit, *(Query.QueryFilter.Get()), Query.Owner.Get(), Steps);
 	}
 
 	return NumSteps;
diff --git a/Engine/Source/Runtime/NavigationSystem/Private/NavigationData.cpp b/Engine/Source/Runtime/NavigationSystem/Private/NavigationData.cpp
--- a/Engine/Source/Runtime/NavigationSystem/Private/NavigationData.cpp
+++ b/Engine/Source/Runtime/NavigationSystem/Private/NavigationData.cpp
@@ -19,8 +19,8 @@
 //----------------------------------------------------------------------//
 // FPathFindingQuery
 //----------------------------------------------------------------------//
-FPathFindingQuery::FPathFindingQuery(const UObject* InOwner, const ANavigationData& InNavData, const FVector& Start, const FVector& End, FSharedConstNavQueryFilter SourceQueryFilter, FNavPathSharedPtr InPathInstanceToFill, const float MaxCost) :
-	FPathFindingQueryData(InOwner, Start, End, SourceQueryFilter, MaxCost),
+FPathFindingQuery::FPathFindingQuery(const UObject* InOwner, const ANavigationData& InNavData, const FVector& Start, const FVector& End, FSharedConstNavQueryFilter SourceQueryFilter, FNavPathSharedPtr InPathInstanceToFill, const float CostLimit) :
+	FPathFindingQueryData(InOwner, Start, End, SourceQueryFilter, CostLimit),
 	NavData(&InNavData), PathInstanceToFill(InPathInstanceToFill), NavAgentProperties(InNavData.GetConfig())
 {
 	if (!QueryFilter.IsValid() && NavData.IsValid())
@@ -29,8 +29,8 @@ FPathFindingQuery::FPathFindingQuery(const UObject* InOwner, const ANavigationDa
 	}
 }
 
-FPathFindingQuery::FPathFindingQuery(const INavAgentInterface& InNavAgent, const ANavigationData& InNavData, const FVector& Start, const FVector& End, FSharedConstNavQueryFilter SourceQueryFilter, FNavPathSharedPtr InPathInstanceToFill, const float MaxCost) :
-	FPathFindingQueryData(Cast<UObject>(&InNavAgent), Start, End, SourceQueryFilter, MaxCost),
+FPathFindingQuery::FPathFindingQuery(const INavAgentInterface& InNavAgent, const ANavigationData& InNavData, const FVector& Start, const FVector& End, FSharedConstNavQueryFilter SourceQueryFilter, FNavPathSharedPtr InPathInstanceToFill, const float CostLimit) :
+	FPathFindingQueryData(Cast<UObject>(&InNavAgent), Start, End, SourceQueryFilter, CostLimit),
 	NavData(&InNavData), PathInstanceToFill(InPathInstanceToFill), NavAgentProperties(InNavAgent.GetNavAgentPropertiesRef())
 {
 	if (!QueryFilter.IsValid() && NavData.IsValid())
@@ -40,7 +40,7 @@ FPathFindingQuery::FPathFindingQuery(const INavAgentInterface& InNavAgent, const
 }
 
 FPathFindingQuery::FPathFindingQuery(const FPathFindingQuery& Source) :
-	FPathFindingQueryData(Source.Owner.Get(), Source.StartLocation, Source.EndLocation, Source.QueryFilter, Source.NavDataFlags, Source.bAllowPartialPaths, Source.MaxCost),
+	FPathFindingQueryData(Source.Owner.Get(), Source.StartLocation, Source.EndLocation, Source.QueryFilter, Source.NavDataFlags, Source.bAllowPartialPaths, Source.CostLimit),
 	NavData(Source.NavData), PathInstanceToFill(Source.PathInstanceToFill), NavAgentProperties(Source.NavAgentProperties)
 {
 	if (!QueryFilter.IsValid() && NavData.IsValid())
@@ -82,16 +82,16 @@ FPathFindingQuery::FPathFindingQuery(FNavPathSharedRef PathToRecalculate, const
 	}
 }
 
-float FPathFindingQuery::ComputeMaxCostFromHeuristic(const FVector& StartPos, const FVector& EndPos, const float HeuristicScale, const float MaxCostFactor) const
+float FPathFindingQuery::ComputeCostLimitFromHeuristic(const FVector& StartPos, const FVector& EndPos, const float HeuristicScale, const float CostLimitFactor, const float MinimumCostLimit) const
 {
-	if (MaxCostFactor == FLT_MAX)
+	if (CostLimitFactor == FLT_MAX)
 	{
 		return FLT_MAX;
 	}
 	else
 	{
 		const float OriginalHeuristicEstimate = HeuristicScale * FVector::Dist(StartPos, EndPos);
-		return FMath::Min(FLT_MAX, MaxCostFactor * OriginalHeuristicEstimate);
+		return FMath::Clamp(CostLimitFactor * OriginalHeuristicEstimate, MinimumCostLimit, FLT_MAX);
 	}
 }
 
@@ -100,7 +100,7 @@ float FPathFindingQuery::ComputeMaxCostFromHeuristic(const FVector& StartPos, co
 //----------------------------------------------------------------------//
 uint32 FAsyncPathFindingQuery::LastPathFindingUniqueID = INVALID_NAVQUERYID;
 
-FAsyncPathFindingQuery::FAsyncPathFindingQuery(const UObject* InOwner, const ANavigationData& InNavData, const FVector& Start, const FVector& End, const FNavPathQueryDelegate& Delegate, FSharedConstNavQueryFilter SourceQueryFilter, const float MaxCost)
+FAsyncPathFindingQuery::FAsyncPathFindingQuery(const UObject* InOwner, const ANavigationData& InNavData, const FVector& Start, const FVector& End, const FNavPathQueryDelegate& Delegate, FSharedConstNavQueryFilter SourceQueryFilter, const float CostLimit)
 : FPathFindingQuery(InOwner, InNavData, Start, End, SourceQueryFilter)
 , QueryID(GetUniqueID())
 , OnDoneDelegate(Delegate)
diff --git a/Engine/Source/Runtime/NavigationSystem/Private/NavigationTestingActor.cpp b/Engine/Source/Runtime/NavigationSystem/Private/NavigationTestingActor.cpp
--- a/Engine/Source/Runtime/NavigationSystem/Private/NavigationTestingActor.cpp
+++ b/Engine/Source/Runtime/NavigationSystem/Private/NavigationTestingActor.cpp
@@ -41,7 +41,8 @@ ANavigationTestingActor::ANavigationTestingActor(const FObjectInitializer& Objec
 	bIsEditorOnlyActor = true;
 	NavAgentProps.AgentRadius = 34.f;
 	NavAgentProps.AgentHeight = 144.f;
-	MaxCostFactor = FLT_MAX;
+	CostLimitFactor = FLT_MAX;
+	MinimumCostLimit = 0.f;
 	ShowStepIndex = -1;
 	bShowNodePool = true;
 	bShowBestPath = true;
@@ -429,10 +430,10 @@ void ANavigationTestingActor::SearchPathTo(ANavigationTestingActor* Goal)
 		Query.QueryFilter = NavigationFilterCopy;
 	}
 	
-	//Apply max cost factor
+	//Apply cost limit factor
 	FSharedConstNavQueryFilter NavQueryFilter = Query.QueryFilter ? Query.QueryFilter : NavData->GetDefaultQueryFilter();
 	const float HeuristicScale = NavQueryFilter->GetHeuristicScale();
-	Query.MaxCost = Query.ComputeMaxCostFromHeuristic(Query.StartLocation, Query.EndLocation, HeuristicScale, MaxCostFactor);
+	Query.CostLimit = Query.ComputeCostLimitFromHeuristic(Query.StartLocation, Query.EndLocation, HeuristicScale, CostLimitFactor, MinimumCostLimit);
 
 	EPathFindingMode::Type Mode = bUseHierarchicalPathfinding ? EPathFindingMode::Hierarchical : EPathFindingMode::Regular;
 	FPathFindingResult Result = NavSys->FindPathSync(NavAgentProps, Query, Mode);
diff --git a/Engine/Source/Runtime/NavigationSystem/Public/NavMesh/PImplRecastNavMesh.h b/Engine/Source/Runtime/NavigationSystem/Public/NavMesh/PImplRecastNavMesh.h
--- a/Engine/Source/Runtime/NavigationSystem/Public/NavMesh/PImplRecastNavMesh.h
+++ b/Engine/Source/Runtime/NavigationSystem/Public/NavMesh/PImplRecastNavMesh.h
@@ -72,11 +72,11 @@ public:
 
 	// @TODONAV
 	/** Generates path from the given query. Synchronous. */
-	UE_DEPRECATED(4.25, "Use the version with the added MaxCost parameter (FLT_MAX can be used as default).")
+	UE_DEPRECATED(4.25, "Use the version with the added CostLimit parameter (FLT_MAX can be used as default).")
 	ENavigationQueryResult::Type FindPath(const FVector& StartLoc, const FVector& EndLoc, FNavMeshPath& Path, const FNavigationQueryFilter& Filter, const UObject* Owner) const;
 	
 	/** Generates path from the given query. Synchronous. */
-	ENavigationQueryResult::Type FindPath(const FVector& StartLoc, const FVector& EndLoc, const float MaxCost, FNavMeshPath& Path, const FNavigationQueryFilter& Filter, const UObject* Owner) const;
+	ENavigationQueryResult::Type FindPath(const FVector& StartLoc, const FVector& EndLoc, const float CostLimit, FNavMeshPath& Path, const FNavigationQueryFilter& Filter, const UObject* Owner) const;
 
 	/** Check if path exists */
 	ENavigationQueryResult::Type TestPath(const FVector& StartLoc, const FVector& EndLoc, const FNavigationQueryFilter& Filter, const UObject* Owner, int32* NumVisitedNodes = 0) const;
@@ -89,7 +89,7 @@ public:
 		ARecastNavMesh::FRaycastResult& RaycastResult, NavNodeRef StartNode = INVALID_NAVNODEREF) const;
 
 	/** Generates path from given query and collect data for every step of A* algorithm */
-	int32 DebugPathfinding(const FVector& StartLoc, const FVector& EndLoc, const float MaxCost, const FNavigationQueryFilter& Filter, const UObject* Owner, TArray<FRecastDebugPathfindingData>& Steps);
+	int32 DebugPathfinding(const FVector& StartLoc, const FVector& EndLoc, const float CostLimit, const FNavigationQueryFilter& Filter, const UObject* Owner, TArray<FRecastDebugPathfindingData>& Steps);
 
 	/** Returns a random location on the navmesh. */
 	FNavLocation GetRandomPoint(const FNavigationQueryFilter& Filter, const UObject* Owner) const;
diff --git a/Engine/Source/Runtime/NavigationSystem/Public/NavigationData.h b/Engine/Source/Runtime/NavigationSystem/Public/NavigationData.h
--- a/Engine/Source/Runtime/NavigationSystem/Public/NavigationData.h
+++ b/Engine/Source/Runtime/NavigationSystem/Public/NavigationData.h
@@ -1002,7 +1002,7 @@ struct FAsyncPathFindingQuery : public FPathFindingQuery
 		, Mode(EPathFindingMode::Regular)
 	{ }
 
-	FAsyncPathFindingQuery(const UObject* InOwner, const ANavigationData& InNavData, const FVector& Start, const FVector& End, const FNavPathQueryDelegate& Delegate, FSharedConstNavQueryFilter SourceQueryFilter, const float MaxCost = FLT_MAX);
+	FAsyncPathFindingQuery(const UObject* InOwner, const ANavigationData& InNavData, const FVector& Start, const FVector& End, const FNavPathQueryDelegate& Delegate, FSharedConstNavQueryFilter SourceQueryFilter, const float CostLimit = FLT_MAX);
 	FAsyncPathFindingQuery(const FPathFindingQuery& Query, const FNavPathQueryDelegate& Delegate, const EPathFindingMode::Type QueryMode);
 
 protected:
diff --git a/Engine/Source/Runtime/NavigationSystem/Public/NavigationSystemTypes.h b/Engine/Source/Runtime/NavigationSystem/Public/NavigationSystemTypes.h
--- a/Engine/Source/Runtime/NavigationSystem/Public/NavigationSystemTypes.h
+++ b/Engine/Source/Runtime/NavigationSystem/Public/NavigationSystemTypes.h
@@ -28,18 +28,18 @@ struct NAVIGATIONSYSTEM_API FPathFindingQueryData
 	FSharedConstNavQueryFilter QueryFilter;
 
 	/** cost limit of nodes allowed to be added to the open list */
-	float MaxCost;
-
+	float CostLimit;
+	
 	/** additional flags passed to navigation data handling request */
 	int32 NavDataFlags;
 
 	/** if set, allow partial paths as a result */
 	uint32 bAllowPartialPaths : 1;
 
-	FPathFindingQueryData() : StartLocation(FNavigationSystem::InvalidLocation), EndLocation(FNavigationSystem::InvalidLocation), MaxCost(FLT_MAX), NavDataFlags(0), bAllowPartialPaths(true) {}
+	FPathFindingQueryData() : StartLocation(FNavigationSystem::InvalidLocation), EndLocation(FNavigationSystem::InvalidLocation), CostLimit(FLT_MAX), NavDataFlags(0), bAllowPartialPaths(true) {}
 
-	FPathFindingQueryData(const UObject* InOwner, const FVector& InStartLocation, const FVector& InEndLocation, FSharedConstNavQueryFilter InQueryFilter = nullptr, int32 InNavDataFlags = 0, bool bInAllowPartialPaths = true, const float InMaxCost = FLT_MAX) :
-		Owner(InOwner), StartLocation(InStartLocation), EndLocation(InEndLocation), QueryFilter(InQueryFilter), MaxCost(InMaxCost), NavDataFlags(InNavDataFlags), bAllowPartialPaths(bInAllowPartialPaths) {}
+	FPathFindingQueryData(const UObject* InOwner, const FVector& InStartLocation, const FVector& InEndLocation, FSharedConstNavQueryFilter InQueryFilter = nullptr, int32 InNavDataFlags = 0, bool bInAllowPartialPaths = true, const float InCostLimit = FLT_MAX) :
+		Owner(InOwner), StartLocation(InStartLocation), EndLocation(InEndLocation), QueryFilter(InQueryFilter), CostLimit(InCostLimit), NavDataFlags(InNavDataFlags), bAllowPartialPaths(bInAllowPartialPaths) {}
 };
 
 struct NAVIGATIONSYSTEM_API FPathFindingQuery : public FPathFindingQueryData
@@ -50,8 +50,8 @@ struct NAVIGATIONSYSTEM_API FPathFindingQuery : public FPathFindingQueryData
 
 	FPathFindingQuery() : FPathFindingQueryData() {}
 	FPathFindingQuery(const FPathFindingQuery& Source);
-	FPathFindingQuery(const UObject* InOwner, const ANavigationData& InNavData, const FVector& Start, const FVector& End, FSharedConstNavQueryFilter SourceQueryFilter = NULL, FNavPathSharedPtr InPathInstanceToFill = NULL, const float MaxCost = FLT_MAX);
-	FPathFindingQuery(const INavAgentInterface& InNavAgent, const ANavigationData& InNavData, const FVector& Start, const FVector& End, FSharedConstNavQueryFilter SourceQueryFilter = NULL, FNavPathSharedPtr InPathInstanceToFill = NULL, const float MaxCost = FLT_MAX);
+	FPathFindingQuery(const UObject* InOwner, const ANavigationData& InNavData, const FVector& Start, const FVector& End, FSharedConstNavQueryFilter SourceQueryFilter = NULL, FNavPathSharedPtr InPathInstanceToFill = NULL, const float CostLimit = FLT_MAX);
+	FPathFindingQuery(const INavAgentInterface& InNavAgent, const ANavigationData& InNavData, const FVector& Start, const FVector& End, FSharedConstNavQueryFilter SourceQueryFilter = NULL, FNavPathSharedPtr InPathInstanceToFill = NULL, const float CostLimit = FLT_MAX);
 
 	explicit FPathFindingQuery(FNavPathSharedRef PathToRecalculate, const ANavigationData* NavDataOverride = NULL);
 
@@ -59,8 +59,10 @@ struct NAVIGATIONSYSTEM_API FPathFindingQuery : public FPathFindingQueryData
 	FPathFindingQuery& SetAllowPartialPaths(bool bAllow) { bAllowPartialPaths = bAllow; return *this; }
 	FPathFindingQuery& SetNavAgentProperties(const FNavAgentProperties& InNavAgentProperties) { NavAgentProperties = InNavAgentProperties; return *this; }
 
-	/** utility function to compute a max cost using an Euclidean heuristic, an heuristic scale and a max cost factor */
-	float ComputeMaxCostFromHeuristic(const FVector& StartPos, const FVector& EndPos, const float HeuristicScale, const float MaxCostFactor) const;
+	/** utility function to compute a cost limit using an Euclidean heuristic, an heuristic scale and a cost limit factor
+	*	CostLimitFactor: multiplier used to compute the cost limit value from the initial heuristic
+	*	MinimumCostLimit: minimum clamping value used to prevent low cost limit for short path query */
+	float ComputeCostLimitFromHeuristic(const FVector& StartPos, const FVector& EndPos, const float HeuristicScale, const float CostLimitFactor, const float MinimumCostLimit) const;
 };
 
 namespace EPathFindingMode
diff --git a/Engine/Source/Runtime/NavigationSystem/Public/NavigationTestingActor.h b/Engine/Source/Runtime/NavigationSystem/Public/NavigationTestingActor.h
--- a/Engine/Source/Runtime/NavigationSystem/Public/NavigationTestingActor.h
+++ b/Engine/Source/Runtime/NavigationSystem/Public/NavigationTestingActor.h
@@ -84,9 +84,14 @@ public:
 	uint32 bSearchStart : 1;
 
 	/** this multiplier is used to compute a max node cost allowed to the open list
-	*	(max cost = MaxCostFacotr*InitialHeuristicEstimate) */
+	 *	(cost limit = CostLimitFacotr*InitialHeuristicEstimate) */
 	UPROPERTY(EditAnywhere, Category=Pathfinding, meta = (ClampMin = "0", UIMin = "0"))
-	float MaxCostFactor;
+	float CostLimitFactor;
+
+	/** minimum cost limit clamping value (in cost units)
+	 *	used to allow large deviation in short paths */
+	UPROPERTY(EditAnywhere, Category = Pathfinding, meta = (ClampMin = "0", UIMin = "0"))
+	float MinimumCostLimit;
 
 	/** Instead of regular pathfinding from source to target location do
 	 *	a 'backwards' search that searches from the source, but as if the allowed
diff --git a/Engine/Source/Runtime/Navmesh/Private/Detour/DetourNavMeshQuery.cpp b/Engine/Source/Runtime/Navmesh/Private/Detour/DetourNavMeshQuery.cpp
--- a/Engine/Source/Runtime/Navmesh/Private/Detour/DetourNavMeshQuery.cpp
+++ b/Engine/Source/Runtime/Navmesh/Private/Detour/DetourNavMeshQuery.cpp
@@ -1423,7 +1423,7 @@ dtStatus dtNavMeshQuery::queryPolygons(const float* center, const float* extents
 ///
 dtStatus dtNavMeshQuery::findPath(dtPolyRef startRef, dtPolyRef endRef,
 								  const float* startPos, const float* endPos,
-								  const float maxCost, const dtQueryFilter* filter, //@UE4
+								  const float costLimit, const dtQueryFilter* filter, //@UE4
 								  dtQueryResult& result, float* totalCost) const
 {
 	dtAssert(m_nav);
@@ -1594,7 +1594,7 @@ dtStatus dtNavMeshQuery::findPath(dtPolyRef startRef, dtPolyRef endRef,
 			if (curCost == DT_UNWALKABLE_POLY_COST)
 				continue;
 
-			if (total > maxCost) //@UE4
+			if (total > costLimit) //@UE4
 				continue;
 
 			// Add or update the node.
@@ -1833,7 +1833,7 @@ dtStatus dtNavMeshQuery::testClusterPath(dtPolyRef startRef, dtPolyRef endRef) c
 /// path query.
 ///
 dtStatus dtNavMeshQuery::initSlicedFindPath(dtPolyRef startRef, dtPolyRef endRef,
-											const float* startPos, const float* endPos, const float maxCost, //@UE4
+											const float* startPos, const float* endPos, const float costLimit, //@UE4
 											const dtQueryFilter* filter)
 {
 	dtAssert(m_nav);
@@ -1847,7 +1847,7 @@ dtStatus dtNavMeshQuery::initSlicedFindPath(dtPolyRef startRef, dtPolyRef endRef
 	m_query.endRef = endRef;
 	dtVcopy(m_query.startPos, startPos);
 	dtVcopy(m_query.endPos, endPos);
-	m_query.maxCost = maxCost; //@UE4
+	m_query.costLimit = costLimit; //@UE4
 	m_query.filter = filter;
 	
 	if (!startRef || !endRef)
@@ -2047,7 +2047,7 @@ dtStatus dtNavMeshQuery::updateSlicedFindPath(const int maxIter, int* doneIters)
 			if (curCost == DT_UNWALKABLE_POLY_COST) //@UE4
 				continue;
 
-			if (total > m_query.maxCost) //@UE4
+			if (total > m_query.costLimit) //@UE4
 				continue;
 			
 			// Add or update the node.
diff --git a/Engine/Source/Runtime/Navmesh/Private/DetourCrowd/DetourCrowd.cpp b/Engine/Source/Runtime/Navmesh/Private/DetourCrowd/DetourCrowd.cpp
--- a/Engine/Source/Runtime/Navmesh/Private/DetourCrowd/DetourCrowd.cpp
+++ b/Engine/Source/Runtime/Navmesh/Private/DetourCrowd/DetourCrowd.cpp
@@ -854,8 +854,8 @@ void dtCrowd::updateMoveRequest(const float /*dt*/)
 			// Quick seach towards the goal.
 			static const int MAX_ITER = 20;
 			m_navquery->updateLinkFilter(ag->params.linkFilter.Get());
-			const float maxCost = FLT_MAX; //@UE4
-			m_navquery->initSlicedFindPath(path[0], ag->targetRef, ag->npos, ag->targetPos, maxCost, &m_filters[ag->params.filter]); //@UE4
+			const float costLimit = FLT_MAX; //@UE4
+			m_navquery->initSlicedFindPath(path[0], ag->targetRef, ag->npos, ag->targetPos, costLimit, &m_filters[ag->params.filter]); //@UE4
 			m_navquery->updateSlicedFindPath(MAX_ITER, 0);
 			dtStatus status = 0;
 			if (ag->targetReplan) // && npath > 10)
@@ -921,9 +921,9 @@ void dtCrowd::updateMoveRequest(const float /*dt*/)
 	for (int i = 0; i < nqueue; ++i)
 	{
 		dtCrowdAgent* ag = queue[i];
-		const float maxCost = FLT_MAX; //@UE4
+		const float costLimit = FLT_MAX; //@UE4
 		ag->targetPathqRef = m_pathq.request(ag->corridor.getLastPoly(), ag->targetRef,
-			ag->corridor.getTarget(), ag->targetPos, maxCost, &m_filters[ag->params.filter], ag->params.linkFilter); //@UE4
+			ag->corridor.getTarget(), ag->targetPos, costLimit, &m_filters[ag->params.filter], ag->params.linkFilter); //@UE4
 		if (ag->targetPathqRef != DT_PATHQ_INVALID)
 			ag->targetState = DT_CROWDAGENT_TARGET_WAITING_FOR_PATH;
 	}
diff --git a/Engine/Source/Runtime/Navmesh/Private/DetourCrowd/DetourPathCorridor.cpp b/Engine/Source/Runtime/Navmesh/Private/DetourCrowd/DetourPathCorridor.cpp
--- a/Engine/Source/Runtime/Navmesh/Private/DetourCrowd/DetourPathCorridor.cpp
+++ b/Engine/Source/Runtime/Navmesh/Private/DetourCrowd/DetourPathCorridor.cpp
@@ -524,8 +524,8 @@ bool dtPathCorridor::optimizePathTopology(dtNavMeshQuery* navquery, const dtQuer
 	
 	dtPolyRef res[MAX_RES];
 	int nres = 0;
-	const float maxCost = FLT_MAX; //@UE4
-	navquery->initSlicedFindPath(m_path[0], m_path[m_npath-1], m_pos, m_target, maxCost, filter); //@UE4
+	const float costLimit = FLT_MAX; //@UE4
+	navquery->initSlicedFindPath(m_path[0], m_path[m_npath-1], m_pos, m_target, costLimit, filter); //@UE4
 	navquery->updateSlicedFindPath(MAX_ITER, 0);
 	dtStatus status = navquery->finalizeSlicedFindPathPartial(m_path, m_npath, res, &nres, MAX_RES);
 	
diff --git a/Engine/Source/Runtime/Navmesh/Private/DetourCrowd/DetourPathQueue.cpp b/Engine/Source/Runtime/Navmesh/Private/DetourCrowd/DetourPathQueue.cpp
--- a/Engine/Source/Runtime/Navmesh/Private/DetourCrowd/DetourPathQueue.cpp
+++ b/Engine/Source/Runtime/Navmesh/Private/DetourCrowd/DetourPathQueue.cpp
@@ -111,7 +111,7 @@ void dtPathQueue::update(const int maxIters)
 		// Handle query start.
 		if (q.status == 0)
 		{
-			q.status = m_navquery->initSlicedFindPath(q.startRef, q.endRef, q.startPos, q.endPos, q.maxCost, q.filter); //@UE4
+			q.status = m_navquery->initSlicedFindPath(q.startRef, q.endRef, q.startPos, q.endPos, q.costLimit, q.filter); //@UE4
 		}		
 		// Handle query in progress.
 		if (dtStatusInProgress(q.status))
@@ -133,7 +133,7 @@ void dtPathQueue::update(const int maxIters)
 }
 
 dtPathQueueRef dtPathQueue::request(dtPolyRef startRef, dtPolyRef endRef,
-									const float* startPos, const float* endPos, const float maxCost, //@UE4
+									const float* startPos, const float* endPos, const float costLimit, //@UE4
 									const dtQueryFilter* filter, TSharedPtr<dtQuerySpecialLinkFilter> linkFilter)
 {
 	// Find empty slot
@@ -159,7 +159,7 @@ dtPathQueueRef dtPathQueue::request(dtPolyRef startRef, dtPolyRef endRef,
 	q.startRef = startRef;
 	dtVcopy(q.endPos, endPos);
 	q.endRef = endRef;
-	q.maxCost = maxCost; //@UE4
+	q.costLimit = costLimit; //@UE4
 	
 	q.status = 0;
 	q.npath = 0;
diff --git a/Engine/Source/Runtime/Navmesh/Public/Detour/DetourNavMeshQuery.h b/Engine/Source/Runtime/Navmesh/Public/Detour/DetourNavMeshQuery.h
--- a/Engine/Source/Runtime/Navmesh/Public/Detour/DetourNavMeshQuery.h
+++ b/Engine/Source/Runtime/Navmesh/Public/Detour/DetourNavMeshQuery.h
@@ -344,12 +344,12 @@ public:
 	///  @param[in]		endRef		The reference id of the end polygon.
 	///  @param[in]		startPos	A position within the start polygon. [(x, y, z)]
 	///  @param[in]		endPos		A position within the end polygon. [(x, y, z)]
-	///  @param[in]		maxCost		Cost limit of nodes allowed to be added to the open list	//@UE4
+	///  @param[in]		costLimit	Cost limit of nodes allowed to be added to the open list	//@UE4
 	///  @param[in]		filter		The polygon filter to apply to the query.
 	///  @param[out]	result		Results for path corridor, fills in refs and costs for each poly from start to end
 	///	 @param[out]	totalCost			If provided will get filled will total cost of path
 	dtStatus findPath(dtPolyRef startRef, dtPolyRef endRef,
-					  const float* startPos, const float* endPos, const float maxCost, //@UE4
+					  const float* startPos, const float* endPos, const float costLimit, //@UE4
 					  const dtQueryFilter* filter,
 					  dtQueryResult& result, float* totalCost) const;
 	
@@ -384,11 +384,11 @@ public:
 	///  @param[in]		endRef		The reference id of the end polygon.
 	///  @param[in]		startPos	A position within the start polygon. [(x, y, z)]
 	///  @param[in]		endPos		A position within the end polygon. [(x, y, z)]
-	///  @param[in]		maxCost		Cost limit of nodes allowed to be added to the open list	//@UE4
+	///  @param[in]		costLimit	Cost limit of nodes allowed to be added to the open list	//@UE4
 	///  @param[in]		filter		The polygon filter to apply to the query.
 	/// @returns The status flags for the query.
 	dtStatus initSlicedFindPath(dtPolyRef startRef, dtPolyRef endRef,
-								const float* startPos, const float* endPos, const float maxCost, //@UE4
+								const float* startPos, const float* endPos, const float costLimit, //@UE4
 								const dtQueryFilter* filter);
 
 	/// Updates an in-progress sliced path query.
@@ -781,7 +781,7 @@ private:
 		float lastBestNodeCost;
 		dtPolyRef startRef, endRef;
 		float startPos[3], endPos[3];
-		float maxCost; 					//@UE4 ///< Cost limit of nodes allowed to be added to the open list
+		float costLimit; 					//@UE4 ///< Cost limit of nodes allowed to be added to the open list
 		const dtQueryFilter* filter;
 	};
 	dtQueryData m_query;				///< Sliced query state.
diff --git a/Engine/Source/Runtime/Navmesh/Public/DetourCrowd/DetourPathQueue.h b/Engine/Source/Runtime/Navmesh/Public/DetourCrowd/DetourPathQueue.h
--- a/Engine/Source/Runtime/Navmesh/Public/DetourCrowd/DetourPathQueue.h
+++ b/Engine/Source/Runtime/Navmesh/Public/DetourCrowd/DetourPathQueue.h
@@ -37,7 +37,7 @@ class dtPathQueue
 		/// Path find start and end location.
 		float startPos[3], endPos[3];
 		dtPolyRef startRef, endRef;
-		float maxCost;
+		float costLimit;
 		/// Result.
 		dtPolyRef* path;
 		int npath;
@@ -67,7 +67,7 @@ public:
 	void update(const int maxIters);
 	
 	dtPathQueueRef request(dtPolyRef startRef, dtPolyRef endRef,
-						   const float* startPos, const float* endPos, const float maxCost,
+						   const float* startPos, const float* endPos, const float costLimit,
 						   const dtQueryFilter* filter,
 						   TSharedPtr<dtQuerySpecialLinkFilter> linkFilter);
 	
