NavigationSystem: prevent registration of CustomNavLink from a different world in PIE
With Editor multiple NavigationSystems may exist at the same time (i.e. Editor, Client Game, Server Game worlds) so we want to make sure that any given NavigationSystem instance performs a single flush of the global pending queue to register the links associated to their outer World.

diff --git a/Engine/Source/Runtime/NavigationSystem/Private/NavigationSystem.cpp b/Engine/Source/Runtime/NavigationSystem/Private/NavigationSystem.cpp
--- a/Engine/Source/Runtime/NavigationSystem/Private/NavigationSystem.cpp
+++ b/Engine/Source/Runtime/NavigationSystem/Private/NavigationSystem.cpp
@@ -719,6 +719,11 @@ void UNavigationSystemV1::OnWorldInitDone(FNavigationSystemRunMode Mode)
 
 	World->OnBeginTearingDown().AddUObject(this, &UNavigationSystemV1::OnBeginTearingDown);
 
+	// process all queued custom link registration requests
+	// (since it's possible navigation system was not ready by the time
+	// those links were serialized-in or spawned)
+	ProcessCustomLinkPendingRegistration();
+
 	if (IsThereAnywhereToBuildNavigation() == false
 		// Simulation mode is a special case - better not do it in this case
 		&& OperationMode != FNavigationSystemRunMode::SimulationMode)
@@ -922,11 +927,6 @@ void UNavigationSystemV1::Tick(float DeltaSeconds)
 	{
 		return;
 	}
-	
-	if (PendingCustomLinkRegistration.Num())
-	{
-		ProcessCustomLinkPendingRegistration();
-	}
 
 	if (PendingNavBoundsUpdates.Num() > 0)
 	{
@@ -1841,7 +1841,17 @@ void UNavigationSystemV1::ProcessCustomLinkPendingRegistration()
 		
 		if (LinkOb.IsValid() && ILink)
 		{
+#if WITH_EDITOR
+			// In Editor multiple NavigationSystems may exist at the same time (i.e. Editor, Client Game, Server Game worlds)
+			// so we want to make sure that any given NavigationSystem instance performs a single flush of the global pending queue
+			// to register the links associated to their outer World.
+			// Following registration requests will be forwarded directly to the NavigationSystem and won't use the queue.
+			// We call RequestCustomLinkRegistering instead of RegisterCustomLink so each link
+			// will register to the navigation system associated to their outer world (if created) or put back in the queue.
+			RequestCustomLinkRegistering(*ILink, LinkOb.Get());
+#else
 			RegisterCustomLink(*ILink);
+#endif // WITH_EDITOR
 		}
 	}
 }
@@ -1988,6 +1998,9 @@ void UNavigationSystemV1::UnregisterNavData(ANavigationData* NavData)
 
 void UNavigationSystemV1::RegisterCustomLink(INavLinkCustomInterface& CustomLink)
 {
+	ensureMsgf(CustomLink.GetLinkOwner() == nullptr || GetWorld() == CustomLink.GetLinkOwner()->GetWorld(), 
+		TEXT("Registering a link from a world different than the navigation system world should not happen."));
+
 	uint32 LinkId = CustomLink.GetLinkId();
 
 	// if there's already a link with that Id registered, assign new Id and mark dirty area
