-Adding prioritization to nav mesh tile building and a virtual function for getting the locations to prioritize.

diff --git a/Engine/Source/Runtime/NavigationSystem/Private/NavMesh/RecastNavMeshGenerator.cpp b/Engine/Source/Runtime/NavigationSystem/Private/NavMesh/RecastNavMeshGenerator.cpp
--- a/Engine/Source/Runtime/NavigationSystem/Private/NavMesh/RecastNavMeshGenerator.cpp
+++ b/Engine/Source/Runtime/NavigationSystem/Private/NavMesh/RecastNavMeshGenerator.cpp
@@ -4301,23 +4301,14 @@ void FRecastNavMeshGenerator::SortPendingBuildTiles()
 		return;
 	}
 
-	TArray<FVector2D> SeedLocations;
 	UWorld* CurWorld = GetWorld();
 	if (CurWorld == nullptr)
 	{
 		return;
 	}
 
-	// Collect players positions
-	for (FConstPlayerControllerIterator PlayerIt = CurWorld->GetPlayerControllerIterator(); PlayerIt; ++PlayerIt)
-	{
-		APlayerController* PC = PlayerIt->Get();
-		if (PC && PC->GetPawn() != NULL)
-		{
-			const FVector2D SeedLoc(PC->GetPawn()->GetActorLocation());
-			SeedLocations.Add(SeedLoc);
-		}
-	}
+	TArray<FVector2D> SeedLocations;
+	GetSeedLocations(*CurWorld, SeedLocations);
 
 	if (SeedLocations.Num() == 0)
 	{
@@ -4345,6 +4336,20 @@ void FRecastNavMeshGenerator::SortPendingBuildTiles()
 	}
 }
 
+void FRecastNavMeshGenerator::GetSeedLocations(UWorld& World, TArray<FVector2D>& OutSeedLocations) const
+{
+	// Collect players positions
+	for (FConstPlayerControllerIterator PlayerIt = World.GetPlayerControllerIterator(); PlayerIt; ++PlayerIt)
+	{
+		APlayerController* PC = PlayerIt->Get();
+		if (PC && PC->GetPawn() != NULL)
+		{
+			const FVector2D SeedLoc(PC->GetPawn()->GetActorLocation());
+			OutSeedLocations.Add(SeedLoc);
+		}
+	}
+}
+
 TSharedRef<FRecastTileGenerator> FRecastNavMeshGenerator::CreateTileGenerator(const FIntPoint& Coord, const TArray<FBox>& DirtyAreas)
 {
 	TSharedRef<FRecastTileGenerator> TileGenerator = MakeShareable(new FRecastTileGenerator(*this, Coord));
diff --git a/Engine/Source/Runtime/NavigationSystem/Public/NavMesh/RecastNavMeshGenerator.h b/Engine/Source/Runtime/NavigationSystem/Public/NavMesh/RecastNavMeshGenerator.h
--- a/Engine/Source/Runtime/NavigationSystem/Public/NavMesh/RecastNavMeshGenerator.h
+++ b/Engine/Source/Runtime/NavigationSystem/Public/NavMesh/RecastNavMeshGenerator.h
@@ -521,6 +521,9 @@ protected:
 	// Sorts pending build tiles by proximity to player, so tiles closer to player will get generated first
 	virtual void SortPendingBuildTiles();
 
+	// Get seed locations used for sorting pending build tiles. Tiles closer to these locations will be prioritized first.
+	virtual void GetSeedLocations(UWorld& World, TArray<FVector2D>& OutSeedLocations) const;
+
 	/** Instantiates dtNavMesh and configures it for tiles generation. Returns false if failed */
 	bool ConstructTiledNavMesh();
 
