FNavigationDirtyAreasController: added method ForceRebuildOnNextTick that sets cumulative time to at least one full cycle so next tick will rebuild dirty areas

#rnx

diff --git a/Engine/Source/Runtime/NavigationSystem/Private/NavigationDirtyAreasController.cpp b/Engine/Source/Runtime/NavigationSystem/Private/NavigationDirtyAreasController.cpp
--- a/Engine/Source/Runtime/NavigationSystem/Private/NavigationDirtyAreasController.cpp
+++ b/Engine/Source/Runtime/NavigationSystem/Private/NavigationDirtyAreasController.cpp
@@ -15,6 +15,12 @@ FNavigationDirtyAreasController::FNavigationDirtyAreasController()
 
 }
 
+void FNavigationDirtyAreasController::ForceRebuildOnNextTick()
+{
+	float MinTimeForUpdate = (DirtyAreasUpdateFreq != 0.f ? (1.0f / DirtyAreasUpdateFreq) : 0.f);
+	DirtyAreasUpdateTime = FMath::Max(DirtyAreasUpdateTime, MinTimeForUpdate);
+}
+
 void FNavigationDirtyAreasController::Tick(const float DeltaSeconds, const TArray<ANavigationData*>& NavDataSet, bool bForceRebuilding)
 {
 	DirtyAreasUpdateTime += DeltaSeconds;
diff --git a/Engine/Source/Runtime/NavigationSystem/Public/NavigationDirtyAreasController.h b/Engine/Source/Runtime/NavigationSystem/Public/NavigationDirtyAreasController.h
--- a/Engine/Source/Runtime/NavigationSystem/Public/NavigationDirtyAreasController.h
+++ b/Engine/Source/Runtime/NavigationSystem/Public/NavigationDirtyAreasController.h
@@ -27,6 +27,9 @@ struct NAVIGATIONSYSTEM_API FNavigationDirtyAreasController
 
 	void Reset();
 	
+	/** sets cumulative time to at least one cycle so next tick will rebuild dirty areas */
+	void ForceRebuildOnNextTick();
+
 	void Tick(float DeltaSeconds, const TArray<ANavigationData*>& NavDataSet, bool bForceRebuilding = false);
 	void AddArea(const FBox& NewArea, int32 Flags);
 	
