Work to Timeslice the Nav Tile Regen.

I've left the Sync code path as seperate to the Sync Time Sliced one, as its atleast 10% faster (but much spikier). Its also good way of easily verifying if the Time Sliced Code is correct.

I've let medium level functions have access to the timeslicer as they are often expensive enough to warrent it and many times we are time slicing within loops of these functions.

Some fucntions don't need timeslicing within them selves and can be timesliced from the outside and used on both code paths. Currently no functions are called from the non time slicing code that return a FTimeSliceWorkResultData type.

I've increaced the time slice regen times a bit to account for the fact we'll be going over the time slice limit less often now.

diff --git a/Engine/Source/Runtime/Engine/Public/EngineStats.h b/Engine/Source/Runtime/Engine/Public/EngineStats.h
--- a/Engine/Source/Runtime/Engine/Public/EngineStats.h
+++ b/Engine/Source/Runtime/Engine/Public/EngineStats.h
@@ -122,16 +122,19 @@ DECLARE_CYCLE_STAT_EXTERN(TEXT("Recast: sync add generated tiles"), STAT_Navigat
 DECLARE_CYCLE_STAT_EXTERN(TEXT("Recast: tick"),STAT_Navigation_RecastTick,STATGROUP_Navigation, );
 DECLARE_CYCLE_STAT_EXTERN(TEXT("Recast: pathfinding"),STAT_Navigation_RecastPathfinding, STATGROUP_Navigation, );
 DECLARE_CYCLE_STAT_EXTERN(TEXT("Recast: test path"),STAT_Navigation_RecastTestPath, STATGROUP_Navigation, );
+DECLARE_CYCLE_STAT_EXTERN(TEXT("Recast: storing compressed layers"), STAT_RecastNavMeshGenerator_StoringCompressedLayers, STATGROUP_Navigation, );
 DECLARE_CYCLE_STAT_EXTERN(TEXT("Recast: build compressed layers"),STAT_Navigation_RecastBuildCompressedLayers,STATGROUP_Navigation, );
 DECLARE_CYCLE_STAT_EXTERN(TEXT("Recast: create height field"),STAT_Navigation_RecastCreateHeightField,STATGROUP_Navigation, );
 DECLARE_CYCLE_STAT_EXTERN(TEXT("Recast: rasterize triangles"),STAT_Navigation_RecastRasterizeTriangles,STATGROUP_Navigation, );
 DECLARE_CYCLE_STAT_EXTERN(TEXT("Recast: create voxel filter"),STAT_Navigation_RecastVoxelFilter,STATGROUP_Navigation, );
 DECLARE_CYCLE_STAT_EXTERN(TEXT("Recast: create filters"),STAT_Navigation_RecastFilter,STATGROUP_Navigation, );
+DECLARE_CYCLE_STAT_EXTERN(TEXT("Recast: filter ledge spans"), STAT_Navigation_FilterLedgeSpans, STATGROUP_Navigation, );
 DECLARE_CYCLE_STAT_EXTERN(TEXT("Recast: build compact height field"),STAT_Navigation_RecastBuildCompactHeightField,STATGROUP_Navigation, );
 DECLARE_CYCLE_STAT_EXTERN(TEXT("Recast: erode walkable areas"),STAT_Navigation_RecastErodeWalkable,STATGROUP_Navigation, );
 DECLARE_CYCLE_STAT_EXTERN(TEXT("Recast: build layers"),STAT_Navigation_RecastBuildLayers,STATGROUP_Navigation, );
 DECLARE_CYCLE_STAT_EXTERN(TEXT("Recast: build tile cache"),STAT_Navigation_RecastBuildTileCache,STATGROUP_Navigation, );
 DECLARE_CYCLE_STAT_EXTERN(TEXT("Recast: build navmesh"),STAT_Navigation_RecastBuildNavigation,STATGROUP_Navigation, );
+DECLARE_CYCLE_STAT_EXTERN(TEXT("Recast: build navigation layer"), STAT_Navigation_GenerateNavigationDataLayer, STATGROUP_Navigation, );
 DECLARE_CYCLE_STAT_EXTERN(TEXT("Recast: build regions"), STAT_Navigation_RecastBuildRegions, STATGROUP_Navigation, );
 DECLARE_CYCLE_STAT_EXTERN(TEXT("Recast: mark areas"), STAT_Navigation_RecastMarkAreas, STATGROUP_Navigation, );
 DECLARE_CYCLE_STAT_EXTERN(TEXT("Recast: build contours"), STAT_Navigation_RecastBuildContours, STATGROUP_Navigation, );
diff --git a/Engine/Source/Runtime/NavigationSystem/Private/NavMesh/RecastNavMeshGenerator.cpp b/Engine/Source/Runtime/NavigationSystem/Private/NavMesh/RecastNavMeshGenerator.cpp
--- a/Engine/Source/Runtime/NavigationSystem/Private/NavMesh/RecastNavMeshGenerator.cpp
+++ b/Engine/Source/Runtime/NavigationSystem/Private/NavMesh/RecastNavMeshGenerator.cpp
@@ -1535,13 +1535,37 @@ static FBox CalculateTileBounds(int32 X, int32 Y, const FVector& RcNavMeshOrigin
 	return TileBox;
 }
 
+void FTimeSlicer::SetTimeSliceDuration(double SliceDuration)
+{
+	TimeSliceDuration = SliceDuration;
+}
+
+void FTimeSlicer::ResetStartTime()
+{
+	TimeSliceStartTime = FPlatformTime::Seconds();
+	bTimeSliceFinishedCached = false;
+}
+
+bool FTimeSlicer::TestTimeSliceFinished() const
+{
+	ensureMsgf(!bTimeSliceFinishedCached, TEXT("Testing time slice is finished when we have already confirmed that!"));
+
+	bTimeSliceFinishedCached = FPlatformTime::Seconds() - TimeSliceStartTime >= TimeSliceDuration;
+	return bTimeSliceFinishedCached;
+}
+
+bool FTimeSlicer::IsTimeSliceFinishedCached() const
+{
+	return bTimeSliceFinishedCached;
+}
+
 //----------------------------------------------------------------------//
 // FRecastTileGenerator
 //----------------------------------------------------------------------//
 
 FRecastTileGenerator::FRecastTileGenerator(FRecastNavMeshGenerator& ParentGenerator, const FIntPoint& Location)
+	: TimeSlicer(ParentGenerator.GetTimeSlicer())
 {
-	bSucceeded = false;
 	bUpdateGeometry = true;
 	bHasLowAreaModifiers = false;
 
@@ -1554,15 +1578,23 @@ FRecastTileGenerator::FRecastTileGenerator(FRecastNavMeshGenerator& ParentGenera
 
 	ParentGeneratorWeakPtr = ParentGenerator.AsShared();
 
-#if TIME_SLICE_NAV_REGEN
-	bDoneAsyncDataGathering = false;
-	bDoneRegenerateCompressedLayers = false;
-	bDoneWork = false;
-#endif
+	RasterizeGeomRecastState = ERasterizeGeomRecastTimeSlicedState::MarkWalkableTriangles;
+	RasterizeGeomState = ERasterizeGeomTimeSlicedState::RasterizeGeometryTransformCoords;
+	DoWorkTimeSlicedState = EDoWorkTimeSlicedState::DoAsyncGeometryGathering;
+	GenerateTileTimeSlicedState = EGenerateTileTimeSlicedState::GenerateCompressedLayers;
+
+	GenerateNavDataTimeSlicedState = EGenerateNavDataTimeSlicedState::Init;
+	GenNavDataLayerTimeSlicedIdx = 0;
+	GenCompressedLayersTimeSlicedState = EGenerateCompressedLayersTimeSliced::Init;
+	RasterizeTrianglesTimeSlicedRawGeomIdx = 0;
+	RasterizeTrianglesTimeSlicedInstTransformIdx = 0;
 }
 
 FRecastTileGenerator::~FRecastTileGenerator()
 {
+	GenNavDataTimeSlicedGenerationContext.Release();
+	GenNavDataTimeSlicedAllocator.Release();
+	GenCompressedlayersTimeSlicedRasterContext.Release();
 }
 
 void FRecastTileGenerator::Setup(const FRecastNavMeshGenerator& ParentGenerator, const TArray<FBox>& DirtyAreas)
@@ -1654,46 +1686,78 @@ bool FRecastTileGenerator::HasDataToBuild() const
 		|| (InclusionBounds.Num() && NavigationRelevantData.Num() > 0);
 }
 
-ETimeSliceWorkResult FRecastTileGenerator::DoWork()
+ETimeSliceWorkResult FRecastTileGenerator::DoWorkTimeSliced()
 {
 	TSharedPtr<FNavDataGenerator, ESPMode::ThreadSafe> ParentGenerator = ParentGeneratorWeakPtr.Pin();
-
-	ETimeSliceWorkResult TimeSliceWorkResult = ETimeSliceWorkResult::Succeeded;
+	ETimeSliceWorkResult WorkResult = ETimeSliceWorkResult::Succeeded;
 
 	if (ParentGenerator.IsValid())
 	{
-		if (InclusionBounds.Num()
-#if TIME_SLICE_NAV_REGEN
-			&& !bDoneAsyncDataGathering
-#endif
-			)
+		switch (DoWorkTimeSlicedState)
 		{
-			DoAsyncGeometryGathering();
+		case EDoWorkTimeSlicedState::Invalid:
+		{
+			ensureMsgf(false, TEXT("Invalid EDoWorkTimeSlicedState, has this function been called when its already finished processing?"));
+			return ETimeSliceWorkResult::Failed;
 		}
+		break;
+		case EDoWorkTimeSlicedState::DoAsyncGeometryGathering:
+		{
+			DoWorkTimeSlicedState = EDoWorkTimeSlicedState::GenerateTile;
+
+			if (InclusionBounds.Num())
+			{
+				const bool bHadNavigationRelevantData = DoAsyncGeometryGathering();
+
+				//check bHadNavigationRelevantData as an optimization so we don't call TestTimeSliceFinished when its unnecessary 
+				if (bHadNavigationRelevantData && TimeSlicer.TestTimeSliceFinished())
+				{
+					return ETimeSliceWorkResult::CallAgainNextTimeSlice;
+				}
+			}
+		} //fall through to next state
+		case EDoWorkTimeSlicedState::GenerateTile:
+		{
+			WorkResult = GenerateTileTimeSliced();
+
+			if (WorkResult != ETimeSliceWorkResult::CallAgainNextTimeSlice)
+			{
+				DumpAsyncData();
 
-		TimeSliceWorkResult = GenerateTile();
+				DoWorkTimeSlicedState = EDoWorkTimeSlicedState::Invalid;//Set to Invalid as we never want to call this again on this instance
+			}
+		}
+		break;
 
-		//bSucceeded is set false in the constructor anyway
-		bSucceeded = (TimeSliceWorkResult == ETimeSliceWorkResult::Succeeded);
+		default:
+		{
+			ensureMsgf(false, TEXT("unhandled EDoWorkTimeSlicedState"));
+			return ETimeSliceWorkResult::Failed;
+		}
+		}
 	}
 
-#if TIME_SLICE_NAV_REGEN
-	if (TimeSliceWorkResult != ETimeSliceWorkResult::CallAgain)
-#endif
-	{
-		DumpAsyncData();
+	return WorkResult;
+}
 
-		if (!bSucceeded)
+bool FRecastTileGenerator::DoWork()
+{
+	TSharedPtr<FNavDataGenerator, ESPMode::ThreadSafe> ParentGenerator = ParentGeneratorWeakPtr.Pin();
+	bool bSucceess = true;
+
+	if (ParentGenerator.IsValid())
+	{
+		if (InclusionBounds.Num())
 		{
-			TimeSliceWorkResult = ETimeSliceWorkResult::Failed;
+			DoAsyncGeometryGathering();
 		}
 
-#if TIME_SLICE_NAV_REGEN
-		bDoneWork = true;
-#endif
+		bSucceess = GenerateTile();
+
+		DumpAsyncData();
 	}
 
-	return TimeSliceWorkResult;
+	return bSucceess;
 }
 
 void FRecastTileGenerator::DumpAsyncData()
@@ -1706,13 +1770,11 @@ void FRecastTileGenerator::DumpAsyncData()
 	NavOctree = nullptr;
 }
 
-void FRecastTileGenerator::DoAsyncGeometryGathering()
+bool FRecastTileGenerator::DoAsyncGeometryGathering()
 {
 	QUICK_SCOPE_CYCLE_COUNTER(STAT_RecastNavMeshGenerator_PrepareGeometrySources);
-	
-#if TIME_SLICE_NAV_REGEN
-	bDoneAsyncDataGathering = true;
-#endif
+
+	const bool bRetVal = NavigationRelevantData.Num() > 0;
 
 	for (auto& ElementData : NavigationRelevantData)
 	{
@@ -1797,6 +1859,7 @@ void FRecastTileGenerator::DoAsyncGeometryGathering()
 			AppendModifier(ModifierInstance, ElementData->NavDataPerInstanceTransformDelegate);
 		}
 	}
+	return bRetVal;
 }
 
 void FRecastTileGenerator::PrepareGeometrySources(const FRecastNavMeshGenerator& ParentGenerator, bool bGeometryChanged)
@@ -2195,54 +2258,94 @@ void FRecastTileGenerator::AppendGeometry(const TNavStatArray<uint8>& RawCollisi
 	}	
 }
 
-ETimeSliceWorkResult FRecastTileGenerator::GenerateTile()
+ETimeSliceWorkResult FRecastTileGenerator::GenerateTileTimeSliced()
 {
-	ETimeSliceWorkResult TimeSliceWorkResult = ETimeSliceWorkResult::Succeeded;
 	FNavMeshBuildContext BuildContext;
+	ETimeSliceWorkResult WorkResult = ETimeSliceWorkResult::Succeeded;
 
-	if (bRegenerateCompressedLayers
-#if TIME_SLICE_NAV_REGEN
-		&& !bDoneRegenerateCompressedLayers
-#endif
-		)
+	switch (GenerateTileTimeSlicedState)
 	{
-		CompressedLayers.Reset();
-		
-		const bool bSuccess = GenerateCompressedLayers(BuildContext);
-
-#if TIME_SLICE_NAV_REGEN
-		bDoneRegenerateCompressedLayers = true;
-#endif
-
-		if (bSuccess)
+	case EGenerateTileTimeSlicedState::Invalid:
+	{
+		ensureMsgf(false, TEXT("Invalid EGenerateTileTimeSlicedState, has this function been called when its already finished time processong?"));
+		return ETimeSliceWorkResult::Failed;
+	}
+	break;
+	case EGenerateTileTimeSlicedState::GenerateCompressedLayers:
+	{
+		if (bRegenerateCompressedLayers)
 		{
-			// Mark all layers as dirty
-			DirtyLayers.Init(true, CompressedLayers.Num());
+			const ETimeSliceWorkResult WorkResultCompressed = GenerateCompressedLayersTimeSliced(BuildContext);
 
-#if TIME_SLICE_NAV_REGEN
-			TimeSliceWorkResult = ETimeSliceWorkResult::CallAgain;
-#else
-			TimeSliceWorkResult = ETimeSliceWorkResult::Succeeded;
-#endif
+			if (WorkResultCompressed == ETimeSliceWorkResult::Succeeded)
+			{
+				GenerateTileTimeSlicedState = EGenerateTileTimeSlicedState::GenerateNavigationData;
+				// Mark all layers as dirty
+				DirtyLayers.Init(true, CompressedLayers.Num());
+			}
+			else if (WorkResultCompressed == ETimeSliceWorkResult::Failed)
+			{
+				GenerateTileTimeSlicedState = EGenerateTileTimeSlicedState::Invalid;
+				return ETimeSliceWorkResult::Failed;
+			}
+
+			if (TimeSlicer.IsTimeSliceFinishedCached())
+			{
+				return ETimeSliceWorkResult::CallAgainNextTimeSlice;
+			}
 		}
 		else
 		{
-			TimeSliceWorkResult = ETimeSliceWorkResult::Failed;
+			GenerateTileTimeSlicedState = EGenerateTileTimeSlicedState::GenerateNavigationData;
+		}
+	} //fall through to next state
+	case EGenerateTileTimeSlicedState::GenerateNavigationData:
+	{
+		WorkResult = GenerateNavigationDataTimeSliced(BuildContext);
+
+		if (WorkResult != ETimeSliceWorkResult::CallAgainNextTimeSlice)
+		{
+			GenerateTileTimeSlicedState = EGenerateTileTimeSlicedState::Invalid;
 		}
 	}
+	break;
+	default:
+	{
+		ensureMsgf(false, TEXT("unhandled EGenerateTileTimeSlicedState"));
+		return ETimeSliceWorkResult::Failed;
+	}
+	};
+
+	// it's possible to have valid generation with empty resulting tile (no navigable geometry in tile)
+	return WorkResult;
+}
+
+
+bool FRecastTileGenerator::GenerateTile()
+{
+	FNavMeshBuildContext BuildContext;
+	bool bSuccess = true;
 
-	if (TimeSliceWorkResult == ETimeSliceWorkResult::Succeeded)
+	if (bRegenerateCompressedLayers)
 	{
-		const bool bSuccess = GenerateNavigationData(BuildContext);
+		CompressedLayers.Reset();
 
-		if (!bSuccess)
+		bSuccess = GenerateCompressedLayers(BuildContext);
+
+		if (bSuccess)
 		{
-			TimeSliceWorkResult = ETimeSliceWorkResult::Failed;
+			// Mark all layers as dirty
+			DirtyLayers.Init(true, CompressedLayers.Num());
 		}
 	}
 
+	if (bSuccess)
+	{
+		bSuccess = GenerateNavigationData(BuildContext);
+	}
+
 	// it's possible to have valid generation with empty resulting tile (no navigable geometry in tile)
-	return TimeSliceWorkResult;
+	return bSuccess;
 }
 
 struct FTileRasterizationContext
@@ -2264,55 +2367,12 @@ struct FTileRasterizationContext
 	TArray<FNavMeshTileData> Layers;
 };
 
-static void RasterizeGeometry(
-	FNavMeshBuildContext& BuildContext, const FRecastBuildConfig& TileConfig, 
-	const TArray<float>& Coords, const TArray<int32>& Indices, FTileRasterizationContext& RasterContext)
-{
-	const int32 NumFaces = Indices.Num() / 3;
-	const int32 NumVerts = Coords.Num() / 3;
-
-	TNavStatArray<uint8> TriAreas;
-	TriAreas.Reserve(NumFaces);
-	TriAreas.AddZeroed(NumFaces);
-
-	rcMarkWalkableTriangles(&BuildContext, TileConfig.walkableSlopeAngle,
-		Coords.GetData(), NumVerts, Indices.GetData(), NumFaces,
-		TriAreas.GetData());
-
-	rcRasterizeTriangles(&BuildContext,
-		Coords.GetData(), NumVerts, 
-		Indices.GetData(), TriAreas.GetData(), NumFaces,
-		*RasterContext.SolidHF, TileConfig.walkableClimb);
-}
-
-static void RasterizeGeometry(
-	FNavMeshBuildContext& BuildContext,const FRecastBuildConfig& TileConfig, 
-	const TArray<float>& Coords, const TArray<int32>& Indices, const FTransform& LocalToWorld, FTileRasterizationContext& RasterContext)
+bool FRecastTileGenerator::CreateHeightField(FNavMeshBuildContext& BuildContext, FTileRasterizationContext& RasterContext)
 {
-	TArray<float> WorldRecastCoords;
-	WorldRecastCoords.SetNumUninitialized(Coords.Num());
-	
-	FMatrix LocalToRecastWorld = LocalToWorld.ToMatrixWithScale()*Unreal2RecastMatrix();
-	// Convert geometry to recast world space
-	for (int32 i = 0; i < Coords.Num(); i+=3)
-	{
-		// collision cache stores coordinates in recast space, convert them to unreal and transform to recast world space
-		FVector WorldRecastCoord = LocalToRecastWorld.TransformPosition(Recast2UnrealPoint(&Coords[i])); 
-
-		WorldRecastCoords[i+0] = WorldRecastCoord.X;
-		WorldRecastCoords[i+1] = WorldRecastCoord.Y;
-		WorldRecastCoords[i+2] = WorldRecastCoord.Z;
-	}
+	SCOPE_CYCLE_COUNTER(STAT_Navigation_RecastCreateHeightField);
 
-	RasterizeGeometry(BuildContext, TileConfig, WorldRecastCoords, Indices, RasterContext);
-}
-
-bool FRecastTileGenerator::GenerateCompressedLayers(FNavMeshBuildContext& BuildContext)
-{
-	SCOPE_CYCLE_COUNTER(STAT_Navigation_RecastBuildCompressedLayers);
-	
-	TileConfig.width = TileConfig.tileSize + TileConfig.borderSize*2;
-	TileConfig.height = TileConfig.tileSize + TileConfig.borderSize*2;
+	TileConfig.width = TileConfig.tileSize + TileConfig.borderSize * 2;
+	TileConfig.height = TileConfig.tileSize + TileConfig.borderSize * 2;
 
 	const float BBoxPadding = TileConfig.borderSize * TileConfig.cs;
 	TileConfig.bmin[0] -= BBoxPadding;
@@ -2320,289 +2380,715 @@ bool FRecastTileGenerator::GenerateCompressedLayers(FNavMeshBuildContext& BuildC
 	TileConfig.bmax[0] += BBoxPadding;
 	TileConfig.bmax[2] += BBoxPadding;
 
-	BuildContext.log(RC_LOG_PROGRESS, "GenerateCompressedLayers:");
+	BuildContext.log(RC_LOG_PROGRESS, "CreateHeightField:");
 	BuildContext.log(RC_LOG_PROGRESS, " - %d x %d cells", TileConfig.width, TileConfig.height);
 
-	FTileRasterizationContext RasterContext;
 	const bool bHasGeometry = RawGeometry.Num() > 0;
 
 	// Allocate voxel heightfield where we rasterize our input data to.
 	if (bHasGeometry)
 	{
+		RasterContext.SolidHF = rcAllocHeightfield();
+		if (RasterContext.SolidHF == nullptr)
 		{
-			SCOPE_CYCLE_COUNTER(STAT_Navigation_RecastCreateHeightField);
-
-			RasterContext.SolidHF = rcAllocHeightfield();
-			if (RasterContext.SolidHF == NULL)
-			{
-				BuildContext.log(RC_LOG_ERROR, "GenerateCompressedLayers: Out of memory 'SolidHF'.");
-				return false;
-			}
-			if (!rcCreateHeightfield(&BuildContext, *RasterContext.SolidHF, TileConfig.width, TileConfig.height, TileConfig.bmin, TileConfig.bmax, TileConfig.cs, TileConfig.ch))
-			{
-				BuildContext.log(RC_LOG_ERROR, "GenerateCompressedLayers: Could not create solid heightfield.");
-				return false;
-			}
+			BuildContext.log(RC_LOG_ERROR, "CreateHeightField: Out of memory 'SolidHF'.");
+			return false;
+		}
+		if (!rcCreateHeightfield(&BuildContext, *RasterContext.SolidHF, TileConfig.width, TileConfig.height, TileConfig.bmin, TileConfig.bmax, TileConfig.cs, TileConfig.ch))
+		{
+			BuildContext.log(RC_LOG_ERROR, "CreateHeightField: Could not create solid heightfield.");
+			return false;
 		}
+	}
+	return true;
+}
+
+ETimeSliceWorkResult FRecastTileGenerator::RasterizeGeometryRecastTimeSliced(FNavMeshBuildContext& BuildContext, const TArray<float>& Coords, const TArray<int32>& Indices, FTileRasterizationContext& RasterContext)
+{
+	QUICK_SCOPE_CYCLE_COUNTER(STAT_Navigation_RasterizeGeometryRecast);
+
+	const int32 NumFaces = Indices.Num() / 3;
+	const int32 NumVerts = Coords.Num() / 3;
+
+	switch (RasterizeGeomRecastState)
+	{
+	case ERasterizeGeomRecastTimeSlicedState::MarkWalkableTriangles:
+	{
+		QUICK_SCOPE_CYCLE_COUNTER(STAT_Navigation_MarkWalkableTriangles);
+
+		RasterizeGeomRecastTriAreas.AddZeroed(NumFaces);
+
+		rcMarkWalkableTriangles(&BuildContext, TileConfig.walkableSlopeAngle,
+			Coords.GetData(), NumVerts, Indices.GetData(), NumFaces,
+			RasterizeGeomRecastTriAreas.GetData());
+
+		RasterizeGeomRecastState = ERasterizeGeomRecastTimeSlicedState::RasterizeTriangles;
 
+		if (TimeSlicer.TestTimeSliceFinished())
 		{
-			// Rasterize geometry
-			SCOPE_CYCLE_COUNTER(STAT_Navigation_RecastRasterizeTriangles)
-		
-			for (const FRecastRawGeometryElement& Element : RawGeometry)
-			{
-				for (const FTransform& InstanceTransform : Element.PerInstanceTransform)
-				{
-					RasterizeGeometry(BuildContext, TileConfig, Element.GeomCoords, Element.GeomIndices, InstanceTransform, RasterContext);
-				}
-			
-				if (Element.PerInstanceTransform.Num() == 0)
-				{
-					RasterizeGeometry(BuildContext, TileConfig, Element.GeomCoords, Element.GeomIndices, RasterContext);
-				}
-			}
+			return ETimeSliceWorkResult::CallAgainNextTimeSlice;
 		}
+	}// fall through to next state
+	case ERasterizeGeomRecastTimeSlicedState::RasterizeTriangles:
+	{
+		QUICK_SCOPE_CYCLE_COUNTER(STAT_Navigation_RasterizeGeomRecastRasterizeTriangles);
+
+		rcRasterizeTriangles(&BuildContext,
+			Coords.GetData(), NumVerts,
+			Indices.GetData(), RasterizeGeomRecastTriAreas.GetData(), NumFaces,
+			*RasterContext.SolidHF, TileConfig.walkableClimb);
+
+		RasterizeGeomRecastTriAreas.Reset();
+
+		//reset this so next call we start by marking walkable triangles
+		RasterizeGeomRecastState = ERasterizeGeomRecastTimeSlicedState::MarkWalkableTriangles;
+
+		TimeSlicer.TestTimeSliceFinished();
 	}
-	
-	if (!RasterContext.SolidHF || RasterContext.SolidHF->pools == 0)
+	break;
+	default:
 	{
-		BuildContext.log(RC_LOG_WARNING, "GenerateCompressedLayers: empty tile - aborting");
-		return true;
+		ensureMsgf(false, TEXT("unhandled ERasterizeGeomRecastTimeSlicedState"));
+		return ETimeSliceWorkResult::Failed;
+	}
 	}
+	return ETimeSliceWorkResult::Succeeded;
+}
 
 
-	// Reject voxels outside generation boundaries
-	if (TileConfig.bPerformVoxelFiltering && !bFullyEncapsulatedByInclusionBounds)
+void FRecastTileGenerator::RasterizeGeometryRecast(FNavMeshBuildContext& BuildContext, const TArray<float>& Coords, const TArray<int32>& Indices, FTileRasterizationContext& RasterContext)
+{
+	QUICK_SCOPE_CYCLE_COUNTER(STAT_Navigation_RasterizeGeometryRecast);
+
+	const int32 NumFaces = Indices.Num() / 3;
+	const int32 NumVerts = Coords.Num() / 3;
+
+	RasterizeGeomRecastTriAreas.AddZeroed(NumFaces);
+
 	{
-		SCOPE_CYCLE_COUNTER(STAT_Navigation_RecastVoxelFilter)
+		QUICK_SCOPE_CYCLE_COUNTER(STAT_Navigation_MarkWalkableTriangles);
 
-		ApplyVoxelFilter(RasterContext.SolidHF, TileConfig.walkableRadius);
+		rcMarkWalkableTriangles(&BuildContext, TileConfig.walkableSlopeAngle,
+			Coords.GetData(), NumVerts, Indices.GetData(), NumFaces,
+			RasterizeGeomRecastTriAreas.GetData());
 	}
 
-	// TileConfig.walkableHeight is set to 1 when marking low spans, calculate real value for filtering
-	const int32 FilterWalkableHeight = FMath::CeilToInt(TileConfig.AgentHeight / TileConfig.ch);
-
 	{
-		SCOPE_CYCLE_COUNTER(STAT_Navigation_RecastFilter)
+		QUICK_SCOPE_CYCLE_COUNTER(STAT_Navigation_RasterizeGeomRecastRasterizeTriangles);
 
-		// Once all geometry is rasterized, we do initial pass of filtering to
-		// remove unwanted overhangs caused by the conservative rasterization
-		// as well as filter spans where the character cannot possibly stand.
-		rcFilterLowHangingWalkableObstacles(&BuildContext, TileConfig.walkableClimb, *RasterContext.SolidHF);
-		rcFilterLedgeSpans(&BuildContext, TileConfig.walkableHeight, TileConfig.walkableClimb, *RasterContext.SolidHF);
-		if (!TileConfig.bMarkLowHeightAreas)
-		{
-			rcFilterWalkableLowHeightSpans(&BuildContext, TileConfig.walkableHeight, *RasterContext.SolidHF);
-		}
-		else if (TileConfig.bFilterLowSpanFromTileCache)
-		{
-			// TODO: investigate if creating detailed 2D map from active modifiers is cheap enough
-			// for now, switch on presence of those modifiers, will save memory as long as they are sparse (should be)
+		rcRasterizeTriangles(&BuildContext,
+			Coords.GetData(), NumVerts,
+			Indices.GetData(), RasterizeGeomRecastTriAreas.GetData(), NumFaces,
+			*RasterContext.SolidHF, TileConfig.walkableClimb);
+	}
 
-			if (TileConfig.bFilterLowSpanSequences && bHasLowAreaModifiers)
-			{
-				rcFilterWalkableLowHeightSpansSequences(&BuildContext, FilterWalkableHeight, *RasterContext.SolidHF);
-			}
-			else
-			{
-				rcFilterWalkableLowHeightSpans(&BuildContext, FilterWalkableHeight, *RasterContext.SolidHF);
-			}
-		}
+	RasterizeGeomRecastTriAreas.Reset();
+}
+
+void FRecastTileGenerator::RasterizeGeometryTransformCoords(const TArray<float>& Coords, const FTransform& LocalToWorld)
+{
+	QUICK_SCOPE_CYCLE_COUNTER(STAT_Navigation_RasterizeGeometryTransformCoords);
+
+	RasterizeGeometryWorldRecastCoords.SetNumUninitialized(Coords.Num(), false);
+
+	FMatrix LocalToRecastWorld = LocalToWorld.ToMatrixWithScale()*Unreal2RecastMatrix();
+
+	// Convert geometry to recast world space
+	for (int32 i = 0; i < Coords.Num(); i+=3)
+	{
+		// collision cache stores coordinates in recast space, convert them to unreal and transform to recast world space
+		FVector WorldRecastCoord = LocalToRecastWorld.TransformPosition(Recast2UnrealPoint(&Coords[i]));
+
+		RasterizeGeometryWorldRecastCoords[i+0] = WorldRecastCoord.X;
+		RasterizeGeometryWorldRecastCoords[i+1] = WorldRecastCoord.Y;
+		RasterizeGeometryWorldRecastCoords[i+2] = WorldRecastCoord.Z;
 	}
+}
 
+ETimeSliceWorkResult FRecastTileGenerator::RasterizeGeometryTimeSliced(FNavMeshBuildContext& BuildContext, const TArray<float>& Coords, const TArray<int32>& Indices, const FTransform& LocalToWorld, FTileRasterizationContext& RasterContext)
+{
+	QUICK_SCOPE_CYCLE_COUNTER(STAT_Navigation_RasterizeGeometry);
+	ETimeSliceWorkResult WorkResult = ETimeSliceWorkResult::Succeeded;
+
+	switch (RasterizeGeomState)
 	{
-		SCOPE_CYCLE_COUNTER(STAT_Navigation_RecastBuildCompactHeightField);
+	case ERasterizeGeomTimeSlicedState::RasterizeGeometryTransformCoords:
+	{
+		RasterizeGeometryTransformCoords(Coords, LocalToWorld);
+
+		RasterizeGeomState = ERasterizeGeomTimeSlicedState::RasterizeGeometryRecast;
 
-		// Compact the heightfield so that it is faster to handle from now on.
-		// This will result more cache coherent data as well as the neighbors
-		// between walkable cells will be calculated.
-		RasterContext.CompactHF = rcAllocCompactHeightfield();
-		if (RasterContext.CompactHF == NULL)
+		if (TimeSlicer.TestTimeSliceFinished())
 		{
-			BuildContext.log(RC_LOG_ERROR, "GenerateCompressedLayers: Out of memory 'CompactHF'.");
-			return false;
+			return ETimeSliceWorkResult::CallAgainNextTimeSlice;
 		}
-		if (!rcBuildCompactHeightfield(&BuildContext, TileConfig.walkableHeight, TileConfig.walkableClimb, *RasterContext.SolidHF, *RasterContext.CompactHF))
+	}// fall through to next state
+	case ERasterizeGeomTimeSlicedState::RasterizeGeometryRecast:
+	{
+		WorkResult = RasterizeGeometryRecastTimeSliced(BuildContext, RasterizeGeometryWorldRecastCoords, Indices, RasterContext);
+
+		if (WorkResult != ETimeSliceWorkResult::CallAgainNextTimeSlice)
 		{
-			const int SpanCount = rcGetHeightFieldSpanCount(&BuildContext, *RasterContext.SolidHF);
-			if (SpanCount > 0)
-			{
-				BuildContext.log(RC_LOG_ERROR, "GenerateCompressedLayers: Could not build compact data.");
-			}
-			// else there's just no spans to walk on (no spans at all or too small/sparse)
-			else
-			{
-				BuildContext.log(RC_LOG_WARNING, "GenerateCompressedLayers: no walkable spans - aborting");
-			}
-			return false;
+			//if we have finished rasterizing this geometry then reset RasterizeGeomTimeSlicedState so next time this function is called we go back to RasterizeGeometryTransformCoords first
+			RasterizeGeomState = ERasterizeGeomTimeSlicedState::RasterizeGeometryTransformCoords;
 		}
 	}
-
+	break;
+	default:
 	{
-		SCOPE_CYCLE_COUNTER(STAT_Navigation_RecastErodeWalkable);
+		ensureMsgf(false, TEXT("unhandled ERasterizeGeomTimeSlicedState"));
+		return ETimeSliceWorkResult(ETimeSliceWorkResult::Failed);
+	}
+	}
+	return WorkResult;
+}
+
+void FRecastTileGenerator::RasterizeGeometry(FNavMeshBuildContext& BuildContext, const TArray<float>& Coords, const TArray<int32>& Indices, const FTransform& LocalToWorld, FTileRasterizationContext& RasterContext)
+{
+	QUICK_SCOPE_CYCLE_COUNTER(STAT_Navigation_RasterizeGeometry);
+
+	RasterizeGeometryTransformCoords(Coords, LocalToWorld);
+	RasterizeGeometryRecast(BuildContext, RasterizeGeometryWorldRecastCoords, Indices, RasterContext);
+}
+
+ETimeSliceWorkResult FRecastTileGenerator::RasterizeTrianglesTimeSliced(FNavMeshBuildContext& BuildContext, FTileRasterizationContext& RasterContext)
+{
+	// Rasterize geometry
+	SCOPE_CYCLE_COUNTER(STAT_Navigation_RecastRasterizeTriangles)
 
-		if (TileConfig.walkableRadius > RECAST_VERY_SMALL_AGENT_RADIUS)
+	while (RasterizeTrianglesTimeSlicedRawGeomIdx < RawGeometry.Num())
+	{
+		const FRecastRawGeometryElement& Element = RawGeometry[RasterizeTrianglesTimeSlicedRawGeomIdx];
+		if (Element.PerInstanceTransform.Num() > 0)
 		{
-			uint8 FilterFlags = 0;
-			if (TileConfig.bFilterLowSpanSequences)
+			while (RasterizeTrianglesTimeSlicedInstTransformIdx < Element.PerInstanceTransform.Num())
 			{
-				FilterFlags = RC_LOW_FILTER_POST_PROCESS | (TileConfig.bFilterLowSpanFromTileCache ? 0 : RC_LOW_FILTER_SEED_SPANS);
-			}
+				const FTransform& InstanceTransform = Element.PerInstanceTransform[RasterizeTrianglesTimeSlicedInstTransformIdx];
+				const ETimeSliceWorkResult WorkResult = RasterizeGeometryTimeSliced(BuildContext, Element.GeomCoords, Element.GeomIndices, InstanceTransform, RasterContext);
+
+				//the original code just kept calling the RasterizeGeometry() functions and had no return type, 
+				//so we will process the next layer (if we are not needing to process this layer again next time slice) 
+				if (TimeSlicer.IsTimeSliceFinishedCached())
+				{
+					if (WorkResult != ETimeSliceWorkResult::CallAgainNextTimeSlice)
+					{
+						++RasterizeTrianglesTimeSlicedInstTransformIdx;
+					}
 
-			const bool bEroded = TileConfig.bMarkLowHeightAreas ?
-				rcErodeWalkableAndLowAreas(&BuildContext, TileConfig.walkableRadius, FilterWalkableHeight, RECAST_LOW_AREA, FilterFlags, *RasterContext.CompactHF) :
-				rcErodeWalkableArea(&BuildContext, TileConfig.walkableRadius, *RasterContext.CompactHF);
+					return ETimeSliceWorkResult::CallAgainNextTimeSlice;
+				}
 
-			if (!bEroded)
-			{
-				BuildContext.log(RC_LOG_ERROR, "GenerateCompressedLayers: Could not erode.");
-				return false;
+				++RasterizeTrianglesTimeSlicedInstTransformIdx;
 			}
+			//reset RasterizeTrianglesTimeSlicedIdx 
+			RasterizeTrianglesTimeSlicedInstTransformIdx = 0;
 		}
-		else if (TileConfig.bMarkLowHeightAreas)
+		else
 		{
-			rcMarkLowAreas(&BuildContext, FilterWalkableHeight, RECAST_LOW_AREA, *RasterContext.CompactHF);
+			const ETimeSliceWorkResult WorkResult = RasterizeGeometryRecastTimeSliced(BuildContext, Element.GeomCoords, Element.GeomIndices, RasterContext);
+
+			if (TimeSlicer.IsTimeSliceFinishedCached())
+			{
+				if (WorkResult != ETimeSliceWorkResult::CallAgainNextTimeSlice)
+				{
+					++RasterizeTrianglesTimeSlicedRawGeomIdx;
+				}
+				return ETimeSliceWorkResult::CallAgainNextTimeSlice;
+			}
 		}
+		++RasterizeTrianglesTimeSlicedRawGeomIdx;
 	}
 
-	// Build layers
+	//return sucess as non timesliced functionality does not detect failure here
+	return ETimeSliceWorkResult::Succeeded;
+}
+
+void FRecastTileGenerator::RasterizeTriangles(FNavMeshBuildContext& BuildContext, FTileRasterizationContext& RasterContext)
+{
+	// Rasterize geometry
+	SCOPE_CYCLE_COUNTER(STAT_Navigation_RecastRasterizeTriangles)
+
+	for (int32 RawGeomIdx = 0; RawGeomIdx < RawGeometry.Num(); ++RawGeomIdx)
 	{
-		SCOPE_CYCLE_COUNTER(STAT_Navigation_RecastBuildLayers);
-		
-		RasterContext.LayerSet = rcAllocHeightfieldLayerSet();
-		if (RasterContext.LayerSet == NULL)
+		const FRecastRawGeometryElement& Element = RawGeometry[RawGeomIdx];
+		if (Element.PerInstanceTransform.Num() > 0)
 		{
-			BuildContext.log(RC_LOG_ERROR, "GenerateCompressedLayers: Out of memory 'LayerSet'.");
-			return false;
+			for (const FTransform& InstanceTransform : Element.PerInstanceTransform)
+			{
+				RasterizeGeometry(BuildContext, Element.GeomCoords, Element.GeomIndices, InstanceTransform, RasterContext);
+			}
+		}
+		else
+		{
+			RasterizeGeometryRecast(BuildContext, Element.GeomCoords, Element.GeomIndices, RasterContext);
 		}
+	}
+}
+
+void FRecastTileGenerator::GenerateRecastFilter(FNavMeshBuildContext& BuildContext, FTileRasterizationContext& RasterContext)
+{
+	SCOPE_CYCLE_COUNTER(STAT_Navigation_RecastFilter)
+
+	// TileConfig.walkableHeight is set to 1 when marking low spans, calculate real value for filtering
+	const int32 FilterWalkableHeight = FMath::CeilToInt(TileConfig.AgentHeight / TileConfig.ch);
+
+	// Once all geometry is rasterized, we do initial pass of filtering to
+	// remove unwanted overhangs caused by the conservative rasterization
+	// as well as filter spans where the character cannot possibly stand.
+	{
+		rcFilterLowHangingWalkableObstacles(&BuildContext, TileConfig.walkableClimb, *RasterContext.SolidHF);
+	}
+	{
+		SCOPE_CYCLE_COUNTER(STAT_Navigation_FilterLedgeSpans)
+
+		rcFilterLedgeSpans(&BuildContext, TileConfig.walkableHeight, TileConfig.walkableClimb, *RasterContext.SolidHF);
+	}
+	if (!TileConfig.bMarkLowHeightAreas)
+	{
+		rcFilterWalkableLowHeightSpans(&BuildContext, TileConfig.walkableHeight, *RasterContext.SolidHF);
+	}
+	else if (TileConfig.bFilterLowSpanFromTileCache)
+	{
+		// TODO: investigate if creating detailed 2D map from active modifiers is cheap enough
+		// for now, switch on presence of those modifiers, will save memory as long as they are sparse (should be)
 
-		if (TileConfig.regionPartitioning == RC_REGION_MONOTONE)
+		if (TileConfig.bFilterLowSpanSequences && bHasLowAreaModifiers)
 		{
-			if (!rcBuildHeightfieldLayersMonotone(&BuildContext, *RasterContext.CompactHF, TileConfig.borderSize, TileConfig.walkableHeight, *RasterContext.LayerSet))
-			{
-				BuildContext.log(RC_LOG_ERROR, "GenerateCompressedLayers: Could not build heightfield layers.");
-				return 0;
-			}
+			rcFilterWalkableLowHeightSpansSequences(&BuildContext, FilterWalkableHeight, *RasterContext.SolidHF);
 		}
-		else if (TileConfig.regionPartitioning == RC_REGION_WATERSHED)
+		else
 		{
-			if (!rcBuildDistanceField(&BuildContext, *RasterContext.CompactHF))
-			{
-				BuildContext.log(RC_LOG_ERROR, "GenerateCompressedLayers: Could not build distance field.");
-				return 0;
-			}
+			rcFilterWalkableLowHeightSpans(&BuildContext, FilterWalkableHeight, *RasterContext.SolidHF);
+		}
+	}
+}
 
-			if (!rcBuildHeightfieldLayers(&BuildContext, *RasterContext.CompactHF, TileConfig.borderSize, TileConfig.walkableHeight, *RasterContext.LayerSet))
-			{
-				BuildContext.log(RC_LOG_ERROR, "GenerateCompressedLayers: Could not build heightfield layers.");
-				return 0;
-			}
+bool FRecastTileGenerator::BuildCompactHeightField(FNavMeshBuildContext& BuildContext, FTileRasterizationContext& RasterContext)
+{
+	SCOPE_CYCLE_COUNTER(STAT_Navigation_RecastBuildCompactHeightField);
+
+	// Compact the heightfield so that it is faster to handle from now on.
+	// This will result more cache coherent data as well as the neighbors
+	// between walkable cells will be calculated.
+	RasterContext.CompactHF = rcAllocCompactHeightfield();
+	if (RasterContext.CompactHF == nullptr)
+	{
+		BuildContext.log(RC_LOG_ERROR, "BuildCompactHeightField: Out of memory 'CompactHF'.");
+		return false;
+	}
+	if (!rcBuildCompactHeightfield(&BuildContext, TileConfig.walkableHeight, TileConfig.walkableClimb, *RasterContext.SolidHF, *RasterContext.CompactHF))
+	{
+		const int SpanCount = rcGetHeightFieldSpanCount(&BuildContext, *RasterContext.SolidHF);
+		if (SpanCount > 0)
+		{
+			BuildContext.log(RC_LOG_ERROR, "BuildCompactHeightField: Could not build compact data.");
 		}
+		// else there's just no spans to walk on (no spans at all or too small/sparse)
 		else
 		{
-			if (!rcBuildHeightfieldLayersChunky(&BuildContext, *RasterContext.CompactHF, TileConfig.borderSize, TileConfig.walkableHeight, TileConfig.regionChunkSize, *RasterContext.LayerSet))
-			{
-				BuildContext.log(RC_LOG_ERROR, "GenerateCompressedLayers: Could not build heightfield layers.");
-				return 0;
-			}
+			BuildContext.log(RC_LOG_WARNING, "BuildCompactHeightField: no walkable spans - aborting");
 		}
+		return false;
 	}
+	return true;
+}
 
+bool FRecastTileGenerator::RecastErodeWalkable(FNavMeshBuildContext& BuildContext, FTileRasterizationContext& RasterContext)
+{
+	SCOPE_CYCLE_COUNTER(STAT_Navigation_RecastErodeWalkable);
+
+	// TileConfig.walkableHeight is set to 1 when marking low spans, calculate real value for filtering
+	const int32 FilterWalkableHeight = FMath::CeilToInt(TileConfig.AgentHeight / TileConfig.ch);
+
+	if (TileConfig.walkableRadius > RECAST_VERY_SMALL_AGENT_RADIUS)
 	{
-		SCOPE_CYCLE_COUNTER(STAT_Navigation_RecastBuildTileCache);
+		uint8 FilterFlags = 0;
+		if (TileConfig.bFilterLowSpanSequences)
+		{
+			FilterFlags = RC_LOW_FILTER_POST_PROCESS | (TileConfig.bFilterLowSpanFromTileCache ? 0 : RC_LOW_FILTER_SEED_SPANS);
+		}
 
-		const int32 NumLayers = RasterContext.LayerSet->nlayers;
-	
-		// use this to expand vertically layer's bounds
-		// this is needed to allow off-mesh connections that are not quite
-		// touching tile layer still connect with it.
-		const float StepHeights = TileConfig.AgentMaxClimb;
-
-		FTileCacheCompressor TileCompressor;
-		for (int32 i = 0; i < NumLayers; i++)
-		{
-			const rcHeightfieldLayer* layer = &RasterContext.LayerSet->layers[i];
-
-			// Store header
-			dtTileCacheLayerHeader header;
-			header.magic = DT_TILECACHE_MAGIC;
-			header.version = DT_TILECACHE_VERSION;
-
-			// Tile layer location in the navmesh.
-			header.tx = TileX;
-			header.ty = TileY;
-			header.tlayer = i;
-			dtVcopy(header.bmin, layer->bmin);
-			dtVcopy(header.bmax, layer->bmax);
-
-			// Tile info.
-			header.width = (unsigned short)layer->width;
-			header.height = (unsigned short)layer->height;
-			header.minx = (unsigned short)layer->minx;
-			header.maxx = (unsigned short)layer->maxx;
-			header.miny = (unsigned short)layer->miny;
-			header.maxy = (unsigned short)layer->maxy;
-			header.hmin = (unsigned short)layer->hmin;
-			header.hmax = (unsigned short)layer->hmax;
-
-			// Layer bounds in unreal coords
-			FBox LayerBBox = Recast2UnrealBox(header.bmin, header.bmax);
-			LayerBBox.Min.Z -= StepHeights;
-			LayerBBox.Max.Z += StepHeights;
-			
-			// Compress tile layer
-			uint8* TileData = NULL;
-			int32 TileDataSize = 0;
-			const dtStatus status = dtBuildTileCacheLayer(&TileCompressor, &header, layer->heights, layer->areas, layer->cons, &TileData, &TileDataSize);
-			if (dtStatusFailed(status))
-			{
-				dtFree(TileData);
-				BuildContext.log(RC_LOG_ERROR, "GenerateCompressedLayers: failed to build layer.");
-				return false;
-			}
+		const bool bEroded = TileConfig.bMarkLowHeightAreas ?
+			rcErodeWalkableAndLowAreas(&BuildContext, TileConfig.walkableRadius, FilterWalkableHeight, RECAST_LOW_AREA, FilterFlags, *RasterContext.CompactHF) :
+			rcErodeWalkableArea(&BuildContext, TileConfig.walkableRadius, *RasterContext.CompactHF);
+
+		if (!bEroded)
+		{
+			BuildContext.log(RC_LOG_ERROR, "GenerateCompressedLayers: Could not erode.");
+			return false;
+		}
+
+	}
+	else if (TileConfig.bMarkLowHeightAreas)
+	{
+		rcMarkLowAreas(&BuildContext, FilterWalkableHeight, RECAST_LOW_AREA, *RasterContext.CompactHF);
+	}
+
+	return true;
+}
+
+bool FRecastTileGenerator::RecastBuildLayers(FNavMeshBuildContext& BuildContext, FTileRasterizationContext& RasterContext)
+{
+	SCOPE_CYCLE_COUNTER(STAT_Navigation_RecastBuildLayers);
+
+	RasterContext.LayerSet = rcAllocHeightfieldLayerSet();
+	if (RasterContext.LayerSet == nullptr)
+	{
+		BuildContext.log(RC_LOG_ERROR, "RecastBuildLayers: Out of memory 'LayerSet'.");
+		return false;
+	}
+
+	if (TileConfig.regionPartitioning == RC_REGION_MONOTONE)
+	{
+		if (!rcBuildHeightfieldLayersMonotone(&BuildContext, *RasterContext.CompactHF, TileConfig.borderSize, TileConfig.walkableHeight, *RasterContext.LayerSet))
+		{
+			BuildContext.log(RC_LOG_ERROR, "RecastBuildLayers: Could not build heightfield layers.");
+			return false;
+
+		}
+	}
+	else if (TileConfig.regionPartitioning == RC_REGION_WATERSHED)
+	{
+		if (!rcBuildDistanceField(&BuildContext, *RasterContext.CompactHF))
+		{
+			BuildContext.log(RC_LOG_ERROR, "RecastBuildLayers: Could not build distance field.");
+			return false;
+		}
+
+		if (!rcBuildHeightfieldLayers(&BuildContext, *RasterContext.CompactHF, TileConfig.borderSize, TileConfig.walkableHeight, *RasterContext.LayerSet))
+		{
+			BuildContext.log(RC_LOG_ERROR, "RecastBuildLayers: Could not build heightfield layers.");
+			return false;
+		}
+	}
+	else
+	{
+		if (!rcBuildHeightfieldLayersChunky(&BuildContext, *RasterContext.CompactHF, TileConfig.borderSize, TileConfig.walkableHeight, TileConfig.regionChunkSize, *RasterContext.LayerSet))
+		{
+			BuildContext.log(RC_LOG_ERROR, "RecastBuildLayers: Could not build heightfield layers.");
+			return false;
+		}
+	}
+	return true;
+}
+
+bool FRecastTileGenerator::RecastBuildTileCache(FNavMeshBuildContext& BuildContext, FTileRasterizationContext& RasterContext)
+{
+	SCOPE_CYCLE_COUNTER(STAT_Navigation_RecastBuildTileCache);
+
+	const int32 NumLayers = RasterContext.LayerSet->nlayers;
+
+	// use this to expand vertically layer's bounds
+	// this is needed to allow off-mesh connections that are not quite
+	// touching tile layer still connect with it.
+	const float StepHeights = TileConfig.AgentMaxClimb;
+
+	FTileCacheCompressor TileCompressor;
+	for (int32 i = 0; i < NumLayers; i++)
+	{
+		const rcHeightfieldLayer* layer = &RasterContext.LayerSet->layers[i];
+
+		// Store header
+		dtTileCacheLayerHeader header;
+		header.magic = DT_TILECACHE_MAGIC;
+		header.version = DT_TILECACHE_VERSION;
+
+		// Tile layer location in the navmesh.
+		header.tx = TileX;
+		header.ty = TileY;
+		header.tlayer = i;
+		dtVcopy(header.bmin, layer->bmin);
+		dtVcopy(header.bmax, layer->bmax);
+
+		// Tile info.
+		header.width = (unsigned short)layer->width;
+		header.height = (unsigned short)layer->height;
+		header.minx = (unsigned short)layer->minx;
+		header.maxx = (unsigned short)layer->maxx;
+		header.miny = (unsigned short)layer->miny;
+		header.maxy = (unsigned short)layer->maxy;
+		header.hmin = (unsigned short)layer->hmin;
+		header.hmax = (unsigned short)layer->hmax;
+
+		// Layer bounds in unreal coords
+		FBox LayerBBox = Recast2UnrealBox(header.bmin, header.bmax);
+		LayerBBox.Min.Z -= StepHeights;
+		LayerBBox.Max.Z += StepHeights;
+
+		// Compress tile layer
+		uint8* TileData = nullptr;
+		int32 TileDataSize = 0;
+		const dtStatus status = dtBuildTileCacheLayer(&TileCompressor, &header, layer->heights, layer->areas, layer->cons, &TileData, &TileDataSize);
+		if (dtStatusFailed(status))
+		{
+			dtFree(TileData);
+			BuildContext.log(RC_LOG_ERROR, "RecastBuildTileCache: failed to build layer.");
+			return false;
+		}
 #if !UE_BUILD_SHIPPING && OUTPUT_NAV_TILE_LAYER_COMPRESSION_DATA
-			else
-			{
-				const int gridSize = (int)header.width * (int)header.height;
-				const int bufferSize = gridSize * 4;
+		else
+		{
+			const int gridSize = (int)header.width * (int)header.height;
+			const int bufferSize = gridSize * 4;
 
-				FPlatformMisc::CustomNamedStat("NavTileLayerUncompSize", static_cast<float>(bufferSize), "NavMesh", "Bytes");
-				FPlatformMisc::CustomNamedStat("NavTileLayerCompSize", static_cast<float>(TileDataSize), "NavMesh", "Bytes");
-			}
+			FPlatformMisc::CustomNamedStat("NavTileLayerUncompSize", static_cast<float>(bufferSize), "NavMesh", "Bytes");
+			FPlatformMisc::CustomNamedStat("NavTileLayerCompSize", static_cast<float>(TileDataSize), "NavMesh", "Bytes");
+		}
 #endif
 
-			// copy compressed data to new buffer in rasterization context
-			// (TileData allocates a lots of space, but only first TileDataSize bytes hold compressed data)
+		// copy compressed data to new buffer in rasterization context
+		// (TileData allocates a lots of space, but only first TileDataSize bytes hold compressed data)
+
+		uint8* CompressedData = (uint8*)dtAlloc(TileDataSize * sizeof(uint8), DT_ALLOC_PERM);
+		if (CompressedData == nullptr)
+		{
+			dtFree(TileData);
+			BuildContext.log(RC_LOG_ERROR, "RecastBuildTileCache: Out of memory 'CompressedData'.");
+			return false;
+		}
+
+		FMemory::Memcpy(CompressedData, TileData, TileDataSize);
+		RasterContext.Layers.Add(FNavMeshTileData(CompressedData, TileDataSize, i, LayerBBox));
+
+		dtFree(TileData);
+
+		const int32 UncompressedSize = ((sizeof(dtTileCacheLayerHeader) + 3) & ~3) + (3 * header.width * header.height);
+		const float Inv1kB = 1.0f / 1024.0f;
+		BuildContext.log(RC_LOG_PROGRESS, ">> Cache[%d,%d:%d] = %.2fkB (full:%.2fkB rate:%.2f%%)", TileX, TileY, i,
+			TileDataSize * Inv1kB, UncompressedSize * Inv1kB, 1.0f * TileDataSize / UncompressedSize);
+	}
+	CompressedLayers = MoveTemp(RasterContext.Layers);
+	return true;
+}
+
+ETimeSliceWorkResult FRecastTileGenerator::GenerateCompressedLayersTimeSliced(FNavMeshBuildContext& BuildContext)
+{
+	SCOPE_CYCLE_COUNTER(STAT_Navigation_RecastBuildCompressedLayers);
+
+	FTileRasterizationContext* RasterContext = GenCompressedlayersTimeSlicedRasterContext.Get();
+
+	switch (GenCompressedLayersTimeSlicedState)
+	{
+	case EGenerateCompressedLayersTimeSliced::Invalid:
+	{
+		ensureMsgf(false, TEXT("Invalid EGenerateCompressedLayersTimeSliced, has this function been called when its already finished processing?"));
+		return ETimeSliceWorkResult::Failed;
+	}
+	break;
+
+	case EGenerateCompressedLayersTimeSliced::Init:
+	{
+		CompressedLayers.Reset();
+		GenCompressedlayersTimeSlicedRasterContext = MakeUnique<FTileRasterizationContext>();
+		RasterContext = GenCompressedlayersTimeSlicedRasterContext.Get();
+		GenCompressedLayersTimeSlicedState = EGenerateCompressedLayersTimeSliced::CreateHeightField;
+	} // fall through to next state
+	case EGenerateCompressedLayersTimeSliced::CreateHeightField:
+	{
+		if (!CreateHeightField(BuildContext, *RasterContext))
+		{
+			GenCompressedLayersTimeSlicedState = EGenerateCompressedLayersTimeSliced::Invalid;
+			//no need to check time slice as not much work done
+			return ETimeSliceWorkResult::Failed;
+		}
 
-			uint8* CompressedData = (uint8*)dtAlloc(TileDataSize * sizeof(uint8), DT_ALLOC_PERM);
-			if (CompressedData == NULL)
+		GenCompressedLayersTimeSlicedState = EGenerateCompressedLayersTimeSliced::RasterizeTriangles;
+
+		if (TimeSlicer.TestTimeSliceFinished())
+		{
+			return ETimeSliceWorkResult::CallAgainNextTimeSlice;
+		}
+	} // fall through to next state
+	case EGenerateCompressedLayersTimeSliced::RasterizeTriangles:
+	{
+		const ETimeSliceWorkResult WorkResult = RasterizeTrianglesTimeSliced(BuildContext, *RasterContext);
+
+		//original code did not care about success or failure here
+		if (WorkResult != ETimeSliceWorkResult::CallAgainNextTimeSlice)
+		{
+			GenCompressedLayersTimeSlicedState = EGenerateCompressedLayersTimeSliced::EmptyLayers;
+		}
+
+		if (TimeSlicer.IsTimeSliceFinishedCached())
+		{
+			return ETimeSliceWorkResult::CallAgainNextTimeSlice;
+		}
+	} // fall through to next state
+	case EGenerateCompressedLayersTimeSliced::EmptyLayers:
+	{
+		if (!RasterContext->SolidHF || RasterContext->SolidHF->pools == 0)
+		{
+			BuildContext.log(RC_LOG_WARNING, "GenerateCompressedLayersTimeSliced: empty tile - aborting");
+
+			//no need to check time slice as not much work done
+			GenCompressedLayersTimeSlicedState = EGenerateCompressedLayersTimeSliced::Invalid;
+			return ETimeSliceWorkResult::Succeeded;
+		}
+
+		GenCompressedLayersTimeSlicedState = EGenerateCompressedLayersTimeSliced::VoxelFilter;
+		//no need to check time slice as not much work done
+	}// fall through to next state
+	case EGenerateCompressedLayersTimeSliced::VoxelFilter:
+	{
+		GenCompressedLayersTimeSlicedState = EGenerateCompressedLayersTimeSliced::RecastFilter;
+		// Reject voxels outside generation boundaries
+		if (TileConfig.bPerformVoxelFiltering && !bFullyEncapsulatedByInclusionBounds)
+		{
+			ApplyVoxelFilter(RasterContext->SolidHF, TileConfig.walkableRadius);
+
+			if (TimeSlicer.TestTimeSliceFinished())
 			{
-				dtFree(TileData);
-				BuildContext.log(RC_LOG_ERROR, "GenerateCompressedLayers: Out of memory 'CompressedData'.");
-				return false;
+				return ETimeSliceWorkResult::CallAgainNextTimeSlice;
 			}
+		}
+	}// fall through to next state
+	case EGenerateCompressedLayersTimeSliced::RecastFilter:
+	{
+		GenerateRecastFilter(BuildContext, *RasterContext);
 
-			FMemory::Memcpy(CompressedData, TileData, TileDataSize);
-			RasterContext.Layers.Add(FNavMeshTileData(CompressedData, TileDataSize, i, LayerBBox));
+		GenCompressedLayersTimeSlicedState = EGenerateCompressedLayersTimeSliced::CompactHeightField;
 
-			dtFree(TileData);
+		if (TimeSlicer.TestTimeSliceFinished())
+		{
+			return ETimeSliceWorkResult::CallAgainNextTimeSlice;
+		}
+	}// fall through to next state
+	case EGenerateCompressedLayersTimeSliced::CompactHeightField:
+	{
+		if (!BuildCompactHeightField(BuildContext, *RasterContext))
+		{
+			//no need to check time slice as not much work done
+			GenCompressedLayersTimeSlicedState = EGenerateCompressedLayersTimeSliced::Invalid;
+
+			return ETimeSliceWorkResult::Failed;
+		}
+
+		GenCompressedLayersTimeSlicedState = EGenerateCompressedLayersTimeSliced::ErodeWalkable;
+
+		if (TimeSlicer.TestTimeSliceFinished())
+		{
+			return ETimeSliceWorkResult::CallAgainNextTimeSlice;
+		}
+	}// fall through to next state
+	case EGenerateCompressedLayersTimeSliced::ErodeWalkable:
+	{
+		if (!RecastErodeWalkable(BuildContext, *RasterContext))
+		{
+			//no need to check time slice as not much work done
+			GenCompressedLayersTimeSlicedState = EGenerateCompressedLayersTimeSliced::Invalid;
+
+			return ETimeSliceWorkResult::Failed;
+		}
+
+		GenCompressedLayersTimeSlicedState = EGenerateCompressedLayersTimeSliced::BuildLayers;
+
+		if (TimeSlicer.TestTimeSliceFinished())
+		{
+			return ETimeSliceWorkResult::CallAgainNextTimeSlice;
+		}
+	}// fall through to next state
+	case EGenerateCompressedLayersTimeSliced::BuildLayers:
+	{
+		const bool bRecastBuildLayers = RecastBuildLayers(BuildContext, *RasterContext);
+
+		//this could have done a fair amount of work either way so check time slice
+		TimeSlicer.TestTimeSliceFinished();
+
+		if (!bRecastBuildLayers)
+		{
+			GenCompressedLayersTimeSlicedState = EGenerateCompressedLayersTimeSliced::Invalid;
 
-			const int32 UncompressedSize = ((sizeof(dtTileCacheLayerHeader)+3) & ~3) + (3 * header.width * header.height);
-			const float Inv1kB = 1.0f / 1024.0f;
-			BuildContext.log(RC_LOG_PROGRESS, ">> Cache[%d,%d:%d] = %.2fkB (full:%.2fkB rate:%.2f%%)", TileX, TileY, i,
-				TileDataSize * Inv1kB, UncompressedSize * Inv1kB, 1.0f * TileDataSize / UncompressedSize);
+			return ETimeSliceWorkResult::Failed;
+		}
+
+		GenCompressedLayersTimeSlicedState = EGenerateCompressedLayersTimeSliced::BuildTileCache;
+
+		if (TimeSlicer.IsTimeSliceFinishedCached())
+		{
+			return ETimeSliceWorkResult::CallAgainNextTimeSlice;
+		}
+	}// fall through to next state
+	case EGenerateCompressedLayersTimeSliced::BuildTileCache:
+	{
+		GenCompressedLayersTimeSlicedState = EGenerateCompressedLayersTimeSliced::Invalid;
+
+		const bool bRecastBuildTileCache = RecastBuildTileCache(BuildContext, *RasterContext);
+	
+		//this could have done a fair amount of work either way so check time slice
+		TimeSlicer.TestTimeSliceFinished();
+
+		if (!bRecastBuildTileCache)
+		{
+			return ETimeSliceWorkResult::Failed;
 		}
 	}
+	break;
 
-	// Transfer final data
-	CompressedLayers = RasterContext.Layers;
-	return true;
+	default:
+	{
+		ensureMsgf(false, TEXT("unknow EGenerateCompressedLayersTimeSliced state"));
+		return ETimeSliceWorkResult::Failed;
+	}
+	}
+
+	return ETimeSliceWorkResult::Succeeded;
+}
+
+bool FRecastTileGenerator::GenerateCompressedLayers(FNavMeshBuildContext& BuildContext)
+{
+	SCOPE_CYCLE_COUNTER(STAT_Navigation_RecastBuildCompressedLayers);
+
+
+	FTileRasterizationContext RasterContext;
+	CompressedLayers.Reset();
+
+	if (!CreateHeightField(BuildContext, RasterContext))
+	{
+		return false;
+	}
+
+	RasterizeTriangles(BuildContext, RasterContext);
+	if (!RasterContext.SolidHF || RasterContext.SolidHF->pools == 0)
+	{
+		BuildContext.log(RC_LOG_WARNING, "GenerateCompressedLayers: empty tile - aborting");
+		return true;
+	}
+
+	// Reject voxels outside generation boundaries
+	if (TileConfig.bPerformVoxelFiltering && !bFullyEncapsulatedByInclusionBounds)
+	{
+		ApplyVoxelFilter(RasterContext.SolidHF, TileConfig.walkableRadius);
+	}
+
+	GenerateRecastFilter(BuildContext, RasterContext);
+
+	if (!BuildCompactHeightField(BuildContext, RasterContext))
+	{
+		return false;
+	}
+
+	if (!RecastErodeWalkable(BuildContext, RasterContext))
+	{
+		return false;
+	}
+
+	if (!RecastBuildLayers(BuildContext, RasterContext))
+	{
+		return false;
+	}
+
+	return RecastBuildTileCache(BuildContext, RasterContext);
 }
 
 struct FTileGenerationContext
 {
 	FTileGenerationContext(dtTileCacheAlloc* MyAllocator) :
-		Allocator(MyAllocator), Layer(0), DistanceField(0), ContourSet(0), ClusterSet(0), PolyMesh(0), DetailMesh(0)
+		Allocator(MyAllocator), Layer(nullptr), DistanceField(nullptr), ContourSet(nullptr), ClusterSet(nullptr), PolyMesh(nullptr), DetailMesh(nullptr)
+	{
+	}
+
+	FTileGenerationContext() :
+		Allocator(nullptr), Layer(nullptr), DistanceField(nullptr), ContourSet(nullptr), ClusterSet(nullptr), PolyMesh(nullptr), DetailMesh(nullptr)
 	{
 	}
 
@@ -2613,19 +3099,22 @@ struct FTileGenerationContext
 
 	void ResetIntermediateData()
 	{
-		dtFreeTileCacheLayer(Allocator, Layer);
-		Layer = 0;
-		dtFreeTileCacheDistanceField(Allocator, DistanceField);
-		DistanceField = 0;
-		dtFreeTileCacheContourSet(Allocator, ContourSet);
-		ContourSet = 0;
-		dtFreeTileCacheClusterSet(Allocator, ClusterSet);
-		ClusterSet = 0;
-		dtFreeTileCachePolyMesh(Allocator, PolyMesh);
-		PolyMesh = 0;
-		dtFreeTileCachePolyMeshDetail(Allocator, DetailMesh);
-		DetailMesh = 0;
-		// don't clear NavigationData here!
+		if (Allocator)
+		{
+			dtFreeTileCacheLayer(Allocator, Layer);
+			Layer = nullptr;
+			dtFreeTileCacheDistanceField(Allocator, DistanceField);
+			DistanceField = nullptr;
+			dtFreeTileCacheContourSet(Allocator, ContourSet);
+			ContourSet = nullptr;
+			dtFreeTileCacheClusterSet(Allocator, ClusterSet);
+			ClusterSet = nullptr;
+			dtFreeTileCachePolyMesh(Allocator, PolyMesh);
+			PolyMesh = nullptr;
+			dtFreeTileCachePolyMeshDetail(Allocator, DetailMesh);
+			DetailMesh = nullptr;
+			// don't clear NavigationData here!
+		}
 	}
 
 	struct dtTileCacheAlloc* Allocator;
@@ -2638,261 +3127,360 @@ struct FTileGenerationContext
 	TArray<FNavMeshTileData> NavigationData;
 };
 
-bool FRecastTileGenerator::GenerateNavigationData(FNavMeshBuildContext& BuildContext)
+bool FRecastTileGenerator::GenerateNavigationDataLayer(FNavMeshBuildContext& BuildContext, FTileCacheCompressor& TileCompressor, FTileCacheAllocator& GenNavAllocator, FTileGenerationContext& GenerationContext, int32 LayerIdx)
 {
-	SCOPE_CYCLE_COUNTER(STAT_Navigation_RecastBuildNavigation);
+	SCOPE_CYCLE_COUNTER(STAT_Navigation_GenerateNavigationDataLayer)
 		
-	FTileCacheAllocator MyAllocator;
-	FTileCacheCompressor TileCompressor;
-	
-	FTileGenerationContext GenerationContext(&MyAllocator);
-	GenerationContext.NavigationData.Reserve(CompressedLayers.Num());
-
 	dtStatus status = DT_SUCCESS;
 
-	for (int32 iLayer = 0; iLayer < CompressedLayers.Num(); iLayer++)
+	FNavMeshTileData& CompressedData = CompressedLayers[LayerIdx];
+	const dtTileCacheLayerHeader* TileHeader = (const dtTileCacheLayerHeader*)CompressedData.GetData();
+	GenerationContext.ResetIntermediateData();
+
+	// Decompress tile layer data. 
+	status = dtDecompressTileCacheLayer(&GenNavAllocator, &TileCompressor, (unsigned char*)CompressedData.GetData(), CompressedData.DataSize, &GenerationContext.Layer);
+	if (dtStatusFailed(status))
+	{
+		BuildContext.log(RC_LOG_ERROR, "GenerateNavigationDataLayer: failed to decompress layer.");
+		return false;
+	}
+
+	// Rasterize obstacles.
+	MarkDynamicAreas(*GenerationContext.Layer);
+
+	{
+		SCOPE_CYCLE_COUNTER(STAT_Navigation_RecastBuildRegions)
+
+		// Build regions
+		if (TileConfig.TileCachePartitionType == RC_REGION_MONOTONE)
+		{
+			status = dtBuildTileCacheRegionsMonotone(&GenNavAllocator, TileConfig.minRegionArea, TileConfig.mergeRegionArea, *GenerationContext.Layer);
+		}
+		else if (TileConfig.TileCachePartitionType == RC_REGION_WATERSHED)
+		{
+			GenerationContext.DistanceField = dtAllocTileCacheDistanceField(&GenNavAllocator);
+			if (GenerationContext.DistanceField == nullptr)
+			{
+				BuildContext.log(RC_LOG_ERROR, "GenerateNavigationDataLayer: Out of memory 'DistanceField'.");
+				return false;
+			}
+
+			status = dtBuildTileCacheDistanceField(&GenNavAllocator, *GenerationContext.Layer, *GenerationContext.DistanceField);
+			if (dtStatusFailed(status))
+			{
+				BuildContext.log(RC_LOG_ERROR, "GenerateNavigationDataLayer: Failed to build distance field.");
+				return false;
+			}
+
+			status = dtBuildTileCacheRegions(&GenNavAllocator, TileConfig.minRegionArea, TileConfig.mergeRegionArea, *GenerationContext.Layer, *GenerationContext.DistanceField);
+		}
+		else
+		{
+			status = dtBuildTileCacheRegionsChunky(&GenNavAllocator, TileConfig.minRegionArea, TileConfig.mergeRegionArea, *GenerationContext.Layer, TileConfig.TileCacheChunkSize);
+		}
+
+		if (dtStatusFailed(status))
+		{
+			BuildContext.log(RC_LOG_ERROR, "GenerateNavigationDataLayer: Failed to build regions.");
+			return false;
+		}
+
+		// skip empty layer
+		if (GenerationContext.Layer->regCount <= 0)
+		{
+			return true;
+		}
+	}
+
 	{
-		if (DirtyLayers[iLayer] == false || !CompressedLayers[iLayer].IsValid())
+		SCOPE_CYCLE_COUNTER(STAT_Navigation_RecastBuildContours);
+		// Build contour set
+		GenerationContext.ContourSet = dtAllocTileCacheContourSet(&GenNavAllocator);
+		if (GenerationContext.ContourSet == nullptr)
 		{
-			// skip layers not marked for rebuild
-			continue;
+			BuildContext.log(RC_LOG_ERROR, "GenerateNavigationDataLayer: Out of memory 'ContourSet'.");
+			return false;
+		}
+
+		GenerationContext.ClusterSet = dtAllocTileCacheClusterSet(&GenNavAllocator);
+		if (GenerationContext.ClusterSet == nullptr)
+		{
+			BuildContext.log(RC_LOG_ERROR, "GenerateNavigationDataLayer: Out of memory 'ClusterSet'.");
+			return false;
 		}
-				
-		FNavMeshTileData& CompressedData = CompressedLayers[iLayer];
-		const dtTileCacheLayerHeader* TileHeader = (const dtTileCacheLayerHeader*)CompressedData.GetData();
-		GenerationContext.ResetIntermediateData();
 
-		// Decompress tile layer data. 
-		status = dtDecompressTileCacheLayer(&MyAllocator, &TileCompressor, (unsigned char*)CompressedData.GetData(), CompressedData.DataSize, &GenerationContext.Layer);
+		status = dtBuildTileCacheContours(&GenNavAllocator, *GenerationContext.Layer,
+			TileConfig.walkableClimb, TileConfig.maxSimplificationError, TileConfig.cs, TileConfig.ch,
+			*GenerationContext.ContourSet, *GenerationContext.ClusterSet);
 		if (dtStatusFailed(status))
 		{
-			BuildContext.log(RC_LOG_ERROR, "GenerateNavigationData: failed to decompress layer.");
+			BuildContext.log(RC_LOG_ERROR, "GenerateNavigationDataLayer: Failed to generate contour set (0x%08X).", status);
 			return false;
 		}
 
-		// Rasterize obstacles.
-		MarkDynamicAreas(*GenerationContext.Layer);
+		// skip empty layer, sometimes there are regions assigned but all flagged as empty (id=0)
+		if (GenerationContext.ContourSet->nconts <= 0)
+		{
+			return true;
+		}
+	}
 
+	{
+		SCOPE_CYCLE_COUNTER(STAT_Navigation_RecastBuildPolyMesh);
+		// Build poly mesh
+		GenerationContext.PolyMesh = dtAllocTileCachePolyMesh(&GenNavAllocator);
+		if (GenerationContext.PolyMesh == nullptr)
 		{
-			SCOPE_CYCLE_COUNTER(STAT_Navigation_RecastBuildRegions)
-			// Build regions
-			if (TileConfig.TileCachePartitionType == RC_REGION_MONOTONE)
-			{
-				status = dtBuildTileCacheRegionsMonotone(&MyAllocator, TileConfig.minRegionArea, TileConfig.mergeRegionArea, *GenerationContext.Layer);
-			}
-			else if (TileConfig.TileCachePartitionType == RC_REGION_WATERSHED)
-			{
-				GenerationContext.DistanceField = dtAllocTileCacheDistanceField(&MyAllocator);
-				if (GenerationContext.DistanceField == NULL)
-				{
-					BuildContext.log(RC_LOG_ERROR, "GenerateNavigationData: Out of memory 'DistanceField'.");
-					return false;
-				}
+			BuildContext.log(RC_LOG_ERROR, "GenerateNavigationData: Out of memory 'PolyMesh'.");
+			return false;
+		}
 
-				status = dtBuildTileCacheDistanceField(&MyAllocator, *GenerationContext.Layer, *GenerationContext.DistanceField);
-				if (dtStatusFailed(status))
-				{
-					BuildContext.log(RC_LOG_ERROR, "GenerateNavigationData: Failed to build distance field.");
-					return false;
-				}
+		status = dtBuildTileCachePolyMesh(&GenNavAllocator, &BuildContext, *GenerationContext.ContourSet, *GenerationContext.PolyMesh);
+		if (dtStatusFailed(status))
+		{
+			BuildContext.log(RC_LOG_ERROR, "GenerateNavigationData: Failed to generate poly mesh.");
+			return false;
+		}
 
-				status = dtBuildTileCacheRegions(&MyAllocator, TileConfig.minRegionArea, TileConfig.mergeRegionArea, *GenerationContext.Layer, *GenerationContext.DistanceField);
-			}
-			else
-			{
-				status = dtBuildTileCacheRegionsChunky(&MyAllocator, TileConfig.minRegionArea, TileConfig.mergeRegionArea, *GenerationContext.Layer, TileConfig.TileCacheChunkSize);
-			}
+		status = dtBuildTileCacheClusters(&GenNavAllocator, *GenerationContext.ClusterSet, *GenerationContext.PolyMesh);
+		if (dtStatusFailed(status))
+		{
+			BuildContext.log(RC_LOG_ERROR, "GenerateNavigationData: Failed to update cluster set.");
+			return false;
+		}
+	}
 
-			if (dtStatusFailed(status))
-			{
-				BuildContext.log(RC_LOG_ERROR, "GenerateNavigationData: Failed to build regions.");
-				return false;
-			}
+	// Build detail mesh
+	if (TileConfig.bGenerateDetailedMesh)
+	{
+		SCOPE_CYCLE_COUNTER(STAT_Navigation_RecastBuildPolyDetail);
 
-			// skip empty layer
-			if (GenerationContext.Layer->regCount <= 0)
-			{
-				continue;
-			}
-		}
-	
+		// Build detail mesh.
+		GenerationContext.DetailMesh = dtAllocTileCachePolyMeshDetail(&GenNavAllocator);
+		if (GenerationContext.DetailMesh == nullptr)
 		{
-			SCOPE_CYCLE_COUNTER(STAT_Navigation_RecastBuildContours);
-			// Build contour set
-			GenerationContext.ContourSet = dtAllocTileCacheContourSet(&MyAllocator);
-			if (GenerationContext.ContourSet == NULL)
-			{
-				BuildContext.log(RC_LOG_ERROR, "GenerateNavigationData: Out of memory 'ContourSet'.");
-				return false;
-			}
+			BuildContext.log(RC_LOG_ERROR, "GenerateNavigationData: Out of memory 'DetailMesh'.");
+			return false;
+		}
 
-			GenerationContext.ClusterSet = dtAllocTileCacheClusterSet(&MyAllocator);
-			if (GenerationContext.ClusterSet == NULL)
-			{
-				BuildContext.log(RC_LOG_ERROR, "GenerateNavigationData: Out of memory 'ClusterSet'.");
-				return false;
-			}
+		status = dtBuildTileCachePolyMeshDetail(&GenNavAllocator, TileConfig.cs, TileConfig.ch, TileConfig.detailSampleDist, TileConfig.detailSampleMaxError,
+			*GenerationContext.Layer, *GenerationContext.PolyMesh, *GenerationContext.DetailMesh);
+		if (dtStatusFailed(status))
+		{
+			BuildContext.log(RC_LOG_ERROR, "GenerateNavigationData: Failed to generate poly detail mesh.");
+			return false;
+		}
+	}
 
-			status = dtBuildTileCacheContours(&MyAllocator, *GenerationContext.Layer,
-				TileConfig.walkableClimb, TileConfig.maxSimplificationError, TileConfig.cs, TileConfig.ch,
-				*GenerationContext.ContourSet, *GenerationContext.ClusterSet);
-			if (dtStatusFailed(status))
-			{
-				BuildContext.log(RC_LOG_ERROR, "GenerateNavigationData: Failed to generate contour set (0x%08X).", status);
-				return false;
-			}
+	unsigned char* NavData = nullptr;
+	int32 NavDataSize = 0;
 
-			// skip empty layer, sometimes there are regions assigned but all flagged as empty (id=0)
-			if (GenerationContext.ContourSet->nconts <= 0)
-			{
-				continue;
-			}
+	if (TileConfig.maxVertsPerPoly <= DT_VERTS_PER_POLYGON &&
+		GenerationContext.PolyMesh->npolys > 0 && GenerationContext.PolyMesh->nverts > 0)
+	{
+		ensure(GenerationContext.PolyMesh->npolys <= TileConfig.MaxPolysPerTile && "Polys per Tile limit exceeded!");
+		if (GenerationContext.PolyMesh->nverts >= 0xffff)
+		{
+			// The vertex indices are ushorts, and cannot point to more than 0xffff vertices.
+			BuildContext.log(RC_LOG_ERROR, "Too many vertices per tile %d (max: %d).", GenerationContext.PolyMesh->nverts, 0xffff);
+			return false;
 		}
 
+		// if we didn't fail already then it's high time we created data for off-mesh links
+		FOffMeshData OffMeshData;
+		if (OffmeshLinks.Num() > 0)
 		{
-			SCOPE_CYCLE_COUNTER(STAT_Navigation_RecastBuildPolyMesh);
-			// Build poly mesh
-			GenerationContext.PolyMesh = dtAllocTileCachePolyMesh(&MyAllocator);
-			if (GenerationContext.PolyMesh == NULL)
-			{
-				BuildContext.log(RC_LOG_ERROR, "GenerateNavigationData: Out of memory 'PolyMesh'.");
-				return false;
-			}
+			SCOPE_CYCLE_COUNTER(STAT_Navigation_RecastGatherOffMeshData);
 
-			status = dtBuildTileCachePolyMesh(&MyAllocator, &BuildContext, *GenerationContext.ContourSet, *GenerationContext.PolyMesh);
-			if (dtStatusFailed(status))
-			{
-				BuildContext.log(RC_LOG_ERROR, "GenerateNavigationData: Failed to generate poly mesh.");
-				return false;
-			}
+			OffMeshData.Reserve(OffmeshLinks.Num());
+			OffMeshData.AreaClassToIdMap = &AdditionalCachedData.AreaClassToIdMap;
+			OffMeshData.FlagsPerArea = AdditionalCachedData.FlagsPerOffMeshLinkArea;
+			const FSimpleLinkNavModifier* LinkModifier = OffmeshLinks.GetData();
+			const float DefaultSnapHeight = TileConfig.walkableClimb * TileConfig.ch;
 
-			status = dtBuildTileCacheClusters(&MyAllocator, *GenerationContext.ClusterSet, *GenerationContext.PolyMesh);
-			if (dtStatusFailed(status))
+			for (int32 LinkModifierIndex = 0; LinkModifierIndex < OffmeshLinks.Num(); ++LinkModifierIndex, ++LinkModifier)
 			{
-				BuildContext.log(RC_LOG_ERROR, "GenerateNavigationData: Failed to update cluster set.");
-				return false;
+				OffMeshData.AddLinks(LinkModifier->Links, LinkModifier->LocalToWorld, TileConfig.AgentIndex, DefaultSnapHeight);
+#if GENERATE_SEGMENT_LINKS
+				OffMeshData.AddSegmentLinks(LinkModifier->SegmentLinks, LinkModifier->LocalToWorld, TileConfig.AgentIndex, DefaultSnapHeight);
+#endif // GENERATE_SEGMENT_LINKS
 			}
 		}
 
-		// Build detail mesh
+		// fill flags, or else detour won't be able to find polygons
+		// Update poly flags from areas.
+		for (int32 i = 0; i < GenerationContext.PolyMesh->npolys; i++)
+		{
+			GenerationContext.PolyMesh->flags[i] = AdditionalCachedData.FlagsPerArea[GenerationContext.PolyMesh->areas[i]];
+		}
+
+		dtNavMeshCreateParams Params;
+		memset(&Params, 0, sizeof(Params));
+		Params.verts = GenerationContext.PolyMesh->verts;
+		Params.vertCount = GenerationContext.PolyMesh->nverts;
+		Params.polys = GenerationContext.PolyMesh->polys;
+		Params.polyAreas = GenerationContext.PolyMesh->areas;
+		Params.polyFlags = GenerationContext.PolyMesh->flags;
+		Params.polyCount = GenerationContext.PolyMesh->npolys;
+		Params.nvp = GenerationContext.PolyMesh->nvp;
 		if (TileConfig.bGenerateDetailedMesh)
 		{
-			SCOPE_CYCLE_COUNTER(STAT_Navigation_RecastBuildPolyDetail);
+			Params.detailMeshes = GenerationContext.DetailMesh->meshes;
+			Params.detailVerts = GenerationContext.DetailMesh->verts;
+			Params.detailVertsCount = GenerationContext.DetailMesh->nverts;
+			Params.detailTris = GenerationContext.DetailMesh->tris;
+			Params.detailTriCount = GenerationContext.DetailMesh->ntris;
+		}
+		Params.offMeshCons = OffMeshData.LinkParams.GetData();
+		Params.offMeshConCount = OffMeshData.LinkParams.Num();
+		Params.walkableHeight = TileConfig.AgentHeight;
+		Params.walkableRadius = TileConfig.AgentRadius;
+		Params.walkableClimb = TileConfig.AgentMaxClimb;
+		Params.tileX = TileX;
+		Params.tileY = TileY;
+		Params.tileLayer = LayerIdx;
+		rcVcopy(Params.bmin, GenerationContext.Layer->header->bmin);
+		rcVcopy(Params.bmax, GenerationContext.Layer->header->bmax);
+		Params.cs = TileConfig.cs;
+		Params.ch = TileConfig.ch;
+		Params.buildBvTree = TileConfig.bGenerateBVTree;
+#if GENERATE_CLUSTER_LINKS
+		Params.clusterCount = GenerationContext.ClusterSet->nclusters;
+		Params.polyClusters = GenerationContext.ClusterSet->polyMap;
+#endif
 
-			// Build detail mesh.
-			GenerationContext.DetailMesh = dtAllocTileCachePolyMeshDetail(&MyAllocator);
-			if (GenerationContext.DetailMesh == NULL)
-			{
-				BuildContext.log(RC_LOG_ERROR, "GenerateNavigationData: Out of memory 'DetailMesh'.");
-				return false;
-			}
+		{
+			SCOPE_CYCLE_COUNTER(STAT_Navigation_RecastCreateNavMeshData);
 
-			status = dtBuildTileCachePolyMeshDetail(&MyAllocator, TileConfig.cs, TileConfig.ch, TileConfig.detailSampleDist, TileConfig.detailSampleMaxError,
-				*GenerationContext.Layer, *GenerationContext.PolyMesh, *GenerationContext.DetailMesh);
-			if (dtStatusFailed(status))
+			if (!dtCreateNavMeshData(&Params, &NavData, &NavDataSize))
 			{
-				BuildContext.log(RC_LOG_ERROR, "GenerateNavigationData: Failed to generate poly detail mesh.");
+				BuildContext.log(RC_LOG_ERROR, "Could not build Detour navmesh.");
 				return false;
 			}
 		}
+	}
 
-		unsigned char* NavData = 0;
-		int32 NavDataSize = 0;
+	GenerationContext.NavigationData.Add(FNavMeshTileData(NavData, NavDataSize, LayerIdx, CompressedData.LayerBBox));
 
-		if (TileConfig.maxVertsPerPoly <= DT_VERTS_PER_POLYGON &&
-			GenerationContext.PolyMesh->npolys > 0 && GenerationContext.PolyMesh->nverts > 0)
+	const float ModkB = 1.0f / 1024.0f;
+	BuildContext.log(RC_LOG_PROGRESS, ">> Layer[%d] = Verts(%d) Polys(%d) Memory(%.2fkB) Cache(%.2fkB)",
+		LayerIdx, GenerationContext.PolyMesh->nverts, GenerationContext.PolyMesh->npolys,
+		GenerationContext.NavigationData.Last().DataSize * ModkB, CompressedLayers[LayerIdx].DataSize * ModkB);
+
+	return true;
+}
+
+ETimeSliceWorkResult FRecastTileGenerator::GenerateNavigationDataTimeSliced(FNavMeshBuildContext& BuildContext)
+{
+	SCOPE_CYCLE_COUNTER(STAT_Navigation_RecastBuildNavigation);
+
+	FTileCacheCompressor TileCompressor;
+	ETimeSliceWorkResult WorkResult = ETimeSliceWorkResult::Succeeded;
+	dtStatus status = DT_SUCCESS;
+
+	switch (GenerateNavDataTimeSlicedState)
+	{
+	case EGenerateNavDataTimeSlicedState::Invalid:
+	{
+		ensureMsgf(false, TEXT("Invalid EGenerateNavDataTimeSlicedState, has this function been called when its already finished processing?"));
+		return ETimeSliceWorkResult::Failed;
+	}
+	break;
+
+	case EGenerateNavDataTimeSlicedState::Init:
+	{
+		GenNavDataTimeSlicedAllocator = MakeUnique<FTileCacheAllocator>();
+		GenNavDataTimeSlicedGenerationContext = MakeUnique<FTileGenerationContext>(GenNavDataTimeSlicedAllocator.Get());
+		GenNavDataTimeSlicedGenerationContext->NavigationData.Reserve(CompressedLayers.Num());
+		GenerateNavDataTimeSlicedState = EGenerateNavDataTimeSlicedState::GenerateLayers;
+	}//fall through to next state
+	case EGenerateNavDataTimeSlicedState::GenerateLayers:
+	{
+		for (; GenNavDataLayerTimeSlicedIdx < CompressedLayers.Num(); GenNavDataLayerTimeSlicedIdx++)
 		{
-			ensure(GenerationContext.PolyMesh->npolys <= TileConfig.MaxPolysPerTile && "Polys per Tile limit exceeded!");
-			if (GenerationContext.PolyMesh->nverts >= 0xffff)
+			if (DirtyLayers[GenNavDataLayerTimeSlicedIdx] == false || !CompressedLayers[GenNavDataLayerTimeSlicedIdx].IsValid())
 			{
-				// The vertex indices are ushorts, and cannot point to more than 0xffff vertices.
-				BuildContext.log(RC_LOG_ERROR, "Too many vertices per tile %d (max: %d).", GenerationContext.PolyMesh->nverts, 0xffff);
-				return false;
+				// skip layers not marked for rebuild
+				continue;
 			}
 
-			// if we didn't fail already then it's high time we created data for off-mesh links
-			FOffMeshData OffMeshData;
-			if (OffmeshLinks.Num() > 0)
+			if (TimeSlicer.IsTimeSliceFinishedCached())
 			{
-				SCOPE_CYCLE_COUNTER(STAT_Navigation_RecastGatherOffMeshData);
+				WorkResult = ETimeSliceWorkResult::CallAgainNextTimeSlice;
+				break;
+			}
 
-				OffMeshData.Reserve(OffmeshLinks.Num());
-				OffMeshData.AreaClassToIdMap = &AdditionalCachedData.AreaClassToIdMap;
-				OffMeshData.FlagsPerArea = AdditionalCachedData.FlagsPerOffMeshLinkArea;
-				const FSimpleLinkNavModifier* LinkModifier = OffmeshLinks.GetData();
-				const float DefaultSnapHeight = TileConfig.walkableClimb * TileConfig.ch;
+			const bool bGenDataLayer = GenerateNavigationDataLayer(BuildContext, TileCompressor, *GenNavDataTimeSlicedAllocator, *GenNavDataTimeSlicedGenerationContext, GenNavDataLayerTimeSlicedIdx);
 
-				for (int32 LinkModifierIndex = 0; LinkModifierIndex < OffmeshLinks.Num(); ++LinkModifierIndex, ++LinkModifier)
-				{
-					OffMeshData.AddLinks(LinkModifier->Links, LinkModifier->LocalToWorld, TileConfig.AgentIndex, DefaultSnapHeight);
-#if GENERATE_SEGMENT_LINKS
-					OffMeshData.AddSegmentLinks(LinkModifier->SegmentLinks, LinkModifier->LocalToWorld, TileConfig.AgentIndex, DefaultSnapHeight);
-#endif // GENERATE_SEGMENT_LINKS
-				}
-			}
+			//carry on iterating but don't do any more work if the time slice is finished (as we may not need to in which case we can avoid calling this function again)
+			TimeSlicer.TestTimeSliceFinished();
 
-			// fill flags, or else detour won't be able to find polygons
-			// Update poly flags from areas.
-			for (int32 i = 0; i < GenerationContext.PolyMesh->npolys; i++)
+			if (!bGenDataLayer)
 			{
-				GenerationContext.PolyMesh->flags[i] = AdditionalCachedData.FlagsPerArea[GenerationContext.PolyMesh->areas[i]];
+				WorkResult = ETimeSliceWorkResult::Failed;
+				break;
 			}
+		}
+
+		if (WorkResult != ETimeSliceWorkResult::CallAgainNextTimeSlice)
+		{
+			GenNavDataLayerTimeSlicedIdx = 0;
+			GenerateNavDataTimeSlicedState = EGenerateNavDataTimeSlicedState::Invalid;
 
-			dtNavMeshCreateParams Params;
-			memset(&Params, 0, sizeof(Params));
-			Params.verts = GenerationContext.PolyMesh->verts;
-			Params.vertCount = GenerationContext.PolyMesh->nverts;
-			Params.polys = GenerationContext.PolyMesh->polys;
-			Params.polyAreas = GenerationContext.PolyMesh->areas;
-			Params.polyFlags = GenerationContext.PolyMesh->flags;
-			Params.polyCount = GenerationContext.PolyMesh->npolys;
-			Params.nvp = GenerationContext.PolyMesh->nvp;
-			if (TileConfig.bGenerateDetailedMesh)
+			if (WorkResult == ETimeSliceWorkResult::Succeeded)
 			{
-				Params.detailMeshes = GenerationContext.DetailMesh->meshes;
-				Params.detailVerts = GenerationContext.DetailMesh->verts;
-				Params.detailVertsCount = GenerationContext.DetailMesh->nverts;
-				Params.detailTris = GenerationContext.DetailMesh->tris;
-				Params.detailTriCount = GenerationContext.DetailMesh->ntris;
+				NavigationData = MoveTemp(GenNavDataTimeSlicedGenerationContext->NavigationData);
 			}
-			Params.offMeshCons = OffMeshData.LinkParams.GetData();
-			Params.offMeshConCount = OffMeshData.LinkParams.Num();
-			Params.walkableHeight = TileConfig.AgentHeight;
-			Params.walkableRadius = TileConfig.AgentRadius;
-			Params.walkableClimb = TileConfig.AgentMaxClimb;
-			Params.tileX = TileX;
-			Params.tileY = TileY;
-			Params.tileLayer = iLayer;
-			rcVcopy(Params.bmin, GenerationContext.Layer->header->bmin);
-			rcVcopy(Params.bmax, GenerationContext.Layer->header->bmax);
-			Params.cs = TileConfig.cs;
-			Params.ch = TileConfig.ch;
-			Params.buildBvTree = TileConfig.bGenerateBVTree;
-#if GENERATE_CLUSTER_LINKS
-			Params.clusterCount = GenerationContext.ClusterSet->nclusters;
-			Params.polyClusters = GenerationContext.ClusterSet->polyMap;
-#endif
+			GenNavDataTimeSlicedGenerationContext->ResetIntermediateData();
+		}
+	}
+	break;
 
-			SCOPE_CYCLE_COUNTER(STAT_Navigation_RecastCreateNavMeshData);
+	default:
+	{
+		ensureMsgf(false, TEXT("unhandled EGenerateNavDataTimeSlicedState"));
+		return ETimeSliceWorkResult::Failed;
+	}
+	}
 
-			if (!dtCreateNavMeshData(&Params, &NavData, &NavDataSize))
-			{
-				BuildContext.log(RC_LOG_ERROR, "Could not build Detour navmesh.");
-				return false;
-			}
+	return WorkResult;
+}
+
+bool FRecastTileGenerator::GenerateNavigationData(FNavMeshBuildContext& BuildContext)
+{
+	SCOPE_CYCLE_COUNTER(STAT_Navigation_RecastBuildNavigation);
+
+	FTileCacheAllocator GenNavAllocator;
+	FTileGenerationContext GenerationContext(&GenNavAllocator);
+	GenerationContext.NavigationData.Reserve(CompressedLayers.Num());
+	FTileCacheCompressor TileCompressor;
+	bool bGenDataLayer = true;
+	dtStatus status = DT_SUCCESS;
+
+	for (int32 LayerIdx = 0; LayerIdx < CompressedLayers.Num(); LayerIdx++)
+	{
+		if (DirtyLayers[LayerIdx] == false || !CompressedLayers[LayerIdx].IsValid())
+		{
+			// skip layers not marked for rebuild
+			continue;
 		}
 
-		GenerationContext.NavigationData.Add(FNavMeshTileData(NavData, NavDataSize, iLayer, CompressedData.LayerBBox));
+		bGenDataLayer = GenerateNavigationDataLayer(BuildContext, TileCompressor, GenNavAllocator, GenerationContext, LayerIdx);
 
-		const float ModkB = 1.0f / 1024.0f;
-		BuildContext.log(RC_LOG_PROGRESS, ">> Layer[%d] = Verts(%d) Polys(%d) Memory(%.2fkB) Cache(%.2fkB)",
-			iLayer, GenerationContext.PolyMesh->nverts, GenerationContext.PolyMesh->npolys,
-			GenerationContext.NavigationData.Last().DataSize * ModkB, CompressedLayers[iLayer].DataSize * ModkB);
+		if (!bGenDataLayer)
+		{
+			break;
+		}
 	}
 
-	// prepare navigation data of actually rebuild layers for transfer
-	NavigationData = GenerationContext.NavigationData;
-	return true;
+	if (bGenDataLayer)
+	{
+		NavigationData = MoveTemp(GenerationContext.NavigationData);
+	}
+	
+	GenerationContext.ResetIntermediateData();
+
+	return bGenDataLayer;
 }
 
 void FRecastTileGenerator::MarkDynamicAreas(dtTileCacheLayer& Layer)
@@ -3215,10 +3803,19 @@ FRecastNavMeshGenerator::FRecastNavMeshGenerator(ARecastNavMesh& InDestNavMesh)
 	, bRestrictBuildingToActiveTiles(false)
 	, bSortTilesWithSeedLocations(true)
 	, Version(0)
-{
+	, TimeSlicer(0.0025)
 #if TIME_SLICE_NAV_REGEN
-	TimeSliceDuration = 0.0025;
+	, bTimeSliceRegenActive(true)
+	, bNextTimeSliceRegenActive(true)
+#else
+	, bTimeSliceRegenActive(false)
+	, bNextTimeSliceRegenActive(false)
 #endif
+{
+	ProcessTileTasksSyncTimeSlicedState = EProcessTileTasksSyncTimeSlicedState::Init;
+
+	AddGeneratedTilesTimeSlicedState = EAddGeneratedTilesTimeSlicedState::Init;
+	AddGenTilesTimeSlicedLayerIndex = 0;
 
 	INC_DWORD_STAT_BY(STAT_NavigationMemory, sizeof(*this));
 }
@@ -3596,6 +4193,8 @@ void FRecastNavMeshGenerator::CancelBuild()
 
 void FRecastNavMeshGenerator::TickAsyncBuild(float DeltaSeconds)
 {
+	QUICK_SCOPE_CYCLE_COUNTER(STAT_RecastNavMeshGenerator_TickAsyncBuilds);
+
 	bool bRequestDrawingUpdate = false;
 
 #if	WITH_EDITOR
@@ -3742,6 +4341,19 @@ bool FRecastNavMeshGenerator::IsInActiveSet(const FIntPoint& Tile) const
 	return bRestrictBuildingToActiveTiles == false || ActiveTiles.Find(Tile) != INDEX_NONE;
 }
 
+void FRecastNavMeshGenerator::ResetTimeSlicedTileGeneratorSync()
+{
+	TimeSlicedTileGeneratorSync.Reset();
+
+	//reset variables used for timeslicing TileGenratorSync
+	ProcessTileTasksSyncTimeSlicedState = EProcessTileTasksSyncTimeSlicedState::Init;
+	ProcessTileTasksSyncTimeSlicedUpdatedTilesCache.Reset();
+	AddGenTilesTimeSlicedOldLayerTileIdMap.Reset();
+	AddGenTilesTimeSlicedResultTileIndices.Reset();
+	AddGeneratedTilesTimeSlicedState = EAddGeneratedTilesTimeSlicedState::Init;
+	AddGenTilesTimeSlicedLayerIndex = 0;
+}
+
 //@TODO Investigate removing from RunningDirtyTiles here too (or atleast not using the results in any way)
 void FRecastNavMeshGenerator::RemoveTiles(const TArray<FIntPoint>& Tiles)
 {
@@ -3756,15 +4368,13 @@ void FRecastNavMeshGenerator::RemoveTiles(const TArray<FIntPoint>& Tiles)
 			PendingDirtyTiles.Remove(DirtyTile);
 		}
 
-#if TIME_SLICE_NAV_REGEN
-		if (TileGeneratorSync.Get())
+		if (TimeSlicedTileGeneratorSync.Get())
 		{
-			if (TileGeneratorSync->GetTileX() == TileXY.X && TileGeneratorSync->GetTileY() == TileXY.Y)
+			if (TimeSlicedTileGeneratorSync->GetTileX() == TileXY.X && TimeSlicedTileGeneratorSync->GetTileY() == TileXY.Y)
 			{
-				TileGeneratorSync.Reset();
+				ResetTimeSlicedTileGeneratorSync();
 			}
 		}
-#endif
 	}
 }
 
@@ -3868,49 +4478,217 @@ TArray<uint32> FRecastNavMeshGenerator::RemoveTileLayers(const int32 TileX, cons
 	{
 		const int32 NumLayers = DetourMesh->getTileCountAt(TileX, TileY);
 
-		if (NumLayers > 0)
+		if (NumLayers > 0)
+		{
+			TArray<dtMeshTile*> Tiles;
+			Tiles.AddZeroed(NumLayers);
+			DetourMesh->getTilesAt(TileX, TileY, (const dtMeshTile**)Tiles.GetData(), NumLayers);
+
+			for (int32 i = 0; i < NumLayers; i++)
+			{
+				const int32 LayerIndex = Tiles[i]->header->layer;
+				dtPolyRef TileRef = DetourMesh->getTileRef(Tiles[i]);
+
+				NumActiveTiles--;
+				UE_LOG(LogNavigation, Log, TEXT("%s> Tile (%d,%d:%d), removing TileRef: 0x%X (active:%d)"),
+					*DestNavMesh->GetName(), TileX, TileY, LayerIndex, TileRef, NumActiveTiles);
+
+				DetourMesh->removeTile(TileRef, nullptr, nullptr);
+
+				uint32 TileId = RecastTileVersionHelper::GetUpdatedTileId(TileRef, DetourMesh);
+				UpdatedIndices.AddUnique(TileId);
+
+				if (OldLayerTileIdMap)
+				{
+					OldLayerTileIdMap->Add(LayerIndex, TileRef);
+				}
+			}
+		}
+
+		// Remove compressed tile cache layers
+		DestNavMesh->RemoveTileCacheLayers(TileX, TileY);
+	}
+
+	return UpdatedIndices;
+}
+
+void FRecastNavMeshGenerator::AddGeneratedTileLayer(int32 LayerIndex, FRecastTileGenerator& TileGenerator, const TMap<int32, dtPolyRef>& OldLayerTileIdMap, TArray<uint32>& OutResultTileIndices)
+{
+	struct FLayerIndexFinder
+	{
+		int32 LayerIndex;
+		explicit FLayerIndexFinder(const int32 InLayerIndex) : LayerIndex(InLayerIndex) {}
+		bool operator()(const FNavMeshTileData& LayerData) const
+		{
+			return LayerData.LayerIndex == LayerIndex;
+		}
+	};
+
+	const int32 TileX = TileGenerator.GetTileX();
+	const int32 TileY = TileGenerator.GetTileY();
+	dtNavMesh* DetourMesh = DestNavMesh->GetRecastNavMeshImpl()->GetRecastMesh();
+	TArray<FNavMeshTileData>& TileLayers = TileGenerator.GetNavigationData();
+	dtTileRef OldTileRef = DetourMesh->getTileRefAt(TileX, TileY, LayerIndex);
+	const int32 LayerDataIndex = TileLayers.IndexOfByPredicate(FLayerIndexFinder(LayerIndex));
+
+	if (LayerDataIndex != INDEX_NONE)
+	{
+		FNavMeshTileData& LayerData = TileLayers[LayerDataIndex];
+		if (OldTileRef)
+		{
+			NumActiveTiles--;
+			UE_LOG(LogNavigation, Log, TEXT("%s> Tile (%d,%d:%d), removing TileRef: 0x%X (active:%d)"),
+				*DestNavMesh->GetName(), TileX, TileY, LayerIndex, OldTileRef, NumActiveTiles);
+
+			DetourMesh->removeTile(OldTileRef, nullptr, nullptr);
+
+			const uint32 TileId = RecastTileVersionHelper::GetUpdatedTileId(OldTileRef, DetourMesh);
+			OutResultTileIndices.AddUnique(TileId);
+		}
+		else
+		{
+			OldTileRef = OldLayerTileIdMap.FindRef(LayerIndex);
+		}
+
+		if (LayerData.IsValid())
+		{
+			bool bRejectNavmesh = false;
+			dtTileRef ResultTileRef = 0;
+
+			dtStatus status = 0;
+
+			{
+				// let navmesh know it's tile generator who owns the data
+				status = DetourMesh->addTile(LayerData.GetData(), LayerData.DataSize, DT_TILE_FREE_DATA, OldTileRef, &ResultTileRef);
+
+				// if tile index was already taken by other layer try adding it on first free entry (salt was already updated by whatever took that spot)
+				if (dtStatusFailed(status) && dtStatusDetail(status, DT_OUT_OF_MEMORY) && OldTileRef)
+				{
+					OldTileRef = 0;
+					status = DetourMesh->addTile(LayerData.GetData(), LayerData.DataSize, DT_TILE_FREE_DATA, OldTileRef, &ResultTileRef);
+				}
+			}
+
+			if (dtStatusFailed(status))
+			{
+				if (dtStatusDetail(status, DT_OUT_OF_MEMORY))
+				{
+					UE_LOG(LogNavigation, Error, TEXT("%s> Tile (%d,%d:%d), tile limit reached!! (%d)"),
+						*DestNavMesh->GetName(), TileX, TileY, LayerIndex, DetourMesh->getMaxTiles());
+				}
+			}
+			else
+			{
+				OutResultTileIndices.AddUnique(DetourMesh->decodePolyIdTile(ResultTileRef));
+				NumActiveTiles++;
+
+				UE_LOG(LogNavigation, Log, TEXT("%s> Tile (%d,%d:%d), added TileRef: 0x%X (active:%d)"),
+					*DestNavMesh->GetName(), TileX, TileY, LayerIndex, ResultTileRef, NumActiveTiles);
+
+				{
+					// NavMesh took the ownership of generated data, so we don't need to deallocate it
+					uint8* ReleasedData = LayerData.Release();
+				}
+			}
+		}
+	}
+	else
+	{
+		// remove the layer since it ended up empty
+		DetourMesh->removeTile(OldTileRef, nullptr, nullptr);
+		const uint32 TileId = RecastTileVersionHelper::GetUpdatedTileId(OldTileRef, DetourMesh);
+		OutResultTileIndices.AddUnique(TileId);
+	}
+}
+
+ETimeSliceWorkResult FRecastNavMeshGenerator::AddGeneratedTilesTimeSliced(FRecastTileGenerator& TileGenerator, TArray<uint32>& OutResultTileIndices)
+{
+	SCOPE_CYCLE_COUNTER(STAT_Navigation_RecastAddGeneratedTiles);
+
+	const int32 TileX = TileGenerator.GetTileX();
+	const int32 TileY = TileGenerator.GetTileY();
+	dtNavMesh* DetourMesh = DestNavMesh->GetRecastNavMeshImpl()->GetRecastMesh();
+	TArray<FNavMeshTileData>& TileLayers = TileGenerator.GetNavigationData();
+	ETimeSliceWorkResult WorkResult = ETimeSliceWorkResult::Succeeded;
+	bool bIteratedThroughDirtyLayers = true;
+
+	switch (AddGeneratedTilesTimeSlicedState)
+	{
+	case EAddGeneratedTilesTimeSlicedState::Init:
+	{
+		AddGenTilesTimeSlicedResultTileIndices.Reset();
+		AddGenTilesTimeSlicedResultTileIndices.Reserve(TileLayers.Num());
+		AddGenTilesTimeSlicedOldLayerTileIdMap.Reset();
+		AddGenTilesTimeSlicedOldLayerTileIdMap.Reserve(TileLayers.Num());
+		AddGenTilesTimeSlicedLayerIndex = TileGenerator.GetDirtyLayersMask().Find(true);
+		if (TileGenerator.IsFullyRegenerated())
 		{
-			TArray<dtMeshTile*> Tiles;
-			Tiles.AddZeroed(NumLayers);
-			DetourMesh->getTilesAt(TileX, TileY, (const dtMeshTile**)Tiles.GetData(), NumLayers);
+			// remove all layers
+			AddGenTilesTimeSlicedResultTileIndices = RemoveTileLayers(TileX, TileY, &AddGenTilesTimeSlicedOldLayerTileIdMap);
+		}
 
-			for (int32 i = 0; i < NumLayers; i++)
+		AddGeneratedTilesTimeSlicedState = EAddGeneratedTilesTimeSlicedState::AddTiles;
+	}//fall through to next state
+	case EAddGeneratedTilesTimeSlicedState::AddTiles:
+	{
+		if (DetourMesh != nullptr
+			// no longer testing this here, we can live with a stray unwanted tile here 
+			// and there. It will be removed the next time around the invokers get
+			// updated 
+			// && IsInActiveSet(FIntPoint(TileX, TileY))
+			&& AddGenTilesTimeSlicedLayerIndex != INDEX_NONE)
+		{
+			for (; AddGenTilesTimeSlicedLayerIndex < TileGenerator.GetDirtyLayersMask().Num(); ++AddGenTilesTimeSlicedLayerIndex)
 			{
-				const int32 LayerIndex = Tiles[i]->header->layer;
-				dtPolyRef TileRef = DetourMesh->getTileRef(Tiles[i]);
-
-				NumActiveTiles--;
-				UE_LOG(LogNavigation, Log, TEXT("%s> Tile (%d,%d:%d), removing TileRef: 0x%X (active:%d)"),
-					*DestNavMesh->GetName(), TileX, TileY, LayerIndex, TileRef, NumActiveTiles);
-
-				DetourMesh->removeTile(TileRef, nullptr, nullptr);
+				if (TileGenerator.IsLayerChanged(AddGenTilesTimeSlicedLayerIndex))
+				{
+					if (TimeSlicer.IsTimeSliceFinishedCached())
+					{
+						WorkResult = ETimeSliceWorkResult::CallAgainNextTimeSlice;
+						break;
+					}
 
-				uint32 TileId = RecastTileVersionHelper::GetUpdatedTileId(TileRef, DetourMesh);
-				UpdatedIndices.AddUnique(TileId);
+					AddGeneratedTileLayer(AddGenTilesTimeSlicedLayerIndex, TileGenerator, AddGenTilesTimeSlicedOldLayerTileIdMap, AddGenTilesTimeSlicedResultTileIndices);
 
-				if (OldLayerTileIdMap)
-				{
-					OldLayerTileIdMap->Add(LayerIndex, TileRef);
+					TimeSlicer.TestTimeSliceFinished();
 				}
 			}
 		}
+		else
+		{
+			WorkResult = ETimeSliceWorkResult::Failed;
+			bIteratedThroughDirtyLayers = false;
+		}
+	}
+	break;
 
-		// Remove compressed tile cache layers
-		DestNavMesh->RemoveTileCacheLayers(TileX, TileY);
+	default:
+	{
+		ensureMsgf(false, TEXT("unhandled EAddGeneratedTilesTimeSlicedState"));
+		WorkResult = ETimeSliceWorkResult::Failed;
+	}
 	}
 
-	return UpdatedIndices;
+	if (AddGenTilesTimeSlicedLayerIndex == TileGenerator.GetDirtyLayersMask().Num() || !bIteratedThroughDirtyLayers)
+	{
+		AddGenTilesTimeSlicedLayerIndex = 0;
+		AddGeneratedTilesTimeSlicedState = EAddGeneratedTilesTimeSlicedState::Init;
+
+		OutResultTileIndices = MoveTemp(AddGenTilesTimeSlicedResultTileIndices);
+	}
+
+	return WorkResult;
 }
 
 TArray<uint32> FRecastNavMeshGenerator::AddGeneratedTiles(FRecastTileGenerator& TileGenerator)
 {
 	SCOPE_CYCLE_COUNTER(STAT_Navigation_RecastAddGeneratedTiles);
-	
+
 	TMap<int32, dtPolyRef> OldLayerTileIdMap;
 	TArray<uint32> ResultTileIndices;
 	const int32 TileX = TileGenerator.GetTileX();
 	const int32 TileY = TileGenerator.GetTileY();
-	
+
 	if (TileGenerator.IsFullyRegenerated())
 	{
 		// remove all layers
@@ -3920,7 +4698,7 @@ TArray<uint32> FRecastNavMeshGenerator::AddGeneratedTiles(FRecastTileGenerator&
 	dtNavMesh* DetourMesh = DestNavMesh->GetRecastNavMeshImpl()->GetRecastMesh();
 	const int32 FirstDirtyTileIndex = TileGenerator.GetDirtyLayersMask().Find(true);
 
-	if (DetourMesh != nullptr 
+	if (DetourMesh != nullptr
 		// no longer testing this here, we can live with a stray unwanted tile here 
 		// and there. It will be removed the next time around the invokers get
 		// updated 
@@ -3930,95 +4708,11 @@ TArray<uint32> FRecastNavMeshGenerator::AddGeneratedTiles(FRecastTileGenerator&
 		TArray<FNavMeshTileData> TileLayers = TileGenerator.GetNavigationData();
 		ResultTileIndices.Reserve(TileLayers.Num());
 
-		struct FLayerIndexFinder
-		{
-			int32 LayerIndex;
-			explicit FLayerIndexFinder(const int32 InLayerIndex) : LayerIndex(InLayerIndex) {}
-			bool operator()(const FNavMeshTileData& LayerData) const
-			{
-				return LayerData.LayerIndex == LayerIndex;
-			}
-		};
-
 		for (int32 LayerIndex = FirstDirtyTileIndex; LayerIndex < TileGenerator.GetDirtyLayersMask().Num(); ++LayerIndex)
 		{
 			if (TileGenerator.IsLayerChanged(LayerIndex))
 			{
-				dtTileRef OldTileRef = DetourMesh->getTileRefAt(TileX, TileY, LayerIndex);
-				const int32 LayerDataIndex = TileLayers.IndexOfByPredicate(FLayerIndexFinder(LayerIndex));
-
-				if (LayerDataIndex != INDEX_NONE)
-				{
-					FNavMeshTileData& LayerData = TileLayers[LayerDataIndex];
-					if (OldTileRef)
-					{
-						QUICK_SCOPE_CYCLE_COUNTER(STAT_NavMesh_RemoveOldTIle);
-
-						NumActiveTiles--;
-						UE_LOG(LogNavigation, Log, TEXT("%s> Tile (%d,%d:%d), removing TileRef: 0x%X (active:%d)"),
-							*DestNavMesh->GetName(), TileX, TileY, LayerIndex, OldTileRef, NumActiveTiles);
-
-						DetourMesh->removeTile(OldTileRef, nullptr, nullptr);
-
-						const uint32 TileId = RecastTileVersionHelper::GetUpdatedTileId(OldTileRef, DetourMesh);
-						ResultTileIndices.AddUnique(TileId);
-					}
-					else
-					{
-						OldTileRef = OldLayerTileIdMap.FindRef(LayerIndex);
-					}
-
-					if (LayerData.IsValid())
-					{
-						bool bRejectNavmesh = false;
-						dtTileRef ResultTileRef = 0;
-
-						dtStatus status = 0;
-
-						{
-							QUICK_SCOPE_CYCLE_COUNTER(STAT_NavMesh_AddTileToDetourMesh);
-							// let navmesh know it's tile generator who owns the data
-							status = DetourMesh->addTile(LayerData.GetData(), LayerData.DataSize, DT_TILE_FREE_DATA, OldTileRef, &ResultTileRef);
-
-							// if tile index was already taken by other layer try adding it on first free entry (salt was already updated by whatever took that spot)
-							if (dtStatusFailed(status) && dtStatusDetail(status, DT_OUT_OF_MEMORY) && OldTileRef)
-							{
-								OldTileRef = 0;
-								status = DetourMesh->addTile(LayerData.GetData(), LayerData.DataSize, DT_TILE_FREE_DATA, OldTileRef, &ResultTileRef);
-							}
-						}
-
-						if (dtStatusFailed(status))
-						{
-							if (dtStatusDetail(status, DT_OUT_OF_MEMORY))
-							{
-								UE_LOG(LogNavigation, Error, TEXT("%s> Tile (%d,%d:%d), tile limit reached!! (%d)"),
-									*DestNavMesh->GetName(), TileX, TileY, LayerIndex, DetourMesh->getMaxTiles());
-							}
-						}
-						else
-						{
-							ResultTileIndices.AddUnique(DetourMesh->decodePolyIdTile(ResultTileRef));
-							NumActiveTiles++;
-
-							UE_LOG(LogNavigation, Log, TEXT("%s> Tile (%d,%d:%d), added TileRef: 0x%X (active:%d)"),
-								*DestNavMesh->GetName(), TileX, TileY, LayerIndex, ResultTileRef, NumActiveTiles);
-
-							{
-								QUICK_SCOPE_CYCLE_COUNTER(STAT_NavMesh_ReleaseTileLayers);
-								// NavMesh took the ownership of generated data, so we don't need to deallocate it
-								uint8* ReleasedData = LayerData.Release();
-							}
-						}
-					}
-				}
-				else
-				{
-					// remove the layer since it ended up empty
-					DetourMesh->removeTile(OldTileRef, nullptr, nullptr);
-					const uint32 TileId = RecastTileVersionHelper::GetUpdatedTileId(OldTileRef, DetourMesh);
-					ResultTileIndices.AddUnique(TileId);
-				}
+				AddGeneratedTileLayer(LayerIndex, TileGenerator, OldLayerTileIdMap, ResultTileIndices);
 			}
 		}
 	}
@@ -4040,10 +4734,8 @@ void FRecastNavMeshGenerator::DiscardCurrentBuildingTasks()
 		}
 	}
 
-#if TIME_SLICE_NAV_REGEN
-	TileGeneratorSync.Reset();
-#endif
-	
+	ResetTimeSlicedTileGeneratorSync();
+
 	RunningDirtyTiles.Empty();
 }
 
@@ -4051,9 +4743,7 @@ bool FRecastNavMeshGenerator::HasDirtyTiles() const
 {
 	return (PendingDirtyTiles.Num() > 0 
 		|| RunningDirtyTiles.Num() > 0
-#if TIME_SLICE_NAV_REGEN
-		|| TileGeneratorSync.Get() != nullptr
-#endif
+		|| TimeSlicedTileGeneratorSync.Get() != nullptr
 		);
 }
 
@@ -4366,6 +5056,20 @@ void FRecastNavMeshGenerator::RemoveLayers(const FIntPoint& Tile, TArray<uint32>
 	DestNavMesh->MarkEmptyTileCacheLayers(Tile.X, Tile.Y);
 }
 
+void FRecastNavMeshGenerator::StoreCompressedTileCacheLayers(const FRecastTileGenerator& TileGenerator, int32 TileX, int32 TileY)
+{
+	// Store compressed tile cache layers so it can be reused later
+	if (TileGenerator.GetCompressedLayers().Num())
+	{
+		SCOPE_CYCLE_COUNTER(STAT_RecastNavMeshGenerator_StoringCompressedLayers);
+		DestNavMesh->AddTileCacheLayers(TileX, TileY, TileGenerator.GetCompressedLayers());
+	}
+	else
+	{
+		DestNavMesh->MarkEmptyTileCacheLayers(TileX, TileY);
+	}
+}
+
 #if RECAST_ASYNC_REBUILDING
 TArray<uint32> FRecastNavMeshGenerator::ProcessTileTasksAsync(const int32 NumTasksToProcess)
 {
@@ -4431,16 +5135,7 @@ TArray<uint32> FRecastNavMeshGenerator::ProcessTileTasksAsync(const int32 NumTas
 				TArray<uint32> UpdatedTileIndices = AddGeneratedTiles(TileGenerator);
 				UpdatedTiles.Append(UpdatedTileIndices);
 			
-				// Store compressed tile cache layers so it can be reused later
-				if (TileGenerator.GetCompressedLayers().Num())
-				{
-					QUICK_SCOPE_CYCLE_COUNTER(STAT_RecastNavMeshGenerator_StoringCompressedLayers);
-					DestNavMesh->AddTileCacheLayers(Element.Coord.X, Element.Coord.Y, TileGenerator.GetCompressedLayers());
-				}
-				else
-				{
-					DestNavMesh->MarkEmptyTileCacheLayers(Element.Coord.X, Element.Coord.Y);
-				}
+				StoreCompressedTileCacheLayers(TileGenerator, Element.Coord.X, Element.Coord.Y);
 			}
 
 			{
@@ -4459,170 +5154,205 @@ TArray<uint32> FRecastNavMeshGenerator::ProcessTileTasksAsync(const int32 NumTas
 }
 #endif
 
-#if TIME_SLICE_NAV_REGEN
-bool FRecastNavMeshGenerator::IsTimeSliceDurationExceeded(const double StartTime) const
+#if !RECAST_ASYNC_REBUILDING
+TSharedRef<FRecastTileGenerator> FRecastNavMeshGenerator::CreateTileGeneratorFromPendingElement(FIntPoint& OutTileLocation)
 {
-	const double CurTime = FPlatformTime::Seconds();
-	
-	return CurTime - StartTime >= TimeSliceDuration;
+	ensureMsgf(PendingDirtyTiles.Num() > 0, TEXT("Its an assumption of this function that PendingDirtyTiles.Num() > 0"));
+
+	const int32 PendingItemIdx = PendingDirtyTiles.Num() - 1;
+	FPendingTileElement& PendingElement = PendingDirtyTiles[PendingItemIdx];
+
+	OutTileLocation.X = PendingElement.Coord.X;
+	OutTileLocation.Y = PendingElement.Coord.Y;
+
+	TSharedRef<FRecastTileGenerator> TileGenerator = CreateTileGenerator(PendingElement.Coord, PendingElement.DirtyAreas);
+	PendingDirtyTiles.RemoveAt(PendingItemIdx, 1, false);
+
+	return TileGenerator;
 }
-#endif
 
-#if !RECAST_ASYNC_REBUILDING
-TArray<uint32> FRecastNavMeshGenerator::ProcessTileTasksSync(const int32 NumTasksToProcess)
+TArray<uint32> FRecastNavMeshGenerator::ProcessTileTasksSyncTimeSliced(const int32 NumTasksToProcess)
 {
-	QUICK_SCOPE_CYCLE_COUNTER(STAT_RecastNavMeshGenerator_ProcessTileTasksSync);
+	QUICK_SCOPE_CYCLE_COUNTER(STAT_RecastNavMeshGenerator_ProcessTileTasksSyncTimeSliced);
 
-	TArray<uint32> UpdatedTiles;
 	const bool bGameStaticNavMesh = IsGameStaticNavMesh(DestNavMesh);
 
-#if TIME_SLICE_NAV_REGEN
-	const double StartTime = FPlatformTime::Seconds();
-
-	//if we are time slice processing a tile (ie we have already done some processing on this tile last frame)
-	bool bIsTimeSliceProcessingTile = TileGeneratorSync.Get() != nullptr;
-#endif
+	TimeSlicer.ResetStartTime();
 
 	int32 NumProcessedTasks = 0;
 
-	//if bIsTimeSliceProcessingTile is true then this will be the Idx of the next Element to process not the Idx of TileGeneratorSync
-	int32 ElementIdx = PendingDirtyTiles.Num() - 1;
+	TArray<uint32> UpdatedTiles;
 
-	if ((ElementIdx >= 0
+	// Submit pending tile elements
+	while ((PendingDirtyTiles.Num() > 0
 		&& NumProcessedTasks < NumTasksToProcess)
-#if TIME_SLICE_NAV_REGEN
-		|| bIsTimeSliceProcessingTile
-#endif
+		|| TimeSlicedTileGeneratorSync.IsValid()
 		)
 	{
-		// Submit pending tile elements
-		while (true)
-		{
-			QUICK_SCOPE_CYCLE_COUNTER(STAT_RecastNavMeshGenerator_ProcessTileTasks_NewTasks);
-
-			FIntPoint TileLocation;
+		QUICK_SCOPE_CYCLE_COUNTER(STAT_RecastNavMeshGenerator_ProcessTileTasks_NewTasks);
 
-			//this will only be set up when bIsTimeSliceProcessingTile == false,
-			//its currently not needed at other times.
-			TArray<FBox> DirtyAreas;
+		FIntPoint TileLocation;
 
-#if TIME_SLICE_NAV_REGEN
-			if (bIsTimeSliceProcessingTile)
+		if (TimeSlicedTileGeneratorSync.IsValid())
+		{
+			TileLocation.X = TimeSlicedTileGeneratorSync->GetTileX();
+			TileLocation.Y = TimeSlicedTileGeneratorSync->GetTileY();
+		}
+		else
+		{
+			//if the next time slice regen state is false, we want to go to non time sliced tile regen so break here and switch
+			//next frame (as we've finished time slice processing the last tile)
+			if (!bNextTimeSliceRegenActive)
 			{
-				TileLocation.X = TileGeneratorSync->GetTileX();
-				TileLocation.Y = TileGeneratorSync->GetTileY();
+				return UpdatedTiles;
 			}
-			else
-#endif //TIME_SLICE_NAV_REGEN
-			{
-				FPendingTileElement& PendingElement = PendingDirtyTiles[ElementIdx];
-
-				TileLocation.X = PendingElement.Coord.X;
-				TileLocation.Y = PendingElement.Coord.Y;
 
-				DirtyAreas = MoveTemp(PendingElement.DirtyAreas);
+			TimeSlicedTileGeneratorSync = CreateTileGeneratorFromPendingElement(TileLocation);
 
-				// Remove submitted element from pending list, this simplifies time sliced tile processing
-				PendingDirtyTiles.RemoveAt(ElementIdx);
-
-				// Release memory, list could be quite big after map load
-				if (PendingDirtyTiles.Num() == 0)
-				{
-					PendingDirtyTiles.Empty(32);
-				}
-
-				//see definition of bIsTimeSliceProcessingTile
-				ensureMsgf(TileGeneratorSync.Get() == nullptr, TEXT("TileGeneratorSync should be NULL if we aren't currently time slicing this tile"));
-
-				TileGeneratorSync = CreateTileGenerator(TileLocation, DirtyAreas);
-
-#if TIME_SLICE_NAV_REGEN
-				if (IsTimeSliceDurationExceeded(StartTime))
-				{
-					break;
-				}
-#endif //TIME_SLICE_NAV_REGEN
+			if (TimeSlicer.TestTimeSliceFinished())
+			{
+				return UpdatedTiles;
 			}
+		}
 
-			FRecastTileGenerator& TileGeneratorRef = *TileGeneratorSync.Get();
-			ensureMsgf(TileGeneratorRef.GetTileX() == TileLocation.X && TileGeneratorRef.GetTileY() == TileLocation.Y,
-				TEXT("X Loc Mismatch between TileGeneratorRef and PendingElement"));
+		FRecastTileGenerator& TileGeneratorRef = *TimeSlicedTileGeneratorSync;
 
+		if (ProcessTileTasksSyncTimeSlicedState == EProcessTileTasksSyncTimeSlicedState::Init)
+		{
 			if (TileGeneratorRef.HasDataToBuild())
 			{
-#if TIME_SLICE_NAV_REGEN
-				bool bContinueTimeSliceProcessing = true;
-				ETimeSliceWorkResult TimeSliceWorkResult = ETimeSliceWorkResult::CallAgain;
+				ProcessTileTasksSyncTimeSlicedState = EProcessTileTasksSyncTimeSlicedState::DoWork;
+			}
+			else
+			{
+				ProcessTileTasksSyncTimeSlicedState = EProcessTileTasksSyncTimeSlicedState::Finish;
 
-				if (!TileGeneratorRef.HasDoneWork())
+				if (!bGameStaticNavMesh)
 				{
-					while (TimeSliceWorkResult == ETimeSliceWorkResult::CallAgain && bContinueTimeSliceProcessing)
-					{
-						TimeSliceWorkResult = TileGeneratorRef.DoWork();
+					RemoveLayers(TileLocation, UpdatedTiles);
+				}
+			}
+		}
 
-						bContinueTimeSliceProcessing = !IsTimeSliceDurationExceeded(StartTime);
-					}
 
-					if (!bContinueTimeSliceProcessing)
-					{
-						break;
-					}
-				}
-#else //TIME_SLICE_NAV_REGEN
-				TileGeneratorRef.DoWork();
-#endif //TIME_SLICE_NAV_REGEN
+		switch (ProcessTileTasksSyncTimeSlicedState)
+		{
+		case EProcessTileTasksSyncTimeSlicedState::Init:
+		{
+			//do nothing 
+			ensureMsgf(false, TEXT("This State should not be used here!"));
+		}
+		break;
 
-				//@TODO Investigate preventing further processing of tiles when generation has failed
-				const TArray<uint32> UpdatedTileIndices = AddGeneratedTiles(TileGeneratorRef);
-				UpdatedTiles.Append(UpdatedTileIndices);
+		case EProcessTileTasksSyncTimeSlicedState::DoWork:
+		{
+			const ETimeSliceWorkResult WorkResult = TileGeneratorRef.DoWorkTimeSliced();
 
-				// Store compressed tile cache layers so it can be reused later
-				if (TileGeneratorRef.GetCompressedLayers().Num())
-				{
-					QUICK_SCOPE_CYCLE_COUNTER(STAT_RecastNavMeshGenerator_StoringCompressedLayers);
-					DestNavMesh->AddTileCacheLayers(TileLocation.X, TileLocation.Y, TileGeneratorRef.GetCompressedLayers());
-				}
-				else
-				{
-					DestNavMesh->MarkEmptyTileCacheLayers(TileLocation.X, TileLocation.Y);
-				}
-			}
-			else if (!bGameStaticNavMesh)
+			if (WorkResult != ETimeSliceWorkResult::CallAgainNextTimeSlice)
 			{
-				RemoveLayers(TileLocation, UpdatedTiles);
+				ProcessTileTasksSyncTimeSlicedState = EProcessTileTasksSyncTimeSlicedState::AddGeneratedTiles;
 			}
 
-			TileGeneratorSync = nullptr;
-
-			NumProcessedTasks++;
-
-#if TIME_SLICE_NAV_REGEN
-			//only subtract from the ElementIdx if we aren't time slice Processing this tile (ie the tile hasn't already had some
-			//processing done on it last frame)
-			if (!bIsTimeSliceProcessingTile)
+			if (TimeSlicer.IsTimeSliceFinishedCached())
 			{
-				--ElementIdx;
+				return UpdatedTiles;
 			}
-			bIsTimeSliceProcessingTile = false;
-#else //TIME_SLICE_NAV_REGEN
-			--ElementIdx;
-#endif //TIME_SLICE_NAV_REGEN
+		}//fall through to next state
+		case EProcessTileTasksSyncTimeSlicedState::AddGeneratedTiles:
+		{
+			const ETimeSliceWorkResult WorkResult = AddGeneratedTilesTimeSliced(TileGeneratorRef, ProcessTileTasksSyncTimeSlicedUpdatedTilesCache);
 
-			if (NumProcessedTasks >= NumTasksToProcess || ElementIdx < 0)
+			if (WorkResult != ETimeSliceWorkResult::CallAgainNextTimeSlice)
 			{
-				break;
+				ProcessTileTasksSyncTimeSlicedState = EProcessTileTasksSyncTimeSlicedState::StoreCompessedTileCacheLayers;
 			}
-#if TIME_SLICE_NAV_REGEN
-			else
+
+			if (TimeSlicer.IsTimeSliceFinishedCached())
 			{
-				//only perform the time slice check if we are going to iterate again
-				if (IsTimeSliceDurationExceeded(StartTime))
-				{
-					break;
-				}
+				return UpdatedTiles;
 			}
-#endif //TIME_SLICE_NAV_REGEN
+		}//fall through to next state
+		case EProcessTileTasksSyncTimeSlicedState::StoreCompessedTileCacheLayers:
+		{
+			StoreCompressedTileCacheLayers(TileGeneratorRef, TileLocation.X, TileLocation.Y);
+
+			//no need to check time slicing as not much work done
+			ProcessTileTasksSyncTimeSlicedState = EProcessTileTasksSyncTimeSlicedState::AppendUpdateTiles;
+		}//fall through to next state
+		case EProcessTileTasksSyncTimeSlicedState::AppendUpdateTiles: //this state was added purely to separate the functionality and allow the code to be more easily changed in future.
+		{
+			UpdatedTiles.Append(ProcessTileTasksSyncTimeSlicedUpdatedTilesCache);
+			ProcessTileTasksSyncTimeSlicedUpdatedTilesCache.Empty();
+
+			//no need to check time slicing as not much work done
+			ProcessTileTasksSyncTimeSlicedState = EProcessTileTasksSyncTimeSlicedState::Finish;
+		}//fall through to next state
+		case EProcessTileTasksSyncTimeSlicedState::Finish:
+		{
+			//no need to check time slicing as not much work done
+			//reset state to Init for next tile to be processed
+			ProcessTileTasksSyncTimeSlicedState = EProcessTileTasksSyncTimeSlicedState::Init;
+			TimeSlicedTileGeneratorSync.Reset();
+
+			NumProcessedTasks++;
 		}
+		break;
+		default:
+		{
+			ensureMsgf(false, TEXT("unhandled EProcessTileTasksSyncTimeSlicedState"));
+		}
+		}
+	}
+
+	// Release memory, list could be quite big after map load
+	if (PendingDirtyTiles.Num() == 0)
+	{
+		PendingDirtyTiles.Empty(32);
+	}
+
+	return UpdatedTiles;
+}
+
+//this code path is approx 10% faster than ProcessTileTasksSyncTimeSliced, however it spikes far worse for most use cases.
+TArray<uint32> FRecastNavMeshGenerator::ProcessTileTasksSync(const int32 NumTasksToProcess)
+{
+	QUICK_SCOPE_CYCLE_COUNTER(STAT_RecastNavMeshGenerator_ProcessTileTasksSync);
+
+	const bool bGameStaticNavMesh = IsGameStaticNavMesh(DestNavMesh);
+	int32 NumProcessedTasks = 0;
+	TArray<uint32> UpdatedTiles;
+	FIntPoint TileLocation;
+
+	// Submit pending tile elements
+	while ((PendingDirtyTiles.Num() > 0 && NumProcessedTasks < NumTasksToProcess))
+	{
+		QUICK_SCOPE_CYCLE_COUNTER(STAT_RecastNavMeshGenerator_ProcessTileTasks_NewTasks);
+
+		TSharedRef<FRecastTileGenerator> TileGenerator = CreateTileGeneratorFromPendingElement(TileLocation);
+		
+		FRecastTileGenerator& TileGeneratorRef = *TileGenerator;
+
+		//Does this remain true whenever we stop time slicing?
+		if (TileGeneratorRef.HasDataToBuild())
+		{
+			TileGeneratorRef.DoWork();
+
+			UpdatedTiles = AddGeneratedTiles(TileGeneratorRef);
+
+			StoreCompressedTileCacheLayers(TileGeneratorRef, TileLocation.X, TileLocation.Y);
+		}
+		else if (!bGameStaticNavMesh)
+		{
+			RemoveLayers(TileLocation, UpdatedTiles);
+		}
+
+		NumProcessedTasks++;
+	}
+
+	// Release memory, list could be quite big after map load
+	if (PendingDirtyTiles.Num() == 0)
+	{
+		PendingDirtyTiles.Empty(32);
 	}
 
 	return UpdatedTiles;
@@ -4632,17 +5362,40 @@ TArray<uint32> FRecastNavMeshGenerator::ProcessTileTasksSync(const int32 NumTask
 TArray<uint32> FRecastNavMeshGenerator::ProcessTileTasks(const int32 NumTasksToProcess)
 {
 	const bool bHasTasksAtStart = GetNumRemaningBuildTasks() > 0;
+	TArray<uint32> UpdatedTiles;
 
 #if RECAST_ASYNC_REBUILDING
-	TArray<uint32> UpdatedTiles = ProcessTileTasksAsync(NumTasksToProcess);
+	UpdatedTiles = ProcessTileTasksAsync(NumTasksToProcess);
 #else
-	TArray<uint32> UpdatedTiles = ProcessTileTasksSync(NumTasksToProcess);
+	//only switch bTimeSliceRegen state if we are not time slicing or if we are but aren't part way through time slicing a tile
+	if (bTimeSliceRegenActive != bNextTimeSliceRegenActive)
+	{
+		if (!bTimeSliceRegenActive)
+		{
+			bTimeSliceRegenActive = bNextTimeSliceRegenActive;
+		}
+		else if (!TimeSlicedTileGeneratorSync.IsValid())//test if we have finished processing a tile
+		{
+			bTimeSliceRegenActive = bNextTimeSliceRegenActive;
+		}
+	}
+
+	if (bTimeSliceRegenActive)
+	{
+		UpdatedTiles = ProcessTileTasksSyncTimeSliced(NumTasksToProcess);
+	}
+	else
+	{
+		UpdatedTiles = ProcessTileTasksSync(NumTasksToProcess);
+	}
 #endif
 
 	// Notify owner in case all tasks has been completed
 	const bool bHasTasksAtEnd = GetNumRemaningBuildTasks() > 0;
 	if (bHasTasksAtStart && !bHasTasksAtEnd)
 	{
+		QUICK_SCOPE_CYCLE_COUNTER(STAT_RecastNavMeshGenerator_TickAsyncBuilds);
+
 		DestNavMesh->OnNavMeshGenerationFinished();
 	}
 
@@ -4777,29 +5530,20 @@ bool FRecastNavMeshGenerator::IsBuildInProgress(bool bCheckDirtyToo) const
 {
 	return RunningDirtyTiles.Num()
 		|| (bCheckDirtyToo && PendingDirtyTiles.Num())
-#if TIME_SLICE_NAV_REGEN
-		|| TileGeneratorSync.Get()
-#endif 
-		;
+		|| TimeSlicedTileGeneratorSync.Get();
 }
 
 int32 FRecastNavMeshGenerator::GetNumRemaningBuildTasks() const
 {
 	return RunningDirtyTiles.Num() 
 		+ PendingDirtyTiles.Num()
-#if TIME_SLICE_NAV_REGEN
-		+ (TileGeneratorSync.Get() ? 1 : 0)
-#endif
-		;
+		+ (TimeSlicedTileGeneratorSync.Get() ? 1 : 0);
 }
 
 int32 FRecastNavMeshGenerator::GetNumRunningBuildTasks() const
 {
 	return RunningDirtyTiles.Num()
-#if TIME_SLICE_NAV_REGEN
-		+ (TileGeneratorSync.Get() ? 1 : 0)
-#endif
-		;
+		+ (TimeSlicedTileGeneratorSync.Get() ? 1 : 0);
 }
 
 bool FRecastNavMeshGenerator::GatherGeometryOnGameThread() const 
@@ -4833,9 +5577,10 @@ uint32 FRecastNavMeshGenerator::LogMemUsed() const
 	for (const FRunningTileElement& Element : RunningDirtyTiles)
 	{
 		GeneratorsMem += Element.AsyncTask->GetTask().TileGenerator->UsedMemoryOnStartup;
-#if TIME_SLICE_NAV_REGEN
-		GeneratorsMem += TileGeneratorSync->UsedMemoryOnStartup;
-#endif
+		if (TimeSlicedTileGeneratorSync.IsValid())
+		{
+			GeneratorsMem += TimeSlicedTileGeneratorSync->UsedMemoryOnStartup;
+		}
 	}
 
 	UE_LOG(LogNavigation, Display, TEXT("    FRecastNavMeshGenerator: Total Generator\'s size %u, count %d"), GeneratorsMem, RunningDirtyTiles.Num());
diff --git a/Engine/Source/Runtime/NavigationSystem/Private/NavigationSystem.cpp b/Engine/Source/Runtime/NavigationSystem/Private/NavigationSystem.cpp
--- a/Engine/Source/Runtime/NavigationSystem/Private/NavigationSystem.cpp
+++ b/Engine/Source/Runtime/NavigationSystem/Private/NavigationSystem.cpp
@@ -85,11 +85,13 @@ DEFINE_STAT(STAT_Navigation_RecastAddGeneratedTiles);
 DEFINE_STAT(STAT_Navigation_RecastTick);
 DEFINE_STAT(STAT_Navigation_RecastPathfinding);
 DEFINE_STAT(STAT_Navigation_RecastTestPath);
+DEFINE_STAT(STAT_RecastNavMeshGenerator_StoringCompressedLayers);
 DEFINE_STAT(STAT_Navigation_RecastBuildCompressedLayers);
 DEFINE_STAT(STAT_Navigation_RecastCreateHeightField);
 DEFINE_STAT(STAT_Navigation_RecastRasterizeTriangles);
 DEFINE_STAT(STAT_Navigation_RecastVoxelFilter);
 DEFINE_STAT(STAT_Navigation_RecastFilter);
+DEFINE_STAT(STAT_Navigation_FilterLedgeSpans);
 DEFINE_STAT(STAT_Navigation_RecastBuildCompactHeightField);
 DEFINE_STAT(STAT_Navigation_RecastErodeWalkable);
 DEFINE_STAT(STAT_Navigation_RecastBuildLayers);
@@ -101,6 +103,7 @@ DEFINE_STAT(STAT_Navigation_RecastCreateNavMeshData);
 DEFINE_STAT(STAT_Navigation_RecastMarkAreas);
 DEFINE_STAT(STAT_Navigation_RecastBuildContours);
 DEFINE_STAT(STAT_Navigation_RecastBuildNavigation);
+DEFINE_STAT(STAT_Navigation_GenerateNavigationDataLayer);
 DEFINE_STAT(STAT_Navigation_RecastBuildRegions);
 DEFINE_STAT(STAT_Navigation_UpdateNavOctree);
 DEFINE_STAT(STAT_Navigation_CollisionTreeMemory);
diff --git a/Engine/Source/Runtime/NavigationSystem/Public/NavMesh/RecastNavMeshGenerator.h b/Engine/Source/Runtime/NavigationSystem/Public/NavMesh/RecastNavMeshGenerator.h
--- a/Engine/Source/Runtime/NavigationSystem/Public/NavMesh/RecastNavMeshGenerator.h
+++ b/Engine/Source/Runtime/NavigationSystem/Public/NavMesh/RecastNavMeshGenerator.h
@@ -24,10 +24,15 @@ class ARecastNavMesh;
 class FNavigationOctree;
 class FNavMeshBuildContext;
 class FRecastNavMeshGenerator;
+struct FTileRasterizationContext;
 struct BuildContext;
 struct FNavigationRelevantData;
 struct dtTileCacheLayer;
 struct FKAggregateGeom;
+struct FTileCacheCompressor;
+struct FTileCacheAllocator;
+struct FTileGenerationContext;
+class dtNavMesh;
 
 #define MAX_VERTS_PER_POLY	6
 
@@ -178,11 +183,95 @@ struct FRcTileBox
 	}
 };
 
+/**
+ * The general idea is that any function that handles time slicing internally will be named XXXXTimeSliced
+ * and returns a ETimeSliceWorkResult. These functions also call TestTimeSliceFinished() internally when required,
+ * IsTimeSliceFinishedCached() can be called externally after they have finished. Non time sliced functions are 
+ * managed externally and the calling function should call TestTimeSliceFinished() when necessary.
+ */
+struct NAVIGATIONSYSTEM_API FTimeSlicer
+{
+public:
+	FTimeSlicer()
+		: TimeSliceDuration(0.)
+		, TimeSliceStartTime(0.)
+		, bTimeSliceFinishedCached(false)
+	{
+	}
+
+	FTimeSlicer(double SliceDuration)
+		: TimeSliceDuration(SliceDuration)
+		, TimeSliceStartTime(0.)
+		, bTimeSliceFinishedCached(false)
+	{
+	}
+
+	void SetTimeSliceDuration(double SliceDuration);
+	void ResetStartTime();
+	bool TestTimeSliceFinished() const;
+	bool IsTimeSliceFinishedCached() const;
+
+private:
+	double TimeSliceDuration;
+	double TimeSliceStartTime;
+	mutable bool bTimeSliceFinishedCached;
+};
+
+/** Return state of calling time sliced functions */
 enum class ETimeSliceWorkResult : uint8
 {
 	Failed,
 	Succeeded,
-	CallAgain, //this state will only be used if TIME_SLICE_NAV_REGEN != 0
+	CallAgainNextTimeSlice, /** time slice is finished this frame but we need to call this functionality again next frame */
+};
+
+/** State representing which area of GenerateCompressedLayersTimeSliced() we are processing */
+enum class EGenerateCompressedLayersTimeSliced : uint8
+{
+	Invalid,
+	Init,
+	CreateHeightField,
+	RasterizeTriangles,
+	EmptyLayers,
+	VoxelFilter,
+	RecastFilter,
+	CompactHeightField,
+	ErodeWalkable,
+	BuildLayers,
+	BuildTileCache,
+};
+
+enum class ERasterizeGeomRecastTimeSlicedState : uint8 
+{
+	MarkWalkableTriangles,
+	RasterizeTriangles,
+};
+
+enum class ERasterizeGeomTimeSlicedState : uint8
+{
+	RasterizeGeometryTransformCoords,
+	RasterizeGeometryRecast,
+};
+
+enum class EDoWorkTimeSlicedState : uint8
+{
+	Invalid,
+	DoAsyncGeometryGathering,
+	GenerateTile,
+};
+
+enum class EGenerateTileTimeSlicedState : uint8
+{
+	Invalid,
+	GenerateCompressedLayers,
+	GenerateNavigationData,
+};
+
+enum class EGenerateNavDataTimeSlicedState : uint8
+{
+	Invalid,
+	Init,
+	GenerateLayers,
 };
 
 /**
@@ -196,8 +285,12 @@ public:
 	FRecastTileGenerator(FRecastNavMeshGenerator& ParentGenerator, const FIntPoint& Location);
 	virtual ~FRecastTileGenerator();
 		
-	/** the value will be Suceeded or Failed unless TIME_SLICE_NAV_REGEN != 0*/
-	ETimeSliceWorkResult DoWork();
+	/** Does the work involved with regenerating this tile using time slicing.
+	 *  The return value determines the result of the time slicing
+	 */
+	ETimeSliceWorkResult DoWorkTimeSliced();
+	/** Does the work involved with regenerating this tile */
+	bool DoWork();
 
 	FORCEINLINE int32 GetTileX() const { return TileX; }
 	FORCEINLINE int32 GetTileY() const { return TileY; }
@@ -209,11 +302,8 @@ public:
 	/** Whether tile task has anything to build */
 	bool HasDataToBuild() const;
 
-#if TIME_SLICE_NAV_REGEN
-	bool HasDoneWork() const { return bDoneWork; }
-#endif
-
 	const TArray<FNavMeshTileData>& GetCompressedLayers() const { return CompressedLayers; }
+
 protected:
 	// to be used solely by FRecastNavMeshGenerator
 	TArray<FNavMeshTileData>& GetNavigationData() { return NavigationData; }
@@ -230,22 +320,56 @@ public:
 	// FGCObject end
 		
 protected:
-	/** Does the actual tile generations. 
-	 *	@note always trigger tile generation only via TriggerAsyncBuild. This is a worker function
-	 *	@return Suceeded if new tile navigation data has been generated and is ready to be added to navmesh instance, 
-	 *	Failed if failed or no need to generate (still valid). CallAgain will only be used if TIME_SLICE_NAV_REGEN != 0
+	/** Does the actual TimeSliced tile generation. 
+	 *	@note always trigger tile generation only via DoWorkTimeSliced(). This is a worker function
+	 *  The return value determines the result of the time slicing
+	 *	@return Suceeded if new tile navigation data has been generated and is ready to be added to navmesh instance,
+	 *	@return Failed if failed or no need to generate (still valid).
+	 *  @return CallAgainNextTimeSlice, time slice is finished this frame but we need to call this function again next frame
+	 */
+	ETimeSliceWorkResult GenerateTileTimeSliced();
+	/** Does the actual tile generation.
+	 *	@note always trigger tile generation only via DoWorkTime(). This is a worker function
+	 *  The return value determines the result of the time slicing
+	 *	@return true if new tile navigation data has been generated and is ready to be added to navmesh instance,
+	 *	@return false if failed or no need to generate (still valid).
 	 */
-	ETimeSliceWorkResult GenerateTile();
+	bool GenerateTile();
 
 	void Setup(const FRecastNavMeshGenerator& ParentGenerator, const TArray<FBox>& DirtyAreas);
 	
 	void GatherGeometry(const FRecastNavMeshGenerator& ParentGenerator, bool bGeometryChanged);
 	void PrepareGeometrySources(const FRecastNavMeshGenerator& ParentGenerator, bool bGeometryChanged);
-	void DoAsyncGeometryGathering();
-
+	/** Returns true if there is any NavigationRelevantData */
+	bool DoAsyncGeometryGathering();
+
+	/** Start functions used by GenerateCompressedLayersTimeSliced / GenerateCompressedLayers */
+	bool CreateHeightField(FNavMeshBuildContext& BuildContext, FTileRasterizationContext& RasterContext);
+	ETimeSliceWorkResult RasterizeTrianglesTimeSliced(FNavMeshBuildContext& BuildContext, FTileRasterizationContext& RasterContext);
+	void RasterizeTriangles(FNavMeshBuildContext& BuildContext, FTileRasterizationContext& RasterContext);
+	ETimeSliceWorkResult RasterizeGeometryRecastTimeSliced(FNavMeshBuildContext& BuildContext, const TArray<float>& Coords, const TArray<int32>& Indices, FTileRasterizationContext& RasterContext);
+	void RasterizeGeometryRecast(FNavMeshBuildContext& BuildContext, const TArray<float>& Coords, const TArray<int32>& Indices, FTileRasterizationContext& RasterContext);
+	void RasterizeGeometryTransformCoords(const TArray<float>& Coords, const FTransform& LocalToWorld);
+	ETimeSliceWorkResult RasterizeGeometryTimeSliced(FNavMeshBuildContext& BuildContext, const TArray<float>& Coords, const TArray<int32>& Indices, const FTransform& LocalToWorld, FTileRasterizationContext& RasterContext);
+	void RasterizeGeometry(FNavMeshBuildContext& BuildContext, const TArray<float>& Coords, const TArray<int32>& Indices, const FTransform& LocalToWorld, FTileRasterizationContext& RasterContext);
+	void GenerateRecastFilter(FNavMeshBuildContext& BuildContext, FTileRasterizationContext& RasterContext);
+	bool BuildCompactHeightField(FNavMeshBuildContext& BuildContext, FTileRasterizationContext& RasterContext);
+	bool RecastErodeWalkable(FNavMeshBuildContext& BuildContext, FTileRasterizationContext& RasterContext);
+	bool RecastBuildLayers(FNavMeshBuildContext& BuildContext, FTileRasterizationContext& RasterContext);
+	bool RecastBuildTileCache(FNavMeshBuildContext& BuildContext, FTileRasterizationContext& RasterContext);
+	/** End functions used by GenerateCompressedLayersTimeSliced / GenerateCompressedLayers */
+
+	/** builds CompressedLayers array (geometry + modifiers) time sliced*/
+	virtual ETimeSliceWorkResult GenerateCompressedLayersTimeSliced(FNavMeshBuildContext& BuildContext);
 	/** builds CompressedLayers array (geometry + modifiers) */
 	virtual bool GenerateCompressedLayers(FNavMeshBuildContext& BuildContext);
 
+	/** Builds a navigation data layer */
+	bool GenerateNavigationDataLayer(FNavMeshBuildContext& BuildContext, FTileCacheCompressor& TileCompressor, FTileCacheAllocator& GenNavAllocator, FTileGenerationContext& GenerationContext, int32 LayerIdx);
+
+	/** builds NavigationData array (layers + obstacles) time sliced */
+	ETimeSliceWorkResult GenerateNavigationDataTimeSliced(FNavMeshBuildContext& BuildContext);
+
 	/** builds NavigationData array (layers + obstacles) */
 	bool GenerateNavigationData(FNavMeshBuildContext& BuildContext);
 
@@ -270,18 +394,30 @@ protected:
 	void DumpAsyncData();
 
 protected:
-	uint32 bSucceeded : 1;
 	uint32 bRegenerateCompressedLayers : 1;
 	uint32 bFullyEncapsulatedByInclusionBounds : 1;
 	uint32 bUpdateGeometry : 1;
 	uint32 bHasLowAreaModifiers : 1;
 
-#if TIME_SLICE_NAV_REGEN
-	uint32 bDoneAsyncDataGathering : 1;
-	uint32 bDoneRegenerateCompressedLayers : 1;
-	uint32 bDoneWork : 1;
-#endif
-	
+	/** Start time slicing variables */
+	ERasterizeGeomRecastTimeSlicedState RasterizeGeomRecastState;
+	ERasterizeGeomTimeSlicedState RasterizeGeomState;
+	EDoWorkTimeSlicedState DoWorkTimeSlicedState;
+	EGenerateTileTimeSlicedState GenerateTileTimeSlicedState;
+
+	EGenerateNavDataTimeSlicedState GenerateNavDataTimeSlicedState;
+	int32 GenNavDataLayerTimeSlicedIdx;
+	EGenerateCompressedLayersTimeSliced GenCompressedLayersTimeSlicedState;
+	int32 RasterizeTrianglesTimeSlicedRawGeomIdx;
+	int32 RasterizeTrianglesTimeSlicedInstTransformIdx;
+	TNavStatArray<uint8> RasterizeGeomRecastTriAreas;
+	const FTimeSlicer& TimeSlicer;
+
+	TUniquePtr<struct FTileCacheAllocator> GenNavDataTimeSlicedAllocator;
+	TUniquePtr<struct FTileGenerationContext> GenNavDataTimeSlicedGenerationContext;
+	TUniquePtr<struct FTileRasterizationContext> GenCompressedlayersTimeSlicedRasterContext;
+	/** End time slicing variables */
+
 	int32 TileX;
 	int32 TileY;
 	uint32 Version;
@@ -306,6 +442,9 @@ protected:
 	// generated tile data
 	TArray<FNavMeshTileData> CompressedLayers;
 	TArray<FNavMeshTileData> NavigationData;
+
+	/** Result of calling RasterizeGeometryInitVars() */
+	TArray<float> RasterizeGeometryWorldRecastCoords;
 	
 	// tile's geometry: without voxel cache
 	TArray<FRecastRawGeometryElement> RawGeometry;
@@ -424,6 +563,22 @@ struct FTileTimestamp
 	}
 };
 
+enum class EProcessTileTasksSyncTimeSlicedState : uint8
+{
+	Init,
+	DoWork,
+	AddGeneratedTiles,
+	StoreCompessedTileCacheLayers,
+	AppendUpdateTiles,
+	Finish,
+};
+
+enum class EAddGeneratedTilesTimeSlicedState : uint8
+{
+	Init,
+	AddTiles,
+};
+
 /**
  * Class that handles generation of the whole Recast-based navmesh.
  */
@@ -435,7 +590,8 @@ public:
 
 private:
 	/** Prevent copying. */
-	FRecastNavMeshGenerator(FRecastNavMeshGenerator const& NoCopy) { check(0); };
+	FRecastNavMeshGenerator(FRecastNavMeshGenerator const& NoCopy) 
+		: bTimeSliceRegenActive(false) { check(0); };
 	FRecastNavMeshGenerator& operator=(FRecastNavMeshGenerator const& NoCopy) { check(0); return *this; }
 
 public:
@@ -507,6 +663,10 @@ public:
 	static void GetDebugGeometry(const FNavigationRelevantData& EncodedData, FNavDebugMeshData& DebugMeshData);
 #endif  // !UE_BUILD_SHIPPING
 
+	FTimeSlicer& GetTimeSlicer() { return TimeSlicer; }
+	
+	void SetNextTimeSliceRegenActive(bool bRegenState) { bNextTimeSliceRegenActive = bRegenState; }
+
 protected:
 	// Performs initial setup of member variables so that generator is ready to
 	// do its thing from this point on. Called just after construction by ARecastNavMesh
@@ -538,24 +698,27 @@ protected:
 
 	void RemoveLayers(const FIntPoint& Tile, TArray<uint32>& UpdatedTiles);
 	
+	void StoreCompressedTileCacheLayers(const FRecastTileGenerator& TileGenerator, int32 TileX, int32 TileY);
 #if RECAST_ASYNC_REBUILDING
 	/** Processes pending tile generation tasks Async*/
 	TArray<uint32> ProcessTileTasksAsync(const int32 NumTasksToProcess);
 #else
-	/** Processes pending tile generation tasks Sync with option for time slicing*/
+	TSharedRef<FRecastTileGenerator> CreateTileGeneratorFromPendingElement(FIntPoint &OutTileLocation);
+	/** Processes pending tile generation tasks Sync with option for time slicing currently an experimental feature. */
+	TArray<uint32> ProcessTileTasksSyncTimeSliced(const int32 NumTasksToProcess);
 	TArray<uint32> ProcessTileTasksSync(const int32 NumTasksToProcess);
 #endif
 	/** Processes pending tile generation tasks */
 	TArray<uint32> ProcessTileTasks(const int32 NumTasksToProcess);
 
+	void ResetTimeSlicedTileGeneratorSync();
+
 public:
+	/** Adds generated tiles to NavMesh, replacing old ones, uses time slicing returns Failed if any layer failed */
+	ETimeSliceWorkResult AddGeneratedTilesTimeSliced(FRecastTileGenerator& TileGenerator, TArray<uint32>& OutResultTileIndices);
 	/** Adds generated tiles to NavMesh, replacing old ones */
 	TArray<uint32> AddGeneratedTiles(FRecastTileGenerator& TileGenerator);
 
-#if TIME_SLICE_NAV_REGEN
-	bool IsTimeSliceDurationExceeded(const double StartTime) const;
-#endif
-
 public:
 	/** Removes all tiles at specified grid location */
 	TArray<uint32> RemoveTileLayers(const int32 TileX, const int32 TileY, TMap<int32, dtPolyRef>* OldLayerTileIdMap = nullptr);
@@ -588,6 +751,8 @@ protected:
 	//----------------------------------------------------------------------//
 	virtual uint32 LogMemUsed() const override;
 
+	void AddGeneratedTileLayer(int32 LayerIndex, FRecastTileGenerator& TileGenerator, const TMap<int32, dtPolyRef>& OldLayerTileIdMap, TArray<uint32>& OutResultTileIndices);
+
 protected:
 	friend ARecastNavMesh;
 
@@ -611,11 +776,6 @@ protected:
 
 	/** Navigation mesh that owns this generator */
 	ARecastNavMesh*	DestNavMesh;
-
-#if !RECAST_ASYNC_REBUILDING
-	/** This is cached as a member variable as its used if TIME_SLICE_NAV_REGEN != 0 */
-	TSharedPtr<FRecastTileGenerator> TileGeneratorSync;
-#endif
 	
 	/** List of dirty tiles that needs to be regenerated */
 	TNavStatArray<FPendingTileElement> PendingDirtyTiles;			
@@ -646,9 +806,24 @@ protected:
 	 *	like when navmesh size changes */
 	uint32 Version;
 
-#if TIME_SLICE_NAV_REGEN
-	double TimeSliceDuration;
-#endif
+	/** Start time sliced variables */
+	/** Do not null or Reset this directly instead call ResetTimeSlicedTileGeneratorSync(). The only exception currently is in ProcessTileTasksSyncTimeSliced */
+	TSharedPtr<FRecastTileGenerator> TimeSlicedTileGeneratorSync;
+	FTimeSlicer TimeSlicer;
+	TArray<uint32> ProcessTileTasksSyncTimeSlicedUpdatedTilesCache;
+	/** if we are currently using time sliced regen or not - currently an experimental feature.
+	 *  do not manipulate this value directly instead call SetNextTimeSliceRegenState() 
+	 */
+	bool bTimeSliceRegenActive;
+	bool bNextTimeSliceRegenActive;
+
+	EProcessTileTasksSyncTimeSlicedState ProcessTileTasksSyncTimeSlicedState;
+
+	TMap<int32, dtPolyRef> AddGenTilesTimeSlicedOldLayerTileIdMap;
+	TArray<uint32> AddGenTilesTimeSlicedResultTileIndices;
+	EAddGeneratedTilesTimeSlicedState AddGeneratedTilesTimeSlicedState;
+	int32 AddGenTilesTimeSlicedLayerIndex;
+	/** End time sliced variables */
 };
 
 #endif // WITH_RECAST
