UNavigationSystemConfig: resolve NavigationSystemClass directly instead of calling TryLoad that may generate warnings during loading. We use native classes for NavigationSystem and it should be resolvable at this point.
ANavSystemConfigOverride: made some changes to remove the requirements of deferred calls from the TimerManager for non game worlds so editor world will be ready for navigation data generation right after loading

#rnx

diff --git a/Engine/Source/Runtime/Engine/Private/AI/NavigationSystemConfig.cpp b/Engine/Source/Runtime/Engine/Private/AI/NavigationSystemConfig.cpp
--- a/Engine/Source/Runtime/Engine/Private/AI/NavigationSystemConfig.cpp
+++ b/Engine/Source/Runtime/Engine/Private/AI/NavigationSystemConfig.cpp
@@ -24,8 +24,8 @@ UNavigationSystemBase* UNavigationSystemConfig::CreateAndConfigureNavigationSyst
 {
 	UNavigationSystemBase* NavSys = nullptr;
 
-	NavigationSystemClass.TryLoad();
 	TSubclassOf<UNavigationSystemBase> NavSysClass = NavigationSystemClass.ResolveClass();
+	ensure(NavSysClass != nullptr || NavigationSystemClass.IsValid() == false);
 	
 	if (NavSysClass)
 	{
diff --git a/Engine/Source/Runtime/NavigationSystem/Private/NavSystemConfigOverride.cpp b/Engine/Source/Runtime/NavigationSystem/Private/NavSystemConfigOverride.cpp
--- a/Engine/Source/Runtime/NavigationSystem/Private/NavSystemConfigOverride.cpp
+++ b/Engine/Source/Runtime/NavigationSystem/Private/NavSystemConfigOverride.cpp
@@ -65,10 +65,11 @@ void ANavSystemConfigOverride::PostLoad()
 	Super::PostLoad();
 	
 #if WITH_EDITOR
+	// ApplyConfig in PostLoad only for Editor worlds; applied in BeginPlay for Game worlds.
 	UWorld* World = GetWorld();
 	if ((World != nullptr) && (World->IsGameWorld() == false))
 	{
-		GEditor->GetTimerManager()->SetTimerForNextTick(this, &ANavSystemConfigOverride::ApplyConfig);
+		ApplyConfig();
 	}
 #endif // WITH_EDITOR
 }
