Added a few perfromance stats to the navigation mesh regen code.

diff --git a/Engine/Source/Runtime/Engine/Public/EngineStats.h b/Engine/Source/Runtime/Engine/Public/EngineStats.h
--- a/Engine/Source/Runtime/Engine/Public/EngineStats.h
+++ b/Engine/Source/Runtime/Engine/Public/EngineStats.h
@@ -119,6 +119,7 @@ DECLARE_CYCLE_STAT_EXTERN(TEXT("Unregistering Nav Octree Element"), STAT_Navigat
 DECLARE_CYCLE_STAT_EXTERN(TEXT("Adding actors to navoctree"),STAT_Navigation_AddingActorsToNavOctree,STATGROUP_Navigation, );
 DECLARE_CYCLE_STAT_EXTERN(TEXT("Adjusting nav links"),STAT_Navigation_AdjustingNavLinks,STATGROUP_Navigation, );
 DECLARE_CYCLE_STAT_EXTERN(TEXT("Recast: sync add generated tiles"), STAT_Navigation_RecastAddGeneratedTiles,STATGROUP_Navigation, );
+DECLARE_CYCLE_STAT_EXTERN(TEXT("Recast: sync add generated tile Layer"), STAT_Navigation_RecastAddGeneratedTileLayer, STATGROUP_Navigation, );
 DECLARE_CYCLE_STAT_EXTERN(TEXT("Recast: tick"),STAT_Navigation_RecastTick,STATGROUP_Navigation, );
 DECLARE_CYCLE_STAT_EXTERN(TEXT("Recast: pathfinding"),STAT_Navigation_RecastPathfinding, STATGROUP_Navigation, );
 DECLARE_CYCLE_STAT_EXTERN(TEXT("Recast: test path"),STAT_Navigation_RecastTestPath, STATGROUP_Navigation, );
diff --git a/Engine/Source/Runtime/NavigationSystem/Private/NavMesh/RecastNavMeshGenerator.cpp b/Engine/Source/Runtime/NavigationSystem/Private/NavMesh/RecastNavMeshGenerator.cpp
--- a/Engine/Source/Runtime/NavigationSystem/Private/NavMesh/RecastNavMeshGenerator.cpp
+++ b/Engine/Source/Runtime/NavigationSystem/Private/NavMesh/RecastNavMeshGenerator.cpp
@@ -4193,7 +4193,7 @@ void FRecastNavMeshGenerator::CancelBuild()
 
 void FRecastNavMeshGenerator::TickAsyncBuild(float DeltaSeconds)
 {
-	QUICK_SCOPE_CYCLE_COUNTER(STAT_RecastNavMeshGenerator_TickAsyncBuilds);
+	QUICK_SCOPE_CYCLE_COUNTER(STAT_RecastNavMeshGenerator_TickAsyncBuild);
 
 	bool bRequestDrawingUpdate = false;
 
@@ -4223,8 +4223,13 @@ void FRecastNavMeshGenerator::TickAsyncBuild(float DeltaSeconds)
 			
 	if (UpdatedTileIndices.Num() > 0)
 	{
-		// Invalidate active paths that go through regenerated tiles
-		DestNavMesh->OnNavMeshTilesUpdated(UpdatedTileIndices);
+		{
+			QUICK_SCOPE_CYCLE_COUNTER(STAT_RecastNavMeshGenerator_OnNavMeshTilesUpdated);
+
+			// Invalidate active paths that go through regenerated tiles
+			DestNavMesh->OnNavMeshTilesUpdated(UpdatedTileIndices);
+		}
+
 		bRequestDrawingUpdate = true;
 
 #if	WITH_EDITOR
@@ -4514,6 +4519,8 @@ TArray<uint32> FRecastNavMeshGenerator::RemoveTileLayers(const int32 TileX, cons
 
 void FRecastNavMeshGenerator::AddGeneratedTileLayer(int32 LayerIndex, FRecastTileGenerator& TileGenerator, const TMap<int32, dtPolyRef>& OldLayerTileIdMap, TArray<uint32>& OutResultTileIndices)
 {
+	SCOPE_CYCLE_COUNTER(STAT_Navigation_RecastAddGeneratedTileLayer);
+
 	struct FLayerIndexFinder
 	{
 		int32 LayerIndex;
@@ -5361,6 +5368,8 @@ TArray<uint32> FRecastNavMeshGenerator::ProcessTileTasksSync(const int32 NumTask
 
 TArray<uint32> FRecastNavMeshGenerator::ProcessTileTasks(const int32 NumTasksToProcess)
 {
+	QUICK_SCOPE_CYCLE_COUNTER(STAT_RecastNavMeshGenerator_ProcessTileTasks);
+
 	const bool bHasTasksAtStart = GetNumRemaningBuildTasks() > 0;
 	TArray<uint32> UpdatedTiles;
 
@@ -5394,7 +5403,7 @@ TArray<uint32> FRecastNavMeshGenerator::ProcessTileTasks(const int32 NumTasksToP
 	const bool bHasTasksAtEnd = GetNumRemaningBuildTasks() > 0;
 	if (bHasTasksAtStart && !bHasTasksAtEnd)
 	{
-		QUICK_SCOPE_CYCLE_COUNTER(STAT_RecastNavMeshGenerator_TickAsyncBuilds);
+		QUICK_SCOPE_CYCLE_COUNTER(STAT_RecastNavMeshGenerator_OnNavMeshGenerationFinished);
 
 		DestNavMesh->OnNavMeshGenerationFinished();
 	}
@@ -5403,6 +5412,8 @@ TArray<uint32> FRecastNavMeshGenerator::ProcessTileTasks(const int32 NumTasksToP
 	//only do this if framepro is recording as its an expensive operation
 	if (FFrameProProfiler::IsFrameProRecording())
 	{
+		QUICK_SCOPE_CYCLE_COUNTER(STAT_RecastNavMeshGenerator_GetCompressedTileCacheSize);
+
 		int32 TileCacheSize = DestNavMesh->GetCompressedTileCacheSize();
 
 		FPlatformMisc::CustomNamedStat("TotalTileCacheSize", static_cast<float>(TileCacheSize), "NavMesh", "Bytes");
diff --git a/Engine/Source/Runtime/NavigationSystem/Private/NavigationSystem.cpp b/Engine/Source/Runtime/NavigationSystem/Private/NavigationSystem.cpp
--- a/Engine/Source/Runtime/NavigationSystem/Private/NavigationSystem.cpp
+++ b/Engine/Source/Runtime/NavigationSystem/Private/NavigationSystem.cpp
@@ -82,6 +82,7 @@ DEFINE_STAT(STAT_Navigation_RegisterNavOctreeElement);
 DEFINE_STAT(STAT_Navigation_UnregisterNavOctreeElement);
 DEFINE_STAT(STAT_Navigation_AddingActorsToNavOctree);
 DEFINE_STAT(STAT_Navigation_RecastAddGeneratedTiles);
+DEFINE_STAT(STAT_Navigation_RecastAddGeneratedTileLayer);
 DEFINE_STAT(STAT_Navigation_RecastTick);
 DEFINE_STAT(STAT_Navigation_RecastPathfinding);
 DEFINE_STAT(STAT_Navigation_RecastTestPath);
