Preventing navmesh generation under specified geometry
Addition property bRejectNavmeshUnderneath on PrimitiveComponent
The flag gets set in the FCompositeNavModifier and is read when rasterizing the geometry
When set, the triangles are rasterize as solid columns going to the min level of the solid heightfield preventing navmesh generation under them

diff --git a/Engine/Source/Runtime/Engine/Classes/Components/PrimitiveComponent.h b/Engine/Source/Runtime/Engine/Classes/Components/PrimitiveComponent.h
--- a/Engine/Source/Runtime/Engine/Classes/Components/PrimitiveComponent.h
+++ b/Engine/Source/Runtime/Engine/Classes/Components/PrimitiveComponent.h
@@ -487,6 +487,12 @@ public:
 	UPROPERTY(EditAnywhere, BlueprintReadWrite, Category = Physics)
 	uint8 bReplicatePhysicsToAutonomousProxy : 1;
 
+	// Navigation
+
+	/** If set, navmesh will not be generated under the surface of the geometry */
+	UPROPERTY(EditAnywhere, Category = Navigation)
+	uint8 bRejectNavmeshUnderneath:1;
+
 	// General flags.
 	
 	/** If this is True, this component must always be loaded on clients, even if Hidden and CollisionEnabled is NoCollision. */
@@ -2327,6 +2333,7 @@ public:
 	virtual bool CanCharacterStepUp(class APawn* Pawn) const;
 
 	//~ Begin INavRelevantInterface Interface
+	virtual void GetNavigationData(FNavigationRelevantData& OutData) const override;
 	virtual FBox GetNavigationBounds() const override;
 	virtual bool IsNavigationRelevant() const override;
 	//~ End INavRelevantInterface Interface
diff --git a/Engine/Source/Runtime/Engine/Private/AI/Navigation/NavigationModifier.cpp b/Engine/Source/Runtime/Engine/Private/AI/Navigation/NavigationModifier.cpp
--- a/Engine/Source/Runtime/Engine/Private/AI/Navigation/NavigationModifier.cpp
+++ b/Engine/Source/Runtime/Engine/Private/AI/Navigation/NavigationModifier.cpp
@@ -731,6 +731,7 @@ void FCompositeNavModifier::Reset()
 	bHasPotentialLinks = false;
 	bAdjustHeight = false;
 	bIsPerInstanceModifier = false;
+	bModifierRejectNavmeshUnderneath = false;
 }
 
 void FCompositeNavModifier::Empty()
@@ -740,6 +741,7 @@ void FCompositeNavModifier::Empty()
 	CustomLinks.Empty();
 	bHasPotentialLinks = false;
 	bAdjustHeight = false;
+	bModifierRejectNavmeshUnderneath = false;
 }
 
 FCompositeNavModifier FCompositeNavModifier::GetInstantiatedMetaModifier(const FNavAgentProperties* NavAgent, TWeakObjectPtr<UObject> WeakOwnerPtr) const
diff --git a/Engine/Source/Runtime/Engine/Private/Components/PrimitiveComponent.cpp b/Engine/Source/Runtime/Engine/Private/Components/PrimitiveComponent.cpp
--- a/Engine/Source/Runtime/Engine/Private/Components/PrimitiveComponent.cpp
+++ b/Engine/Source/Runtime/Engine/Private/Components/PrimitiveComponent.cpp
@@ -324,6 +324,7 @@ UPrimitiveComponent::UPrimitiveComponent(const FObjectInitializer& ObjectInitial
 	bCastVolumetricTranslucentShadow = false;
 	IndirectLightingCacheQuality = ILCQ_Point;
 	bSelectable = true;
+	bRejectNavmeshUnderneath = false;
 	AlwaysLoadOnClient = true;
 	AlwaysLoadOnServer = true;
 	SetCollisionProfileName(UCollisionProfile::BlockAll_ProfileName);
@@ -2376,6 +2377,15 @@ void UPrimitiveComponent::DispatchWakeEvents(ESleepEvent WakeEvent, FName BoneNa
 	}
 }
 
+void UPrimitiveComponent::GetNavigationData(FNavigationRelevantData& OutData) const
+{
+	if (bRejectNavmeshUnderneath)
+	{
+		FCompositeNavModifier CompositeNavModifier;
+		CompositeNavModifier.SetRejectNavmeshUnderneath(bRejectNavmeshUnderneath);
+		OutData.Modifiers.Add(CompositeNavModifier);
+	}
+}
 
 bool UPrimitiveComponent::IsNavigationRelevant() const 
 { 
diff --git a/Engine/Source/Runtime/Engine/Private/Components/ShapeComponent.cpp b/Engine/Source/Runtime/Engine/Private/Components/ShapeComponent.cpp
--- a/Engine/Source/Runtime/Engine/Private/Components/ShapeComponent.cpp
+++ b/Engine/Source/Runtime/Engine/Private/Components/ShapeComponent.cpp
@@ -80,6 +80,8 @@ bool UShapeComponent::DoCustomNavigableGeometryExport(FNavigableGeometryExport&
 
 void UShapeComponent::GetNavigationData(FNavigationRelevantData& Data) const
 {
+	Super::GetNavigationData(Data);
+
 	if (bDynamicObstacle)
 	{
 		Data.Modifiers.CreateAreaModifiers(this, AreaClass);
diff --git a/Engine/Source/Runtime/Engine/Public/AI/NavigationModifier.h b/Engine/Source/Runtime/Engine/Public/AI/NavigationModifier.h
--- a/Engine/Source/Runtime/Engine/Public/AI/NavigationModifier.h
+++ b/Engine/Source/Runtime/Engine/Public/AI/NavigationModifier.h
@@ -249,7 +249,7 @@ protected:
 
 struct ENGINE_API FCompositeNavModifier : public FNavigationModifier
 {
-	FCompositeNavModifier() : bHasPotentialLinks(false), bAdjustHeight(false), bHasLowAreaModifiers(false), bIsPerInstanceModifier(false) {}
+	FCompositeNavModifier() : bHasPotentialLinks(false), bAdjustHeight(false), bHasLowAreaModifiers(false), bIsPerInstanceModifier(false), bModifierRejectNavmeshUnderneath(false) {}
 
 	void Shrink();
 	void Reset();
@@ -288,6 +288,7 @@ struct ENGINE_API FCompositeNavModifier : public FNavigationModifier
 		bHasMetaAreas |= Modifiers.bHasMetaAreas; 
 		bAdjustHeight |= Modifiers.HasAgentHeightAdjust();
 		bHasLowAreaModifiers |= Modifiers.HasLowAreaModifiers();
+		bModifierRejectNavmeshUnderneath |= Modifiers.GetRejectNavmeshUnderneath();
 	}
 
 	void CreateAreaModifiers(const UPrimitiveComponent* PrimComp, const TSubclassOf<UNavAreaBase> AreaClass);
@@ -301,7 +302,9 @@ struct ENGINE_API FCompositeNavModifier : public FNavigationModifier
 	FORCEINLINE bool HasAgentHeightAdjust() const { return bAdjustHeight; }
 	FORCEINLINE bool HasAreas() const { return Areas.Num() > 0; }
 	FORCEINLINE bool HasLowAreaModifiers() const { return bHasLowAreaModifiers; }
-    FORCEINLINE bool IsPerInstanceModifier() const { return bIsPerInstanceModifier; }
+	FORCEINLINE bool IsPerInstanceModifier() const { return bIsPerInstanceModifier; }
+	FORCEINLINE bool GetRejectNavmeshUnderneath() const { return bModifierRejectNavmeshUnderneath; }
+	FORCEINLINE void SetRejectNavmeshUnderneath(bool bInRejectNavmeshUnderneath) { bModifierRejectNavmeshUnderneath = bInRejectNavmeshUnderneath; }
 
 	FORCEINLINE void ReserveForAdditionalAreas(int32 AdditionalElementsCount) { Areas.Reserve(Areas.Num() + AdditionalElementsCount); }
 
@@ -333,4 +336,5 @@ private:
 	uint32 bAdjustHeight : 1;
 	uint32 bHasLowAreaModifiers : 1;
     uint32 bIsPerInstanceModifier : 1;
+	uint32 bModifierRejectNavmeshUnderneath : 1;
 };
diff --git a/Engine/Source/Runtime/NavigationSystem/Private/NavMesh/RecastNavMeshGenerator.cpp b/Engine/Source/Runtime/NavigationSystem/Private/NavMesh/RecastNavMeshGenerator.cpp
--- a/Engine/Source/Runtime/NavigationSystem/Private/NavMesh/RecastNavMeshGenerator.cpp
+++ b/Engine/Source/Runtime/NavigationSystem/Private/NavMesh/RecastNavMeshGenerator.cpp
@@ -1814,6 +1814,8 @@ bool FRecastTileGenerator::DoAsyncGeometryGathering()
 			NavOctree->DemandLazyDataGathering(*ElementData);
 		}
 
+		const FCompositeNavModifier ModifierInstance = ElementData->Modifiers.HasMetaAreas() ? ElementData->Modifiers.GetInstantiatedMetaModifier(&NavDataConfig, ElementData->SourceObject) : ElementData->Modifiers;
+
 		const bool bExportGeometry = bUpdateGeometry && ElementData->HasGeometry();
 		if (bExportGeometry)
 		{
@@ -1828,7 +1830,7 @@ bool FRecastTileGenerator::DoAsyncGeometryGathering()
 				if (!HasVoxelCache(ElementData->VoxelData, CachedVoxels, NumCachedVoxels))
 				{
 					// rasterize
-					PrepareVoxelCache(ElementData->CollisionData, SpanData);
+					PrepareVoxelCache(ElementData->CollisionData, ModifierInstance, SpanData);
 					CachedVoxels = SpanData.GetData();
 					NumCachedVoxels = SpanData.Num();
 
@@ -1844,7 +1846,7 @@ bool FRecastTileGenerator::DoAsyncGeometryGathering()
 			}
 			else
 			{
-				ValidateAndAppendGeometry(ElementData);
+				ValidateAndAppendGeometry(ElementData, ModifierInstance);
 			}
 
 			if (bDumpGeometryData)
@@ -1853,7 +1855,6 @@ bool FRecastTileGenerator::DoAsyncGeometryGathering()
 			}
 		}
 
-		const FCompositeNavModifier ModifierInstance = ElementData->Modifiers.HasMetaAreas() ? ElementData->Modifiers.GetInstantiatedMetaModifier(&NavDataConfig, ElementData->SourceObject) : ElementData->Modifiers;
 		if (ModifierInstance.IsEmpty() == false)
 		{
 			AppendModifier(ModifierInstance, ElementData->NavDataPerInstanceTransformDelegate);
@@ -1946,6 +1947,8 @@ void FRecastTileGenerator::GatherGeometry(const FRecastNavMeshGenerator& ParentG
 				}
 			}
 
+			const FCompositeNavModifier ModifierInstance = Element.GetModifierForAgent(&OwnerNavDataConfig);
+
 			const bool bExportGeometry = bGeometryChanged && Element.Data->HasGeometry();
 			if (bExportGeometry)
 			{
@@ -1960,7 +1963,7 @@ void FRecastTileGenerator::GatherGeometry(const FRecastNavMeshGenerator& ParentG
 					if (!HasVoxelCache(Element.Data->VoxelData, CachedVoxels, NumCachedVoxels))
 					{
 						// rasterize
-						PrepareVoxelCache(Element.Data->CollisionData, SpanData);
+						PrepareVoxelCache(Element.Data->CollisionData, ModifierInstance, SpanData);
 						CachedVoxels = SpanData.GetData();
 						NumCachedVoxels = SpanData.Num();
 
@@ -1976,7 +1979,7 @@ void FRecastTileGenerator::GatherGeometry(const FRecastNavMeshGenerator& ParentG
 				}
 				else
 				{
-					ValidateAndAppendGeometry(Element.Data);
+					ValidateAndAppendGeometry(Element.Data, ModifierInstance);
 				}
 
 				if (bDumpGeometryData)
@@ -1985,7 +1988,6 @@ void FRecastTileGenerator::GatherGeometry(const FRecastNavMeshGenerator& ParentG
 				}
 			}
 						
-			const FCompositeNavModifier ModifierInstance = Element.GetModifierForAgent(&OwnerNavDataConfig);
 			if (ModifierInstance.IsEmpty() == false)
 			{
 				AppendModifier(ModifierInstance, Element.Data->NavDataPerInstanceTransformDelegate);
@@ -2103,7 +2105,7 @@ void FRecastTileGenerator::ApplyVoxelFilter(rcHeightfield* HF, float WalkableRad
 	}
 }
 
-void FRecastTileGenerator::PrepareVoxelCache(const TNavStatArray<uint8>& RawCollisionCache, TNavStatArray<rcSpanCache>& SpanData)
+void FRecastTileGenerator::PrepareVoxelCache(const TNavStatArray<uint8>& RawCollisionCache, const FCompositeNavModifier& InModifier, TNavStatArray<rcSpanCache>& SpanData)
 {
 	// tile's geometry: voxel cache (only for synchronous rebuilds)
 	const int32 WalkableClimbVX = TileConfig.walkableClimb;
@@ -2126,9 +2128,13 @@ void FRecastTileGenerator::PrepareVoxelCache(const TNavStatArray<uint8>& RawColl
 		CachedCollisions.Indices, CachedCollisions.Header.NumFaces,
 		TriAreas.GetData());
 
+	// To prevent navmesh generation under the triangles, set the RC_PROJECT_TO_BOTTOM flag to true.
+	// This rasterize triangles as filled columns down to the HF lower bound.
+	const rcRasterizationFlags Flags = InModifier.GetRejectNavmeshUnderneath() ? RC_PROJECT_TO_BOTTOM : rcRasterizationFlags(0);
+
 	rcRasterizeTriangles(0, CachedCollisions.Verts, CachedCollisions.Header.NumVerts,
 		CachedCollisions.Indices, TriAreas.GetData(), CachedCollisions.Header.NumFaces,
-		*VoxelCacheContext.RasterizeHF, WalkableClimbVX);
+		*VoxelCacheContext.RasterizeHF, WalkableClimbVX, Flags);
 
 	const int32 NumSpans = rcCountSpans(0, *VoxelCacheContext.RasterizeHF);
 	if (NumSpans > 0)
@@ -2215,16 +2221,16 @@ void FRecastTileGenerator::AppendModifier(const FCompositeNavModifier& Modifier,
 	Modifiers.Add(MoveTemp(ModifierElement));
 }
 
-void FRecastTileGenerator::ValidateAndAppendGeometry(TSharedRef<FNavigationRelevantData, ESPMode::ThreadSafe> ElementData)
+void FRecastTileGenerator::ValidateAndAppendGeometry(TSharedRef<FNavigationRelevantData, ESPMode::ThreadSafe> ElementData, const FCompositeNavModifier& InModifier)
 {
 	const FNavigationRelevantData& DataRef = ElementData.Get();
 	if (DataRef.IsCollisionDataValid())
 	{
-		AppendGeometry(DataRef.CollisionData, DataRef.NavDataPerInstanceTransformDelegate);
+		AppendGeometry(DataRef.CollisionData, InModifier, DataRef.NavDataPerInstanceTransformDelegate);
 	}
 }
 
-void FRecastTileGenerator::AppendGeometry(const TNavStatArray<uint8>& RawCollisionCache, const FNavDataPerInstanceTransformDelegate& InTransformsDelegate)
+void FRecastTileGenerator::AppendGeometry(const TNavStatArray<uint8>& RawCollisionCache, const FCompositeNavModifier& InModifier, const FNavDataPerInstanceTransformDelegate& InTransformsDelegate)
 {	
 	if (RawCollisionCache.Num() == 0)
 	{
@@ -2232,6 +2238,11 @@ void FRecastTileGenerator::AppendGeometry(const TNavStatArray<uint8>& RawCollisi
 	}
 	
 	FRecastRawGeometryElement GeometryElement;
+
+	// To prevent navmesh generation under the geometry, set the RC_PROJECT_TO_BOTTOM flag to true.
+	// This rasterize triangles as filled columns down to the HF lower bound.
+	GeometryElement.RasterizationFlags = InModifier.GetRejectNavmeshUnderneath() ? RC_PROJECT_TO_BOTTOM : rcRasterizationFlags(0);
+
 	FRecastGeometryCache CollisionCache(RawCollisionCache.GetData());
 	
 	// Gather per instance transforms
@@ -2350,7 +2361,7 @@ bool FRecastTileGenerator::GenerateTile()
 
 struct FTileRasterizationContext
 {
-	FTileRasterizationContext() : SolidHF(0), LayerSet(0), CompactHF(0)
+	FTileRasterizationContext() : SolidHF(0), LayerSet(0), CompactHF(0), RasterizationFlags(rcRasterizationFlags(0))
 	{
 	}
 
@@ -2361,10 +2372,16 @@ struct FTileRasterizationContext
 		rcFreeCompactHeightfield(CompactHF);
 	}
 
+	rcRasterizationFlags GetRasterizationFlags() const { return RasterizationFlags; }
+	void SetRasterizationFlags(rcRasterizationFlags Value) { RasterizationFlags = Value; }
+
 	struct rcHeightfield* SolidHF;
 	struct rcHeightfieldLayerSet* LayerSet;
 	struct rcCompactHeightfield* CompactHF;
 	TArray<FNavMeshTileData> Layers;
+
+private:
+	rcRasterizationFlags RasterizationFlags;
 };
 
 bool FRecastTileGenerator::CreateHeightField(FNavMeshBuildContext& BuildContext, FTileRasterizationContext& RasterContext)
@@ -2403,7 +2420,7 @@ bool FRecastTileGenerator::CreateHeightField(FNavMeshBuildContext& BuildContext,
 	return true;
 }
 
-ETimeSliceWorkResult FRecastTileGenerator::RasterizeGeometryRecastTimeSliced(FNavMeshBuildContext& BuildContext, const TArray<float>& Coords, const TArray<int32>& Indices, FTileRasterizationContext& RasterContext)
+ETimeSliceWorkResult FRecastTileGenerator::RasterizeGeometryRecastTimeSliced(FNavMeshBuildContext& BuildContext, const TArray<float>& Coords, const TArray<int32>& Indices, const rcRasterizationFlags RasterizationFlags, FTileRasterizationContext& RasterContext)
 {
 	QUICK_SCOPE_CYCLE_COUNTER(STAT_Navigation_RasterizeGeometryRecast);
 
@@ -2436,7 +2453,7 @@ ETimeSliceWorkResult FRecastTileGenerator::RasterizeGeometryRecastTimeSliced(FNa
 		rcRasterizeTriangles(&BuildContext,
 			Coords.GetData(), NumVerts,
 			Indices.GetData(), RasterizeGeomRecastTriAreas.GetData(), NumFaces,
-			*RasterContext.SolidHF, TileConfig.walkableClimb);
+			*RasterContext.SolidHF, TileConfig.walkableClimb, RasterizationFlags);
 
 		RasterizeGeomRecastTriAreas.Reset();
 
@@ -2456,7 +2473,7 @@ ETimeSliceWorkResult FRecastTileGenerator::RasterizeGeometryRecastTimeSliced(FNa
 }
 
 
-void FRecastTileGenerator::RasterizeGeometryRecast(FNavMeshBuildContext& BuildContext, const TArray<float>& Coords, const TArray<int32>& Indices, FTileRasterizationContext& RasterContext)
+void FRecastTileGenerator::RasterizeGeometryRecast(FNavMeshBuildContext& BuildContext, const TArray<float>& Coords, const TArray<int32>& Indices, const rcRasterizationFlags RasterizationFlags, FTileRasterizationContext& RasterContext)
 {
 	QUICK_SCOPE_CYCLE_COUNTER(STAT_Navigation_RasterizeGeometryRecast);
 
@@ -2479,7 +2496,7 @@ void FRecastTileGenerator::RasterizeGeometryRecast(FNavMeshBuildContext& BuildCo
 		rcRasterizeTriangles(&BuildContext,
 			Coords.GetData(), NumVerts,
 			Indices.GetData(), RasterizeGeomRecastTriAreas.GetData(), NumFaces,
-			*RasterContext.SolidHF, TileConfig.walkableClimb);
+			*RasterContext.SolidHF, TileConfig.walkableClimb, RasterizationFlags);
 	}
 
 	RasterizeGeomRecastTriAreas.Reset();
@@ -2505,7 +2522,7 @@ void FRecastTileGenerator::RasterizeGeometryTransformCoords(const TArray<float>&
 	}
 }
 
-ETimeSliceWorkResult FRecastTileGenerator::RasterizeGeometryTimeSliced(FNavMeshBuildContext& BuildContext, const TArray<float>& Coords, const TArray<int32>& Indices, const FTransform& LocalToWorld, FTileRasterizationContext& RasterContext)
+ETimeSliceWorkResult FRecastTileGenerator::RasterizeGeometryTimeSliced(FNavMeshBuildContext& BuildContext, const TArray<float>& Coords, const TArray<int32>& Indices, const FTransform& LocalToWorld, const rcRasterizationFlags RasterizationFlags, FTileRasterizationContext& RasterContext)
 {
 	QUICK_SCOPE_CYCLE_COUNTER(STAT_Navigation_RasterizeGeometry);
 	ETimeSliceWorkResult WorkResult = ETimeSliceWorkResult::Succeeded;
@@ -2525,7 +2542,7 @@ ETimeSliceWorkResult FRecastTileGenerator::RasterizeGeometryTimeSliced(FNavMeshB
 	}// fall through to next state
 	case ERasterizeGeomTimeSlicedState::RasterizeGeometryRecast:
 	{
-		WorkResult = RasterizeGeometryRecastTimeSliced(BuildContext, RasterizeGeometryWorldRecastCoords, Indices, RasterContext);
+		WorkResult = RasterizeGeometryRecastTimeSliced(BuildContext, RasterizeGeometryWorldRecastCoords, Indices, RasterizationFlags, RasterContext);
 
 		if (WorkResult != ETimeSliceWorkResult::CallAgainNextTimeSlice)
 		{
@@ -2543,12 +2560,12 @@ ETimeSliceWorkResult FRecastTileGenerator::RasterizeGeometryTimeSliced(FNavMeshB
 	return WorkResult;
 }
 
-void FRecastTileGenerator::RasterizeGeometry(FNavMeshBuildContext& BuildContext, const TArray<float>& Coords, const TArray<int32>& Indices, const FTransform& LocalToWorld, FTileRasterizationContext& RasterContext)
+void FRecastTileGenerator::RasterizeGeometry(FNavMeshBuildContext& BuildContext, const TArray<float>& Coords, const TArray<int32>& Indices, const FTransform& LocalToWorld, const rcRasterizationFlags RasterizationFlags, FTileRasterizationContext& RasterContext)
 {
 	QUICK_SCOPE_CYCLE_COUNTER(STAT_Navigation_RasterizeGeometry);
 
 	RasterizeGeometryTransformCoords(Coords, LocalToWorld);
-	RasterizeGeometryRecast(BuildContext, RasterizeGeometryWorldRecastCoords, Indices, RasterContext);
+	RasterizeGeometryRecast(BuildContext, RasterizeGeometryWorldRecastCoords, Indices, RasterizationFlags, RasterContext);
 }
 
 ETimeSliceWorkResult FRecastTileGenerator::RasterizeTrianglesTimeSliced(FNavMeshBuildContext& BuildContext, FTileRasterizationContext& RasterContext)
@@ -2564,8 +2581,8 @@ ETimeSliceWorkResult FRecastTileGenerator::RasterizeTrianglesTimeSliced(FNavMesh
 			while (RasterizeTrianglesTimeSlicedInstTransformIdx < Element.PerInstanceTransform.Num())
 			{
 				const FTransform& InstanceTransform = Element.PerInstanceTransform[RasterizeTrianglesTimeSlicedInstTransformIdx];
-				const ETimeSliceWorkResult WorkResult = RasterizeGeometryTimeSliced(BuildContext, Element.GeomCoords, Element.GeomIndices, InstanceTransform, RasterContext);
-
+				const ETimeSliceWorkResult WorkResult = RasterizeGeometryTimeSliced(BuildContext, Element.GeomCoords, Element.GeomIndices, InstanceTransform, Element.RasterizationFlags, RasterContext);
+			
 				//the original code just kept calling the RasterizeGeometry() functions and had no return type, 
 				//so we will process the next layer (if we are not needing to process this layer again next time slice) 
 				if (TimeSlicer.IsTimeSliceFinishedCached())
@@ -2585,8 +2602,8 @@ ETimeSliceWorkResult FRecastTileGenerator::RasterizeTrianglesTimeSliced(FNavMesh
 		}
 		else
 		{
-			const ETimeSliceWorkResult WorkResult = RasterizeGeometryRecastTimeSliced(BuildContext, Element.GeomCoords, Element.GeomIndices, RasterContext);
-
+			const ETimeSliceWorkResult WorkResult = RasterizeGeometryRecastTimeSliced(BuildContext, Element.GeomCoords, Element.GeomIndices, Element.RasterizationFlags, RasterContext);
+	
 			if (TimeSlicer.IsTimeSliceFinishedCached())
 			{
 				if (WorkResult != ETimeSliceWorkResult::CallAgainNextTimeSlice)
@@ -2615,12 +2632,12 @@ void FRecastTileGenerator::RasterizeTriangles(FNavMeshBuildContext& BuildContext
 		{
 			for (const FTransform& InstanceTransform : Element.PerInstanceTransform)
 			{
-				RasterizeGeometry(BuildContext, Element.GeomCoords, Element.GeomIndices, InstanceTransform, RasterContext);
+				RasterizeGeometry(BuildContext, Element.GeomCoords, Element.GeomIndices, InstanceTransform, Element.RasterizationFlags, RasterContext);
 			}
 		}
 		else
 		{
-			RasterizeGeometryRecast(BuildContext, Element.GeomCoords, Element.GeomIndices, RasterContext);
+			RasterizeGeometryRecast(BuildContext, Element.GeomCoords, Element.GeomIndices, Element.RasterizationFlags, RasterContext);
 		}
 	}
 }
diff --git a/Engine/Source/Runtime/NavigationSystem/Public/NavMesh/RecastNavMeshGenerator.h b/Engine/Source/Runtime/NavigationSystem/Public/NavMesh/RecastNavMeshGenerator.h
--- a/Engine/Source/Runtime/NavigationSystem/Public/NavMesh/RecastNavMeshGenerator.h
+++ b/Engine/Source/Runtime/NavigationSystem/Public/NavMesh/RecastNavMeshGenerator.h
@@ -150,6 +150,8 @@ struct FRecastRawGeometryElement
 	// Per instance transformations in unreal coords
 	// When empty geometry is in world space
 	TArray<FTransform>	PerInstanceTransform;
+
+	rcRasterizationFlags RasterizationFlags;
 };
 
 struct FRecastAreaNavModifierElement
@@ -347,11 +349,11 @@ protected:
 	bool CreateHeightField(FNavMeshBuildContext& BuildContext, FTileRasterizationContext& RasterContext);
 	ETimeSliceWorkResult RasterizeTrianglesTimeSliced(FNavMeshBuildContext& BuildContext, FTileRasterizationContext& RasterContext);
 	void RasterizeTriangles(FNavMeshBuildContext& BuildContext, FTileRasterizationContext& RasterContext);
-	ETimeSliceWorkResult RasterizeGeometryRecastTimeSliced(FNavMeshBuildContext& BuildContext, const TArray<float>& Coords, const TArray<int32>& Indices, FTileRasterizationContext& RasterContext);
-	void RasterizeGeometryRecast(FNavMeshBuildContext& BuildContext, const TArray<float>& Coords, const TArray<int32>& Indices, FTileRasterizationContext& RasterContext);
+	ETimeSliceWorkResult RasterizeGeometryRecastTimeSliced(FNavMeshBuildContext& BuildContext, const TArray<float>& Coords, const TArray<int32>& Indices, const rcRasterizationFlags RasterizationFlags, FTileRasterizationContext& RasterContext);
+	void RasterizeGeometryRecast(FNavMeshBuildContext& BuildContext, const TArray<float>& Coords, const TArray<int32>& Indices, const rcRasterizationFlags RasterizationFlags, FTileRasterizationContext& RasterContext);
 	void RasterizeGeometryTransformCoords(const TArray<float>& Coords, const FTransform& LocalToWorld);
-	ETimeSliceWorkResult RasterizeGeometryTimeSliced(FNavMeshBuildContext& BuildContext, const TArray<float>& Coords, const TArray<int32>& Indices, const FTransform& LocalToWorld, FTileRasterizationContext& RasterContext);
-	void RasterizeGeometry(FNavMeshBuildContext& BuildContext, const TArray<float>& Coords, const TArray<int32>& Indices, const FTransform& LocalToWorld, FTileRasterizationContext& RasterContext);
+	ETimeSliceWorkResult RasterizeGeometryTimeSliced(FNavMeshBuildContext& BuildContext, const TArray<float>& Coords, const TArray<int32>& Indices, const FTransform& LocalToWorld, const rcRasterizationFlags RasterizationFlags, FTileRasterizationContext& RasterContext);
+	void RasterizeGeometry(FNavMeshBuildContext& BuildContext, const TArray<float>& Coords, const TArray<int32>& Indices, const FTransform& LocalToWorld, const rcRasterizationFlags RasterizationFlags, FTileRasterizationContext& RasterContext);
 	void GenerateRecastFilter(FNavMeshBuildContext& BuildContext, FTileRasterizationContext& RasterContext);
 	bool BuildCompactHeightField(FNavMeshBuildContext& BuildContext, FTileRasterizationContext& RasterContext);
 	bool RecastErodeWalkable(FNavMeshBuildContext& BuildContext, FTileRasterizationContext& RasterContext);
@@ -382,12 +384,12 @@ protected:
 
 	void AppendModifier(const FCompositeNavModifier& Modifier, const FNavDataPerInstanceTransformDelegate& InTransformsDelegate);
 	/** Appends specified geometry to tile's geometry */
-	void ValidateAndAppendGeometry(TSharedRef<FNavigationRelevantData, ESPMode::ThreadSafe> ElementData);
-	void AppendGeometry(const TNavStatArray<uint8>& RawCollisionCache, const FNavDataPerInstanceTransformDelegate& InTransformsDelegate);
+	void ValidateAndAppendGeometry(TSharedRef<FNavigationRelevantData, ESPMode::ThreadSafe> ElementData, const FCompositeNavModifier& InModifier);
+	void AppendGeometry(const TNavStatArray<uint8>& RawCollisionCache, const FCompositeNavModifier& InModifier, const FNavDataPerInstanceTransformDelegate& InTransformsDelegate);
 	void AppendVoxels(rcSpanCache* SpanData, int32 NumSpans);
 	
 	/** prepare voxel cache from collision data */
-	void PrepareVoxelCache(const TNavStatArray<uint8>& RawCollisionCache, TNavStatArray<rcSpanCache>& SpanData);
+	void PrepareVoxelCache(const TNavStatArray<uint8>& RawCollisionCache, const FCompositeNavModifier& InModifier, TNavStatArray<rcSpanCache>& SpanData);
 	bool HasVoxelCache(const TNavStatArray<uint8>& RawVoxelCache, rcSpanCache*& CachedVoxels, int32& NumCachedVoxels) const;
 	void AddVoxelCache(TNavStatArray<uint8>& RawVoxelCache, const rcSpanCache* CachedVoxels, const int32 NumCachedVoxels) const;
 
diff --git a/Engine/Source/Runtime/Navmesh/Private/Recast/RecastRasterization.cpp b/Engine/Source/Runtime/Navmesh/Private/Recast/RecastRasterization.cpp
--- a/Engine/Source/Runtime/Navmesh/Private/Recast/RecastRasterization.cpp
+++ b/Engine/Source/Runtime/Navmesh/Private/Recast/RecastRasterization.cpp
@@ -97,7 +97,7 @@ static void addSpan(rcHeightfield& hf, const int x, const int y,
 	s->data.area = area;
 	s->next = 0;
 	
-	// Empty cell, add he first span.
+	// Empty cell, add the first span.
 	if (!hf.spans[idx])
 	{
 		hf.spans[idx] = s;
@@ -425,12 +425,14 @@ static void rasterizeTriTest(const float* v0, const float* v1, const float* v2,
 static void rasterizeTri(const float* v0, const float* v1, const float* v2,
 						 const unsigned char area, rcHeightfield& hf,
 						 const float* bmin, const float* bmax,
-						 const float cs, const float ics, const float ich,
-						 const int flagMergeThr)
+						 const float cs, const float ics, const float ich, 
+						 const int flagMergeThr,
+						 const int rasterizationFlags /*UE4*/)
 {
 	const int w = hf.width;
 	const int h = hf.height;
 	const float by = bmax[1] - bmin[1];
+	const bool projectToBottom = rasterizationFlags & RC_PROJECT_TO_BOTTOM; //UE4
 
 	int intverts[3][2];
 
@@ -462,11 +464,11 @@ static void rasterizeTri(const float* v0, const float* v1, const float* v2,
 	if (x0 == x1 && y0 == y1)
 	{
 		// Clamp the span to the heightfield bbox.
-		if (triangle_smin < 0.0f) triangle_smin = 0;
+		if (triangle_smin < 0.0f) triangle_smin = 0.0f;
 		if (triangle_smax > by) triangle_smax = by;
 
 		// Snap the span to the heightfield height grid.
-		unsigned short triangle_ismin = (unsigned short)rcClamp((int)floorf(triangle_smin * ich), 0, RC_SPAN_MAX_HEIGHT);
+		unsigned short triangle_ismin = projectToBottom ? 0 : (unsigned short)rcClamp((int)floorf(triangle_smin * ich), 0, RC_SPAN_MAX_HEIGHT); //UE4
 		unsigned short triangle_ismax = (unsigned short)rcClamp((int)ceilf(triangle_smax * ich), (int)triangle_ismin+1, RC_SPAN_MAX_HEIGHT);
 
 		addSpan(hf, x0, y0, triangle_ismin, triangle_ismax, area, flagMergeThr);
@@ -626,7 +628,7 @@ static void rasterizeTri(const float* v0, const float* v1, const float* v2,
 		}
 
 		// Snap the span to the heightfield height grid.
-		unsigned short triangle_ismin_clamp = (unsigned short)rcClamp((int)triangle_ismin, 0, RC_SPAN_MAX_HEIGHT);
+		unsigned short triangle_ismin_clamp = projectToBottom ? 0 : (unsigned short)rcClamp((int)triangle_ismin, 0, RC_SPAN_MAX_HEIGHT); //UE4
 		unsigned short triangle_ismax_clamp = (unsigned short)rcClamp((int)triangle_ismin, (int)triangle_ismin_clamp+1, RC_SPAN_MAX_HEIGHT);
 
 		for (int y = y0; y <= y1; y++)
@@ -808,7 +810,7 @@ static void rasterizeTri(const float* v0, const float* v1, const float* v2,
 				// Skip the span if it is outside the heightfield bbox
 				if (smin >= RC_SPAN_MAX_HEIGHT || smax < 0) continue;
 
-				smin = intMax(smin, 0);
+				smin = projectToBottom ? 0 : intMax(smin, 0); //UE4
 				smax = intMin(intMax(smax,smin+1), RC_SPAN_MAX_HEIGHT);
 
 	#if TEST_NEW_RASTERIZER
@@ -874,12 +876,14 @@ static void rasterizeTri(const float* v0, const float* v1, const float* v2,
 						 const unsigned char area, rcHeightfield& hf,
 						 const float* bmin, const float* bmax,
 						 const float cs, const float ics, const float ich,
-						 const int flagMergeThr)
+						 const int flagMergeThr,
+						 const int rasterizationFlags) //UE4
 {
 	const int w = hf.width;
 	const int h = hf.height;
 	float tmin[3], tmax[3];
 	const float by = bmax[1] - bmin[1];
+	const bool projectToBottom = flags & RC_PROJECT_TO_BOTTOM; //UE4
 	
 	// Calculate the bounding box of the triangle.
 	rcVcopy(tmin, v0);
@@ -889,11 +893,11 @@ static void rasterizeTri(const float* v0, const float* v1, const float* v2,
 	rcVmax(tmax, v1);
 	rcVmax(tmax, v2);
 	
-	// If the triangle does not touch the bbox of the heightfield, skip the triagle.
+	// If the triangle does not touch the bbox of the heightfield, skip the triangle.
 	if (!overlapBounds(bmin, bmax, tmin, tmax))
 		return;
 	
-	// Calculate the footpring of the triangle on the grid.
+	// Calculate the footprint of the triangle on the grid.
 	int x0 = (int)((tmin[0] - bmin[0])*ics);
 	int y0 = (int)((tmin[2] - bmin[2])*ics);
 	int x1 = (int)((tmax[0] - bmin[0])*ics);
@@ -946,7 +950,7 @@ static void rasterizeTri(const float* v0, const float* v1, const float* v2,
 			if (smax > by) smax = by;
 			
 			// Snap the span to the heightfield height grid.
-			unsigned short ismin = (unsigned short)rcClamp((int)floorf(smin * ich), 0, RC_SPAN_MAX_HEIGHT);
+			unsigned short ismin = projectToBottom ? 0 : (unsigned short)rcClamp((int)floorf(smin * ich), 0, RC_SPAN_MAX_HEIGHT); //UE4
 			unsigned short ismax = (unsigned short)rcClamp((int)ceilf(smax * ich), (int)ismin+1, RC_SPAN_MAX_HEIGHT);
 			
 			addSpan(hf, x, y, ismin, ismax, area, flagMergeThr);
@@ -961,7 +965,8 @@ static void rasterizeTri(const float* v0, const float* v1, const float* v2,
 /// @see rcHeightfield
 void rcRasterizeTriangle(rcContext* ctx, const float* v0, const float* v1, const float* v2,
 						 const unsigned char area, rcHeightfield& solid,
-						 const int flagMergeThr)
+						 const int flagMergeThr, 
+						 const int rasterizationFlags) //UE4
 {
 	rcAssert(ctx);
 
@@ -969,7 +974,7 @@ void rcRasterizeTriangle(rcContext* ctx, const float* v0, const float* v1, const
 
 	const float ics = 1.0f/solid.cs;
 	const float ich = 1.0f/solid.ch;
-	rasterizeTri(v0, v1, v2, area, solid, solid.bmin, solid.bmax, solid.cs, ics, ich, flagMergeThr);
+	rasterizeTri(v0, v1, v2, area, solid, solid.bmin, solid.bmax, solid.cs, ics, ich, flagMergeThr, rasterizationFlags); //UE4
 
 	ctx->stopTimer(RC_TIMER_RASTERIZE_TRIANGLES);
 }
@@ -981,7 +986,7 @@ void rcRasterizeTriangle(rcContext* ctx, const float* v0, const float* v1, const
 /// @see rcHeightfield
 void rcRasterizeTriangles(rcContext* ctx, const float* verts, const int /*nv*/,
 						  const int* tris, const unsigned char* areas, const int nt,
-						  rcHeightfield& solid, const int flagMergeThr)
+						  rcHeightfield& solid, const int flagMergeThr, const int rasterizationFlags) //UE4
 {
 	if (ctx)
 		ctx->startTimer(RC_TIMER_RASTERIZE_TRIANGLES);
@@ -995,7 +1000,7 @@ void rcRasterizeTriangles(rcContext* ctx, const float* verts, const int /*nv*/,
 		const float* v1 = &verts[tris[i*3+1]*3];
 		const float* v2 = &verts[tris[i*3+2]*3];
 		// Rasterize.
-		rasterizeTri(v0, v1, v2, areas[i], solid, solid.bmin, solid.bmax, solid.cs, ics, ich, flagMergeThr);
+		rasterizeTri(v0, v1, v2, areas[i], solid, solid.bmin, solid.bmax, solid.cs, ics, ich, flagMergeThr, rasterizationFlags); //UE4
 	}
 	
 	if (ctx)
@@ -1009,7 +1014,7 @@ void rcRasterizeTriangles(rcContext* ctx, const float* verts, const int /*nv*/,
 /// @see rcHeightfield
 void rcRasterizeTriangles(rcContext* ctx, const float* verts, const int /*nv*/,
 						  const unsigned short* tris, const unsigned char* areas, const int nt,
-						  rcHeightfield& solid, const int flagMergeThr)
+						  rcHeightfield& solid, const int flagMergeThr, const int rasterizationFlags) //UE4
 {
 	if (ctx)
 		ctx->startTimer(RC_TIMER_RASTERIZE_TRIANGLES);
@@ -1023,7 +1028,7 @@ void rcRasterizeTriangles(rcContext* ctx, const float* verts, const int /*nv*/,
 		const float* v1 = &verts[tris[i*3+1]*3];
 		const float* v2 = &verts[tris[i*3+2]*3];
 		// Rasterize.
-		rasterizeTri(v0, v1, v2, areas[i], solid, solid.bmin, solid.bmax, solid.cs, ics, ich, flagMergeThr);
+		rasterizeTri(v0, v1, v2, areas[i], solid, solid.bmin, solid.bmax, solid.cs, ics, ich, flagMergeThr, rasterizationFlags); //UE4
 	}
 	
 	if (ctx)
@@ -1036,7 +1041,7 @@ void rcRasterizeTriangles(rcContext* ctx, const float* verts, const int /*nv*/,
 ///
 /// @see rcHeightfield
 void rcRasterizeTriangles(rcContext* ctx, const float* verts, const unsigned char* areas, const int nt,
-						  rcHeightfield& solid, const int flagMergeThr)
+						  rcHeightfield& solid, const int flagMergeThr, const int rasterizationFlags) //UE4
 {
 	if (ctx)
 		ctx->startTimer(RC_TIMER_RASTERIZE_TRIANGLES);
@@ -1050,7 +1055,7 @@ void rcRasterizeTriangles(rcContext* ctx, const float* verts, const unsigned cha
 		const float* v1 = &verts[(i*3+1)*3];
 		const float* v2 = &verts[(i*3+2)*3];
 		// Rasterize.
-		rasterizeTri(v0, v1, v2, areas[i], solid, solid.bmin, solid.bmax, solid.cs, ics, ich, flagMergeThr);
+		rasterizeTri(v0, v1, v2, areas[i], solid, solid.bmin, solid.bmax, solid.cs, ics, ich, flagMergeThr, rasterizationFlags); //UE4
 	}
 	
 	if (ctx)
diff --git a/Engine/Source/Runtime/Navmesh/Public/Recast/Recast.h b/Engine/Source/Runtime/Navmesh/Public/Recast/Recast.h
--- a/Engine/Source/Runtime/Navmesh/Public/Recast/Recast.h
+++ b/Engine/Source/Runtime/Navmesh/Public/Recast/Recast.h
@@ -236,7 +236,7 @@ struct rcConfig
 	/// The maximum allowed length for contour edges along the border of the mesh. [Limit: >=0] [Units: vx] 
 	int maxEdgeLen;
 	
-	/// The maximum distance a simplfied contour's border edges should deviate 
+	/// The maximum distance a simplified contour's border edges should deviate 
 	/// the original raw contour. [Limit: >=0] [Units: wu]
 	float maxSimplificationError;
 	
@@ -613,6 +613,12 @@ enum rcFilterLowAreaFlags
 	RC_LOW_FILTER_POST_PROCESS = 0x02,	///< additional filtering at the end
 };
 
+// UE4
+enum rcRasterizationFlags
+{
+	RC_PROJECT_TO_BOTTOM = 1 << 0,		///< Will create spans from the triangle surface to the bottom of the heightfield
+};
+
 /// Applied to the region id field of contour vertices in order to extract the region id.
 /// The region id field of a vertex may have several flags applied to it.  So the
 /// fields value can't be used directly.
@@ -934,24 +940,26 @@ NAVMESH_API void rcRasterizeTriangle(rcContext* ctx, const float* v0, const floa
 ///  @param[in,out]	solid			An initialized heightfield.
 ///  @param[in]		flagMergeThr	The distance where the walkable flag is favored over the non-walkable flag. 
 ///  								[Limit: >= 0] [Units: vx]
+///  @param[in]     int				Flags to change the rasterization behavior		//UE4
 NAVMESH_API void rcRasterizeTriangles(rcContext* ctx, const float* verts, const int nv,
 						  const int* tris, const unsigned char* areas, const int nt,
-						  rcHeightfield& solid, const int flagMergeThr = 1);
+						  rcHeightfield& solid, const int flagMergeThr = 1, const int rasterizationFlags = 0 /*UE4*/);
 
 /// Rasterizes an indexed triangle mesh into the specified heightfield.
 ///  @ingroup recast
-///  @param[in,out]	ctx			The build context to use during the operation.
-///  @param[in]		verts		The vertices. [(x, y, z) * @p nv]
-///  @param[in]		nv			The number of vertices.
-///  @param[in]		tris		The triangle indices. [(vertA, vertB, vertC) * @p nt]
-///  @param[in]		areas		The area id's of the triangles. [Limit: <= #RC_WALKABLE_AREA] [Size: @p nt]
-///  @param[in]		nt			The number of triangles.
-///  @param[in,out]	solid		An initialized heightfield.
+///  @param[in,out]	ctx				The build context to use during the operation.
+///  @param[in]		verts			The vertices. [(x, y, z) * @p nv]
+///  @param[in]		nv				The number of vertices.
+///  @param[in]		tris			The triangle indices. [(vertA, vertB, vertC) * @p nt]
+///  @param[in]		areas			The area id's of the triangles. [Limit: <= #RC_WALKABLE_AREA] [Size: @p nt]
+///  @param[in]		nt				The number of triangles.
+///  @param[in,out]	solid			An initialized heightfield.
 ///  @param[in]		flagMergeThr	The distance where the walkable flag is favored over the non-walkable flag. 
-///  							[Limit: >= 0] [Units: vx]
+///  								[Limit: >= 0] [Units: vx]
+///  @param[in]     int				Flags to change the rasterization behavior		//UE4
 NAVMESH_API void rcRasterizeTriangles(rcContext* ctx, const float* verts, const int nv,
 						  const unsigned short* tris, const unsigned char* areas, const int nt,
-						  rcHeightfield& solid, const int flagMergeThr = 1);
+						  rcHeightfield& solid, const int flagMergeThr = 1, const int rasterizationFlags = 0 /*UE4*/);
 
 /// Rasterizes triangles into the specified heightfield.
 ///  @ingroup recast
@@ -965,7 +973,7 @@ NAVMESH_API void rcRasterizeTriangles(rcContext* ctx, const float* verts, const
 NAVMESH_API void rcRasterizeTriangles(rcContext* ctx, const float* verts, const unsigned char* areas, const int nt,
 						  rcHeightfield& solid, const int flagMergeThr = 1);
 
-/// Marks non-walkable spans as walkable if their maximum is within @p walkableClimp of a walkable neihbor. 
+/// Marks non-walkable spans as walkable if their maximum is within @p walkableClimp of a walkable neighbor. 
 ///  @ingroup recast
 ///  @param[in,out]	ctx				The build context to use during the operation.
 ///  @param[in]		walkableClimb	Maximum ledge height that is considered to still be traversable. 
@@ -984,7 +992,7 @@ NAVMESH_API void rcFilterLowHangingWalkableObstacles(rcContext* ctx, const int w
 NAVMESH_API void rcFilterLedgeSpans(rcContext* ctx, const int walkableHeight,
 						const int walkableClimb, rcHeightfield& solid);
 
-/// Marks walkable spans as not walkable if the clearence above the span is less than the specified height. 
+/// Marks walkable spans as not walkable if the clearance above the span is less than the specified height. 
 ///  @ingroup recast
 ///  @param[in,out]	ctx				The build context to use during the operation.
 ///  @param[in]		walkableHeight	Minimum floor to 'ceiling' height that will still allow the floor area to 
