Navmesh -> renaming bRejectNavmeshUnderneath to bFillCollisionUnderneathForNavmesh and derivatives
Fix update issue with the flag

diff --git a/Engine/Source/Runtime/Engine/Classes/Components/PrimitiveComponent.h b/Engine/Source/Runtime/Engine/Classes/Components/PrimitiveComponent.h
--- a/Engine/Source/Runtime/Engine/Classes/Components/PrimitiveComponent.h
+++ b/Engine/Source/Runtime/Engine/Classes/Components/PrimitiveComponent.h
@@ -491,7 +491,7 @@ public:
 
 	/** If set, navmesh will not be generated under the surface of the geometry */
 	UPROPERTY(EditAnywhere, Category = Navigation)
-	uint8 bRejectNavmeshUnderneath:1;
+	uint8 bFillCollisionUnderneathForNavmesh:1;
 
 	// General flags.
 	
diff --git a/Engine/Source/Runtime/Engine/Private/AI/Navigation/NavigationModifier.cpp b/Engine/Source/Runtime/Engine/Private/AI/Navigation/NavigationModifier.cpp
--- a/Engine/Source/Runtime/Engine/Private/AI/Navigation/NavigationModifier.cpp
+++ b/Engine/Source/Runtime/Engine/Private/AI/Navigation/NavigationModifier.cpp
@@ -741,7 +741,7 @@ void FCompositeNavModifier::Empty()
 	CustomLinks.Empty();
 	bHasPotentialLinks = false;
 	bAdjustHeight = false;
-	bModifierRejectNavmeshUnderneath = false;
+	bModifierFillCollisionUnderneathForNavmesh = false;
 }
 
 FCompositeNavModifier FCompositeNavModifier::GetInstantiatedMetaModifier(const FNavAgentProperties* NavAgent, TWeakObjectPtr<UObject> WeakOwnerPtr) const
diff --git a/Engine/Source/Runtime/Engine/Private/Components/PrimitiveComponent.cpp b/Engine/Source/Runtime/Engine/Private/Components/PrimitiveComponent.cpp
--- a/Engine/Source/Runtime/Engine/Private/Components/PrimitiveComponent.cpp
+++ b/Engine/Source/Runtime/Engine/Private/Components/PrimitiveComponent.cpp
@@ -324,7 +324,7 @@ UPrimitiveComponent::UPrimitiveComponent(const FObjectInitializer& ObjectInitial
 	bCastVolumetricTranslucentShadow = false;
 	IndirectLightingCacheQuality = ILCQ_Point;
 	bSelectable = true;
-	bRejectNavmeshUnderneath = false;
+	bFillCollisionUnderneathForNavmesh = false;
 	AlwaysLoadOnClient = true;
 	AlwaysLoadOnServer = true;
 	SetCollisionProfileName(UCollisionProfile::BlockAll_ProfileName);
@@ -2379,10 +2379,10 @@ void UPrimitiveComponent::DispatchWakeEvents(ESleepEvent WakeEvent, FName BoneNa
 
 void UPrimitiveComponent::GetNavigationData(FNavigationRelevantData& OutData) const
 {
-	if (bRejectNavmeshUnderneath)
+	if (bFillCollisionUnderneathForNavmesh)
 	{
 		FCompositeNavModifier CompositeNavModifier;
-		CompositeNavModifier.SetRejectNavmeshUnderneath(bRejectNavmeshUnderneath);
+		CompositeNavModifier.SetFillCollisionUnderneathForNavmesh(bFillCollisionUnderneathForNavmesh);
 		OutData.Modifiers.Add(CompositeNavModifier);
 	}
 }
diff --git a/Engine/Source/Runtime/Engine/Public/AI/NavigationModifier.h b/Engine/Source/Runtime/Engine/Public/AI/NavigationModifier.h
--- a/Engine/Source/Runtime/Engine/Public/AI/NavigationModifier.h
+++ b/Engine/Source/Runtime/Engine/Public/AI/NavigationModifier.h
@@ -288,7 +288,7 @@ struct ENGINE_API FCompositeNavModifier : public FNavigationModifier
 		bHasMetaAreas |= Modifiers.bHasMetaAreas; 
 		bAdjustHeight |= Modifiers.HasAgentHeightAdjust();
 		bHasLowAreaModifiers |= Modifiers.HasLowAreaModifiers();
-		bModifierRejectNavmeshUnderneath |= Modifiers.GetRejectNavmeshUnderneath();
+		bModifierFillCollisionUnderneathForNavmesh |= Modifiers.GetFillCollisionUnderneathForNavmesh();
 	}
 
 	void CreateAreaModifiers(const UPrimitiveComponent* PrimComp, const TSubclassOf<UNavAreaBase> AreaClass);
diff --git a/Engine/Source/Runtime/Landscape/Classes/LandscapeProxy.h b/Engine/Source/Runtime/Landscape/Classes/LandscapeProxy.h
--- a/Engine/Source/Runtime/Landscape/Classes/LandscapeProxy.h
+++ b/Engine/Source/Runtime/Landscape/Classes/LandscapeProxy.h
@@ -697,7 +697,7 @@ public:
 
 	/** Set to true to prevent navmesh generation under the terrain geometry */
 	UPROPERTY(EditAnywhere, Category = Landscape)
-	uint32 bRejectNavmeshUnderLandscapeGeometry:1;
+	uint32 bFillCollisionUnderLandscapeForNavmesh:1;
 
 	/** When set to true it will generate MaterialInstanceDynamic for each components, so material can be changed at runtime */
 	UPROPERTY(EditAnywhere, Category = Landscape)
diff --git a/Engine/Source/Runtime/Landscape/Private/Landscape.cpp b/Engine/Source/Runtime/Landscape/Private/Landscape.cpp
--- a/Engine/Source/Runtime/Landscape/Private/Landscape.cpp
+++ b/Engine/Source/Runtime/Landscape/Private/Landscape.cpp
@@ -1048,7 +1048,7 @@ ALandscapeProxy::ALandscapeProxy(const FObjectInitializer& ObjectInitializer)
 	bCastStaticShadow = true;
 	bCastShadowAsTwoSided = false;
 	bUsedForNavigation = true;
-	bRejectNavmeshUnderLandscapeGeometry = false;
+	bFillCollisionUnderLandscapeForNavmesh = false;
 	CollisionThickness = 16;
 	BodyInstance.SetCollisionProfileName(UCollisionProfile::BlockAll_ProfileName);
 	bGenerateOverlapEvents = false;
@@ -2003,14 +2003,6 @@ void ALandscape::PostLoad()
 			Brush.SetOwner(this);
 		}
 	}
-
-	for (ULandscapeComponent* Comp : LandscapeComponents)
-	{
-		if (Comp)
-		{
-			Comp->UpdateRejectNavmeshUnderneath();
-		}
-	}
 #endif
 
 	Super::PostLoad();
@@ -2442,6 +2434,14 @@ void ALandscapeProxy::PostLoad()
 		}
 	}
 
+	for (ULandscapeComponent* Comp : LandscapeComponents)
+	{
+		if (Comp)
+		{
+			Comp->UpdateRejectNavmeshUnderneath();
+		}
+	}
+
 	// track feature level change to flush grass cache
 	FOnFeatureLevelChanged::FDelegate FeatureLevelChangedDelegate = FOnFeatureLevelChanged::FDelegate::CreateUObject(this, &ALandscapeProxy::OnFeatureLevelChanged);
 	FeatureLevelChangedDelegateHandle = GetWorld()->AddOnFeatureLevelChangedHandler(FeatureLevelChangedDelegate);
diff --git a/Engine/Source/Runtime/Landscape/Private/LandscapeEdit.cpp b/Engine/Source/Runtime/Landscape/Private/LandscapeEdit.cpp
--- a/Engine/Source/Runtime/Landscape/Private/LandscapeEdit.cpp
+++ b/Engine/Source/Runtime/Landscape/Private/LandscapeEdit.cpp
@@ -154,7 +154,7 @@ void ULandscapeComponent::UpdateRejectNavmeshUnderneath()
 	ALandscapeProxy* Proxy = GetLandscapeProxy();
 	if (CollisionComponent && Proxy)
 	{
-		CollisionComponent->bRejectNavmeshUnderneath = Proxy->bRejectNavmeshUnderLandscapeGeometry;
+		CollisionComponent->bFillCollisionUnderneathForNavmesh = Proxy->bFillCollisionUnderLandscapeForNavmesh;
 	}
 }
 
@@ -4860,7 +4860,7 @@ void ALandscape::PostEditChangeProperty(FPropertyChangedEvent& PropertyChangedEv
 	{
 		bChangedNavRelevance = true;
 	}
-	else if (GIsEditor && PropertyName == GET_MEMBER_NAME_CHECKED(ALandscapeProxy, bRejectNavmeshUnderLandscapeGeometry))
+	else if (GIsEditor && PropertyName == GET_MEMBER_NAME_CHECKED(ALandscapeProxy, bFillCollisionUnderLandscapeForNavmesh))
 	{
 		bChangeRejectNavmeshUnder = true;
 	}
diff --git a/Engine/Source/Runtime/NavigationSystem/Private/NavMesh/RecastNavMeshGenerator.cpp b/Engine/Source/Runtime/NavigationSystem/Private/NavMesh/RecastNavMeshGenerator.cpp
--- a/Engine/Source/Runtime/NavigationSystem/Private/NavMesh/RecastNavMeshGenerator.cpp
+++ b/Engine/Source/Runtime/NavigationSystem/Private/NavMesh/RecastNavMeshGenerator.cpp
@@ -2130,7 +2130,7 @@ void FRecastTileGenerator::PrepareVoxelCache(const TNavStatArray<uint8>& RawColl
 
 	// To prevent navmesh generation under the triangles, set the RC_PROJECT_TO_BOTTOM flag to true.
 	// This rasterize triangles as filled columns down to the HF lower bound.
-	const rcRasterizationFlags Flags = InModifier.GetRejectNavmeshUnderneath() ? RC_PROJECT_TO_BOTTOM : rcRasterizationFlags(0);
+	const rcRasterizationFlags Flags = InModifier.GetFillCollisionUnderneathForNavmesh() ? RC_PROJECT_TO_BOTTOM : rcRasterizationFlags(0);
 
 	rcRasterizeTriangles(0, CachedCollisions.Verts, CachedCollisions.Header.NumVerts,
 		CachedCollisions.Indices, TriAreas.GetData(), CachedCollisions.Header.NumFaces,
@@ -2241,7 +2241,7 @@ void FRecastTileGenerator::AppendGeometry(const TNavStatArray<uint8>& RawCollisi
 
 	// To prevent navmesh generation under the geometry, set the RC_PROJECT_TO_BOTTOM flag to true.
 	// This rasterize triangles as filled columns down to the HF lower bound.
-	GeometryElement.RasterizationFlags = InModifier.GetRejectNavmeshUnderneath() ? RC_PROJECT_TO_BOTTOM : rcRasterizationFlags(0);
+	GeometryElement.RasterizationFlags = InModifier.GetFillCollisionUnderneathForNavmesh() ? RC_PROJECT_TO_BOTTOM : rcRasterizationFlags(0);
 
 	FRecastGeometryCache CollisionCache(RawCollisionCache.GetData());
 	
