Work to make the navigation mesh timeslice adaptive for better handling of both large and small numbers of queued tiles. This is based on the moving window average tile regen time, the moving window average delta time, and a real world time we'd ideally like to get all the tiles in the queue regened by clamped to a min and max frame time slice time.

I've also removed the notion of a max amount of tiles to regen per frame for the timesliced tile regen.

diff --git a/Engine/Source/Runtime/Core/Public/Math/MovingWindowAverageFast.h b/Engine/Source/Runtime/Core/Public/Math/MovingWindowAverageFast.h
new file mode 100644
--- /dev/null
+++ b/Engine/Source/Runtime/Core/Public/Math/MovingWindowAverageFast.h
@@ -0,0 +1,60 @@
+// Copyright 1998-2019 Epic Games, Inc. All Rights Reserved.
+#pragma once
+#include "CoreMinimal.h"
+
+/**
+ * This class calculates a moving window average. Its designed to be used with floats or doubles and
+ * keeps track of the average with every value pushed so is ideal when there is a one to one or one to many relationship
+ * between calls to PushValue() and GetMovingWindowAverage() respectively.
+ */
+template <typename T, int32 ArraySize>
+class FMovingWindowAverageFast
+{
+public:
+	FMovingWindowAverageFast()
+		: TotalValues(static_cast<T>(0))
+		, AverageValue(static_cast<T>(0))
+		, RemoveNextIdx(0)
+		, NumValuesUsed(0)
+	{
+		static_assert(ArraySize > 0, TEXT("ArraySize must be greater than zero"));
+	}
+
+	void PushValue(T Value)
+	{
+		T ValueRemoved = static_cast<T>(0);
+
+		const T NumItemsPrev = static_cast<T>(NumValuesUsed);
+
+		if (ArraySize == NumValuesUsed)
+		{
+			ValueRemoved = ValuesArray[RemoveNextIdx];
+			ValuesArray[RemoveNextIdx] = Value;
+			RemoveNextIdx = (RemoveNextIdx + 1) % ArraySize;
+		}
+		else
+		{
+			ValuesArray[NumValuesUsed] = Value;
+			++NumValuesUsed;
+		}
+
+		const T MovingWindowItemsNumCur = static_cast<T>(NumValuesUsed);
+		TotalValues = TotalValues - ValueRemoved + Value;
+		AverageValue = TotalValues / MovingWindowItemsNumCur;
+	}
+
+	T GetAverage() const
+	{
+		return AverageValue;
+	}
+
+private:
+	TStaticArray<T, ArraySize> ValuesArray;
+
+	T TotalValues;
+	T AverageValue;
+
+	/** The array Index of the next item to remove when the moving window is full */
+	int32 RemoveNextIdx;
+	int32 NumValuesUsed;
+};
diff --git a/Engine/Source/Runtime/NavigationSystem/Private/NavMesh/RecastNavMeshGenerator.cpp b/Engine/Source/Runtime/NavigationSystem/Private/NavMesh/RecastNavMeshGenerator.cpp
--- a/Engine/Source/Runtime/NavigationSystem/Private/NavMesh/RecastNavMeshGenerator.cpp
+++ b/Engine/Source/Runtime/NavigationSystem/Private/NavMesh/RecastNavMeshGenerator.cpp
@@ -45,6 +45,8 @@
 
 #define TEXT_WEAKOBJ_NAME(obj) (obj.IsValid(false) ? *obj->GetName() : (obj.IsValid(false, true)) ? TEXT("MT-Unreachable") : TEXT("INVALID"))
 
+CSV_DEFINE_CATEGORY(NAVREGEN, false);
+
 struct dtTileCacheAlloc;
 
 FORCEINLINE bool DoesBoxContainOrOverlapVector(const FBox& BigBox, const FVector& In)
@@ -1540,12 +1542,17 @@ void FTimeSlicer::SetTimeSliceDuration(double SliceDuration)
 	TimeSliceDuration = SliceDuration;
 }
 
-void FTimeSlicer::ResetStartTime()
+void FTimeSlicer::StartTimeSlice()
 {
 	TimeSliceStartTime = FPlatformTime::Seconds();
 	bTimeSliceFinishedCached = false;
 }
 
+double FTimeSlicer::GetStartTime() const
+{
+	return TimeSliceStartTime;
+}
+
 bool FTimeSlicer::TestTimeSliceFinished() const
 {
 	ensureMsgf(!bTimeSliceFinishedCached, TEXT("Testing time slice is finished when we have already confirmed that!"));
@@ -3820,20 +3827,7 @@ FRecastNavMeshGenerator::FRecastNavMeshGenerator(ARecastNavMesh& InDestNavMesh)
 	, bRestrictBuildingToActiveTiles(false)
 	, bSortTilesWithSeedLocations(true)
 	, Version(0)
-	, TimeSlicer(0.0025)
-#if TIME_SLICE_NAV_REGEN
-	, bTimeSliceRegenActive(true)
-	, bNextTimeSliceRegenActive(true)
-#else
-	, bTimeSliceRegenActive(false)
-	, bNextTimeSliceRegenActive(false)
-#endif
 {
-	ProcessTileTasksSyncTimeSlicedState = EProcessTileTasksSyncTimeSlicedState::Init;
-
-	AddGeneratedTilesTimeSlicedState = EAddGeneratedTilesTimeSlicedState::Init;
-	AddGenTilesTimeSlicedLayerIndex = 0;
-
 	INC_DWORD_STAT_BY(STAT_NavigationMemory, sizeof(*this));
 }
 
@@ -4365,15 +4359,15 @@ bool FRecastNavMeshGenerator::IsInActiveSet(const FIntPoint& Tile) const
 
 void FRecastNavMeshGenerator::ResetTimeSlicedTileGeneratorSync()
 {
-	TimeSlicedTileGeneratorSync.Reset();
+	SyncTimeSlicedData.TileGeneratorSync.Reset();
 
 	//reset variables used for timeslicing TileGenratorSync
-	ProcessTileTasksSyncTimeSlicedState = EProcessTileTasksSyncTimeSlicedState::Init;
-	ProcessTileTasksSyncTimeSlicedUpdatedTilesCache.Reset();
-	AddGenTilesTimeSlicedOldLayerTileIdMap.Reset();
-	AddGenTilesTimeSlicedResultTileIndices.Reset();
-	AddGeneratedTilesTimeSlicedState = EAddGeneratedTilesTimeSlicedState::Init;
-	AddGenTilesTimeSlicedLayerIndex = 0;
+	SyncTimeSlicedData.ProcessTileTasksSyncState = EProcessTileTasksSyncTimeSlicedState::Init;
+	SyncTimeSlicedData.UpdatedTilesCache.Reset();
+	SyncTimeSlicedData.OldLayerTileIdMapCached.Reset();
+	SyncTimeSlicedData.ResultTileIndicesCached.Reset();
+	SyncTimeSlicedData.AddGeneratedTilesState = EAddGeneratedTilesTimeSlicedState::Init;
+	SyncTimeSlicedData.AddGenTilesLayerIndex = 0;
 }
 
 //@TODO Investigate removing from RunningDirtyTiles here too (or atleast not using the results in any way)
@@ -4390,9 +4384,9 @@ void FRecastNavMeshGenerator::RemoveTiles(const TArray<FIntPoint>& Tiles)
 			PendingDirtyTiles.Remove(DirtyTile);
 		}
 
-		if (TimeSlicedTileGeneratorSync.Get())
+		if (SyncTimeSlicedData.TileGeneratorSync.Get())
 		{
-			if (TimeSlicedTileGeneratorSync->GetTileX() == TileXY.X && TimeSlicedTileGeneratorSync->GetTileY() == TileXY.Y)
+			if (SyncTimeSlicedData.TileGeneratorSync->GetTileX() == TileXY.X && SyncTimeSlicedData.TileGeneratorSync->GetTileY() == TileXY.Y)
 			{
 				ResetTimeSlicedTileGeneratorSync();
 			}
@@ -4534,6 +4528,26 @@ TArray<uint32> FRecastNavMeshGenerator::RemoveTileLayers(const int32 TileX, cons
 	return UpdatedIndices;
 }
 
+FRecastNavMeshGenerator::FSyncTimeSlicedData::FSyncTimeSlicedData()
+	: CurrentTileRegenDuration(0.)
+	, MinTimeSliceDuration(0.00075)
+	, MaxTimeSliceDuration(0.004)
+	, RealTimeSecsLastCall(-1.f)
+	, MaxDesiredTileRegenDuration(0.7f)
+#if TIME_SLICE_NAV_REGEN
+	, bTimeSliceRegenActive(true)
+	, bNextTimeSliceRegenActive(true)
+#else
+	, bTimeSliceRegenActive(false)
+	, bNextTimeSliceRegenActive(false)
+#endif
+	, ProcessTileTasksSyncState(EProcessTileTasksSyncTimeSlicedState::Init)
+	, AddGeneratedTilesState(EAddGeneratedTilesTimeSlicedState::Init)
+	, AddGenTilesLayerIndex(0)
+	, TimeSlicer(0.0025)
+{
+}
+
 void FRecastNavMeshGenerator::AddGeneratedTileLayer(int32 LayerIndex, FRecastTileGenerator& TileGenerator, const TMap<int32, dtPolyRef>& OldLayerTileIdMap, TArray<uint32>& OutResultTileIndices)
 {
 	SCOPE_CYCLE_COUNTER(STAT_Navigation_RecastAddGeneratedTileLayer);
@@ -4636,22 +4650,22 @@ ETimeSliceWorkResult FRecastNavMeshGenerator::AddGeneratedTilesTimeSliced(FRecas
 	ETimeSliceWorkResult WorkResult = ETimeSliceWorkResult::Succeeded;
 	bool bIteratedThroughDirtyLayers = true;
 
-	switch (AddGeneratedTilesTimeSlicedState)
+	switch (SyncTimeSlicedData.AddGeneratedTilesState)
 	{
 	case EAddGeneratedTilesTimeSlicedState::Init:
 	{
-		AddGenTilesTimeSlicedResultTileIndices.Reset();
-		AddGenTilesTimeSlicedResultTileIndices.Reserve(TileLayers.Num());
-		AddGenTilesTimeSlicedOldLayerTileIdMap.Reset();
-		AddGenTilesTimeSlicedOldLayerTileIdMap.Reserve(TileLayers.Num());
-		AddGenTilesTimeSlicedLayerIndex = TileGenerator.GetDirtyLayersMask().Find(true);
+		SyncTimeSlicedData.ResultTileIndicesCached.Reset();
+		SyncTimeSlicedData.ResultTileIndicesCached.Reserve(TileLayers.Num());
+		SyncTimeSlicedData.OldLayerTileIdMapCached.Reset();
+		SyncTimeSlicedData.OldLayerTileIdMapCached.Reserve(TileLayers.Num());
+		SyncTimeSlicedData.AddGenTilesLayerIndex = TileGenerator.GetDirtyLayersMask().Find(true);
 		if (TileGenerator.IsFullyRegenerated())
 		{
 			// remove all layers
-			AddGenTilesTimeSlicedResultTileIndices = RemoveTileLayers(TileX, TileY, &AddGenTilesTimeSlicedOldLayerTileIdMap);
+			SyncTimeSlicedData.ResultTileIndicesCached = RemoveTileLayers(TileX, TileY, &SyncTimeSlicedData.OldLayerTileIdMapCached);
 		}
 
-		AddGeneratedTilesTimeSlicedState = EAddGeneratedTilesTimeSlicedState::AddTiles;
+		SyncTimeSlicedData.AddGeneratedTilesState = EAddGeneratedTilesTimeSlicedState::AddTiles;
 	}//fall through to next state
 	case EAddGeneratedTilesTimeSlicedState::AddTiles:
 	{
@@ -4660,21 +4674,21 @@ ETimeSliceWorkResult FRecastNavMeshGenerator::AddGeneratedTilesTimeSliced(FRecas
 			// and there. It will be removed the next time around the invokers get
 			// updated 
 			// && IsInActiveSet(FIntPoint(TileX, TileY))
-			&& AddGenTilesTimeSlicedLayerIndex != INDEX_NONE)
+			&& SyncTimeSlicedData.AddGenTilesLayerIndex != INDEX_NONE)
 		{
-			for (; AddGenTilesTimeSlicedLayerIndex < TileGenerator.GetDirtyLayersMask().Num(); ++AddGenTilesTimeSlicedLayerIndex)
+			for (; SyncTimeSlicedData.AddGenTilesLayerIndex < TileGenerator.GetDirtyLayersMask().Num(); ++SyncTimeSlicedData.AddGenTilesLayerIndex)
 			{
-				if (TileGenerator.IsLayerChanged(AddGenTilesTimeSlicedLayerIndex))
+				if (TileGenerator.IsLayerChanged(SyncTimeSlicedData.AddGenTilesLayerIndex))
 				{
-					if (TimeSlicer.IsTimeSliceFinishedCached())
+					if (SyncTimeSlicedData.TimeSlicer.IsTimeSliceFinishedCached())
 					{
 						WorkResult = ETimeSliceWorkResult::CallAgainNextTimeSlice;
 						break;
 					}
 
-					AddGeneratedTileLayer(AddGenTilesTimeSlicedLayerIndex, TileGenerator, AddGenTilesTimeSlicedOldLayerTileIdMap, AddGenTilesTimeSlicedResultTileIndices);
+					AddGeneratedTileLayer(SyncTimeSlicedData.AddGenTilesLayerIndex, TileGenerator, SyncTimeSlicedData.OldLayerTileIdMapCached, SyncTimeSlicedData.ResultTileIndicesCached);
 
-					TimeSlicer.TestTimeSliceFinished();
+					SyncTimeSlicedData.TimeSlicer.TestTimeSliceFinished();
 				}
 			}
 		}
@@ -4693,12 +4707,12 @@ ETimeSliceWorkResult FRecastNavMeshGenerator::AddGeneratedTilesTimeSliced(FRecas
 	}
 	}
 
-	if (AddGenTilesTimeSlicedLayerIndex == TileGenerator.GetDirtyLayersMask().Num() || !bIteratedThroughDirtyLayers)
+	if (SyncTimeSlicedData.AddGenTilesLayerIndex == TileGenerator.GetDirtyLayersMask().Num() || !bIteratedThroughDirtyLayers)
 	{
-		AddGenTilesTimeSlicedLayerIndex = 0;
-		AddGeneratedTilesTimeSlicedState = EAddGeneratedTilesTimeSlicedState::Init;
+		SyncTimeSlicedData.AddGenTilesLayerIndex = 0;
+		SyncTimeSlicedData.AddGeneratedTilesState = EAddGeneratedTilesTimeSlicedState::Init;
 
-		OutResultTileIndices = MoveTemp(AddGenTilesTimeSlicedResultTileIndices);
+		OutResultTileIndices = MoveTemp(SyncTimeSlicedData.ResultTileIndicesCached);
 	}
 
 	return WorkResult;
@@ -4767,7 +4781,7 @@ bool FRecastNavMeshGenerator::HasDirtyTiles() const
 {
 	return (PendingDirtyTiles.Num() > 0 
 		|| RunningDirtyTiles.Num() > 0
-		|| TimeSlicedTileGeneratorSync.Get() != nullptr
+		|| SyncTimeSlicedData.TileGeneratorSync.Get() != nullptr
 		);
 }
 
@@ -5139,7 +5153,7 @@ TArray<uint32> FRecastNavMeshGenerator::ProcessTileTasksAsync(const int32 NumTas
 	// Release memory, list could be quite big after map load
 	if (NumProcessedTasks > 0 && PendingDirtyTiles.Num() == 0)
 	{
-		PendingDirtyTiles.Empty(32);
+		PendingDirtyTiles.Empty(64);
 	}
 	
 	// Collect completed tasks and apply generated data to navmesh
@@ -5195,146 +5209,217 @@ TSharedRef<FRecastTileGenerator> FRecastNavMeshGenerator::CreateTileGeneratorFro
 	return TileGenerator;
 }
 
-TArray<uint32> FRecastNavMeshGenerator::ProcessTileTasksSyncTimeSliced(const int32 NumTasksToProcess)
+TArray<uint32> FRecastNavMeshGenerator::ProcessTileTasksSyncTimeSliced()
 {
 	QUICK_SCOPE_CYCLE_COUNTER(STAT_RecastNavMeshGenerator_ProcessTileTasksSyncTimeSliced);
-
-	const bool bGameStaticNavMesh = IsGameStaticNavMesh(DestNavMesh);
-
-	TimeSlicer.ResetStartTime();
-
-	int32 NumProcessedTasks = 0;
+	CSV_SCOPED_TIMING_STAT(NAVREGEN, ProcessTileTasksSyncTimeSliced);
 
 	TArray<uint32> UpdatedTiles;
+	const UWorld* World = GetWorld();
+	double TimeStartProcessingTileThisFrame = 0.;
 
-	// Submit pending tile elements
-	while ((PendingDirtyTiles.Num() > 0
-		&& NumProcessedTasks < NumTasksToProcess)
-		|| TimeSlicedTileGeneratorSync.IsValid()
-		)
+	auto HasWorkToDo = [this]()
 	{
-		QUICK_SCOPE_CYCLE_COUNTER(STAT_RecastNavMeshGenerator_ProcessTileTasks_NewTasks);
+		return (PendingDirtyTiles.Num() > 0) || SyncTimeSlicedData.TileGeneratorSync.IsValid();
+	};
 
-		FIntPoint TileLocation;
 
-		if (TimeSlicedTileGeneratorSync.IsValid())
+	auto EndFunction = [&, this](bool bCalcTileRegenDuration) {
+		// Release memory, list could be quite big after map load
+		if (PendingDirtyTiles.Num() == 0)
 		{
-			TileLocation.X = TimeSlicedTileGeneratorSync->GetTileX();
-			TileLocation.Y = TimeSlicedTileGeneratorSync->GetTileY();
+			PendingDirtyTiles.Empty(64);
 		}
-		else
-		{
-			//if the next time slice regen state is false, we want to go to non time sliced tile regen so break here and switch
-			//next frame (as we've finished time slice processing the last tile)
-			if (!bNextTimeSliceRegenActive)
-			{
-				return UpdatedTiles;
-			}
 
-			TimeSlicedTileGeneratorSync = CreateTileGeneratorFromPendingElement(TileLocation);
+		if (World)
+		{
+			SyncTimeSlicedData.RealTimeSecsLastCall = World->GetRealTimeSeconds();
+		}
 
-			if (TimeSlicer.TestTimeSliceFinished())
-			{
-				return UpdatedTiles;
-			}
+		//this will only be true when we haven't finished generating this tile but are ending
+		//the function and need to record the TileRegenDuration so far for the tile
+		//being currently processed
+		if (bCalcTileRegenDuration)
+		{
+			SyncTimeSlicedData.CurrentTileRegenDuration += (FPlatformTime::Seconds() - TimeStartProcessingTileThisFrame);
 		}
 
-		FRecastTileGenerator& TileGeneratorRef = *TimeSlicedTileGeneratorSync;
+		return UpdatedTiles;
+	};
+
+	//Calculate the time slice duration
+	//Calc the MovingWindowDeltaTimeAverage this accounts for all scenarios we could be tile regening including unbounded frame rates or dropping frames as well as keeping
+	//calculation to an average which is fairly local temporally
+	if (World && SyncTimeSlicedData.RealTimeSecsLastCall >= 0.f)
+	{
+		const float DeltaTime = World->GetRealTimeSeconds() - SyncTimeSlicedData.RealTimeSecsLastCall;
+		SyncTimeSlicedData.MovingWindowDeltaTime.PushValue(DeltaTime);
+	}	
+	
+	//only calculate the time slice and process tiles if we have work to do
+	if (HasWorkToDo())
+	{
+		CSV_SCOPED_TIMING_STAT(NAVREGEN, ProcessTileTasksSyncTimeSlicedDoWork);
+
+		const bool bGameStaticNavMesh = IsGameStaticNavMesh(DestNavMesh);
+
+		SyncTimeSlicedData.TimeSlicer.StartTimeSlice();
+
+		const float DeltaTimesAverage = (SyncTimeSlicedData.MovingWindowDeltaTime.GetAverage() > 0.f) ? SyncTimeSlicedData.MovingWindowDeltaTime.GetAverage() : (1.f / 30.f); //use default 33 ms
+
+		const double TileRegenTimesAverage = (SyncTimeSlicedData.MovingWindowTileRegenTime.GetAverage() > 0.) ? SyncTimeSlicedData.MovingWindowTileRegenTime.GetAverage() : 0.0025; //use default of 2.5 milli secs to regen a full tile
+
+		//calculate the max desired frames to regen all the tiles in PendingDirtyTiles
+		const float MaxDesiredFramesToRegen = FMath::FloorToFloat(SyncTimeSlicedData.MaxDesiredTileRegenDuration / DeltaTimesAverage);
+
+		//tiles to add to PendingDirtyTiles if the current tile is taking longer than average to regen
+		//we add 1 tile for however many times longer the current tile is taking compared with the moving window average
+		const int32 TilesToAddForLongCurrentTileRegen = (SyncTimeSlicedData.CurrentTileRegenDuration > 0.) ? (static_cast<int32>(SyncTimeSlicedData.CurrentTileRegenDuration/ TileRegenTimesAverage)) : 0;
 
-		if (ProcessTileTasksSyncTimeSlicedState == EProcessTileTasksSyncTimeSlicedState::Init)
+		const int32 TotalTilesToRegen = PendingDirtyTiles.Num() + (SyncTimeSlicedData.TileGeneratorSync.IsValid() ? 1 : 0);
+
+		//calculate the total processing time to regen all the tiles based on the moving window average
+		const double TotalRegenTime = TileRegenTimesAverage * static_cast<double>(TotalTilesToRegen + TilesToAddForLongCurrentTileRegen);
+
+		//calculate the time slice per frame required to regen all the tiles clamped between MinTimeSliceDuration and MaxTimeSliceDuration
+		const double NextRegenTimeSliceTime = FMath::Clamp(TotalRegenTime / static_cast<double>(MaxDesiredFramesToRegen), SyncTimeSlicedData.MinTimeSliceDuration, SyncTimeSlicedData.MaxTimeSliceDuration);
+		SyncTimeSlicedData.TimeSlicer.SetTimeSliceDuration(NextRegenTimeSliceTime);
+
+#if !UE_BUILD_SHIPPING
+		CSV_CUSTOM_STAT(NAVREGEN, NavTileRegenTimeSliceTime, static_cast<float>(NextRegenTimeSliceTime), ECsvCustomStatOp::Set);
+		CSV_CUSTOM_STAT(NAVREGEN, NavTileRegenQueueLength, TotalTilesToRegen, ECsvCustomStatOp::Set);
+		CSV_CUSTOM_STAT(NAVREGEN, TilesToAddForLongCurrentTileRegen, TilesToAddForLongCurrentTileRegen, ECsvCustomStatOp::Set);
+		CSV_CUSTOM_STAT(NAVREGEN, NavTileAvRegenTime, static_cast<float>(SyncTimeSlicedData.MovingWindowTileRegenTime.GetAverage()), ECsvCustomStatOp::Set);
+		CSV_CUSTOM_STAT(NAVREGEN, NavTileAvRegenDeltaTime, static_cast<float>(SyncTimeSlicedData.MovingWindowDeltaTime.GetAverage()), ECsvCustomStatOp::Set);
+#endif
+
+		// Submit pending tile elements
+		do
 		{
-			if (TileGeneratorRef.HasDataToBuild())
-			{
-				ProcessTileTasksSyncTimeSlicedState = EProcessTileTasksSyncTimeSlicedState::DoWork;
-			}
-			else
-			{
-				ProcessTileTasksSyncTimeSlicedState = EProcessTileTasksSyncTimeSlicedState::Finish;
+			QUICK_SCOPE_CYCLE_COUNTER(STAT_RecastNavMeshGenerator_ProcessTileTasks_NewTasks);
 
-				if (!bGameStaticNavMesh)
+			FIntPoint TileLocation;
+			TimeStartProcessingTileThisFrame = FPlatformTime::Seconds();
+
+			if (SyncTimeSlicedData.ProcessTileTasksSyncState == EProcessTileTasksSyncTimeSlicedState::Init)
+			{
+				//if the next time slice regen state is false, we want to go to non time sliced tile regen so break here and switch
+				//next frame (as we've finished time slice processing the last tile)
+				if (!SyncTimeSlicedData.bNextTimeSliceRegenActive)
 				{
-					RemoveLayers(TileLocation, UpdatedTiles);
+					return EndFunction(false);
 				}
-			}
-		}
 
+				SyncTimeSlicedData.TileGeneratorSync = CreateTileGeneratorFromPendingElement(TileLocation);
 
-		switch (ProcessTileTasksSyncTimeSlicedState)
-		{
-		case EProcessTileTasksSyncTimeSlicedState::Init:
-		{
-			//do nothing 
-			ensureMsgf(false, TEXT("This State should not be used here!"));
-		}
-		break;
+				SyncTimeSlicedData.CurrentTileRegenDuration = 0.;
 
-		case EProcessTileTasksSyncTimeSlicedState::DoWork:
-		{
-			const ETimeSliceWorkResult WorkResult = TileGeneratorRef.DoWorkTimeSliced();
+				if (SyncTimeSlicedData.TileGeneratorSync->HasDataToBuild())
+				{
+					SyncTimeSlicedData.ProcessTileTasksSyncState = EProcessTileTasksSyncTimeSlicedState::DoWork;
+				}
+				else
+				{
+					SyncTimeSlicedData.ProcessTileTasksSyncState = EProcessTileTasksSyncTimeSlicedState::Finish;
 
-			if (WorkResult != ETimeSliceWorkResult::CallAgainNextTimeSlice)
+					if (!bGameStaticNavMesh)
+					{
+						RemoveLayers(TileLocation, UpdatedTiles);
+					}
+				}
+
+				if (SyncTimeSlicedData.TimeSlicer.TestTimeSliceFinished())
+				{
+					return EndFunction(true);
+				}
+			}
+			else
 			{
-				ProcessTileTasksSyncTimeSlicedState = EProcessTileTasksSyncTimeSlicedState::AddGeneratedTiles;
+				TileLocation.X = SyncTimeSlicedData.TileGeneratorSync->GetTileX();
+				TileLocation.Y = SyncTimeSlicedData.TileGeneratorSync->GetTileY();
 			}
 
-			if (TimeSlicer.IsTimeSliceFinishedCached())
+			FRecastTileGenerator& TileGeneratorRef = *SyncTimeSlicedData.TileGeneratorSync;
+
+			switch (SyncTimeSlicedData.ProcessTileTasksSyncState)
 			{
-				return UpdatedTiles;
+			case EProcessTileTasksSyncTimeSlicedState::Init:
+			{
+				//do nothing 
+				ensureMsgf(false, TEXT("This State should not be used here!"));
 			}
-		}//fall through to next state
-		case EProcessTileTasksSyncTimeSlicedState::AddGeneratedTiles:
-		{
-			const ETimeSliceWorkResult WorkResult = AddGeneratedTilesTimeSliced(TileGeneratorRef, ProcessTileTasksSyncTimeSlicedUpdatedTilesCache);
+			break;
 
-			if (WorkResult != ETimeSliceWorkResult::CallAgainNextTimeSlice)
+			case EProcessTileTasksSyncTimeSlicedState::DoWork:
 			{
-				ProcessTileTasksSyncTimeSlicedState = EProcessTileTasksSyncTimeSlicedState::StoreCompessedTileCacheLayers;
-			}
+				const ETimeSliceWorkResult WorkResult = TileGeneratorRef.DoWorkTimeSliced();
 
-			if (TimeSlicer.IsTimeSliceFinishedCached())
+				if (WorkResult != ETimeSliceWorkResult::CallAgainNextTimeSlice)
+				{
+					SyncTimeSlicedData.ProcessTileTasksSyncState = EProcessTileTasksSyncTimeSlicedState::AddGeneratedTiles;
+				}
+
+				if (SyncTimeSlicedData.TimeSlicer.IsTimeSliceFinishedCached())
+				{
+					return EndFunction(true);
+				}
+			}//fall through to next state
+			case EProcessTileTasksSyncTimeSlicedState::AddGeneratedTiles:
 			{
-				return UpdatedTiles;
-			}
-		}//fall through to next state
-		case EProcessTileTasksSyncTimeSlicedState::StoreCompessedTileCacheLayers:
-		{
-			StoreCompressedTileCacheLayers(TileGeneratorRef, TileLocation.X, TileLocation.Y);
+				const ETimeSliceWorkResult WorkResult = AddGeneratedTilesTimeSliced(TileGeneratorRef, SyncTimeSlicedData.UpdatedTilesCache);
 
-			//no need to check time slicing as not much work done
-			ProcessTileTasksSyncTimeSlicedState = EProcessTileTasksSyncTimeSlicedState::AppendUpdateTiles;
-		}//fall through to next state
-		case EProcessTileTasksSyncTimeSlicedState::AppendUpdateTiles: //this state was added purely to separate the functionality and allow the code to be more easily changed in future.
-		{
-			UpdatedTiles.Append(ProcessTileTasksSyncTimeSlicedUpdatedTilesCache);
-			ProcessTileTasksSyncTimeSlicedUpdatedTilesCache.Empty();
+				if (WorkResult != ETimeSliceWorkResult::CallAgainNextTimeSlice)
+				{
+					SyncTimeSlicedData.ProcessTileTasksSyncState = EProcessTileTasksSyncTimeSlicedState::StoreCompessedTileCacheLayers;
+				}
 
-			//no need to check time slicing as not much work done
-			ProcessTileTasksSyncTimeSlicedState = EProcessTileTasksSyncTimeSlicedState::Finish;
-		}//fall through to next state
-		case EProcessTileTasksSyncTimeSlicedState::Finish:
-		{
-			//no need to check time slicing as not much work done
-			//reset state to Init for next tile to be processed
-			ProcessTileTasksSyncTimeSlicedState = EProcessTileTasksSyncTimeSlicedState::Init;
-			TimeSlicedTileGeneratorSync.Reset();
+				if (SyncTimeSlicedData.TimeSlicer.IsTimeSliceFinishedCached())
+				{
+					return EndFunction(true);
+				}
+			}//fall through to next state
+			case EProcessTileTasksSyncTimeSlicedState::StoreCompessedTileCacheLayers:
+			{
+				StoreCompressedTileCacheLayers(TileGeneratorRef, TileLocation.X, TileLocation.Y);
 
-			NumProcessedTasks++;
-		}
-		break;
-		default:
-		{
-			ensureMsgf(false, TEXT("unhandled EProcessTileTasksSyncTimeSlicedState"));
-		}
-		}
-	}
+				//no need to check time slicing as not much work done
+				SyncTimeSlicedData.ProcessTileTasksSyncState = EProcessTileTasksSyncTimeSlicedState::AppendUpdateTiles;
+			}//fall through to next state
+			case EProcessTileTasksSyncTimeSlicedState::AppendUpdateTiles: //this state was added purely to separate the functionality and allow the code to be more easily changed in future.
+			{
+				UpdatedTiles.Append(SyncTimeSlicedData.UpdatedTilesCache);
+				SyncTimeSlicedData.UpdatedTilesCache.Empty();
 
-	// Release memory, list could be quite big after map load
-	if (PendingDirtyTiles.Num() == 0)
-	{
-		PendingDirtyTiles.Empty(32);
+				//no need to check time slicing as not much work done
+				SyncTimeSlicedData.ProcessTileTasksSyncState = EProcessTileTasksSyncTimeSlicedState::Finish;
+			}//fall through to next state
+			case EProcessTileTasksSyncTimeSlicedState::Finish:
+			{
+				//no need to check time slicing as not much work done
+				//reset state to Init for next tile to be processed
+				SyncTimeSlicedData.ProcessTileTasksSyncState = EProcessTileTasksSyncTimeSlicedState::Init;
+				SyncTimeSlicedData.TileGeneratorSync.Reset();
+
+				SyncTimeSlicedData.CurrentTileRegenDuration += (FPlatformTime::Seconds() - TimeStartProcessingTileThisFrame);
+
+				SyncTimeSlicedData.MovingWindowTileRegenTime.PushValue(SyncTimeSlicedData.CurrentTileRegenDuration);
+
+				SyncTimeSlicedData.CurrentTileRegenDuration = 0.;
+			}
+			break;
+			default:
+			{
+				ensureMsgf(false, TEXT("unhandled EProcessTileTasksSyncTimeSlicedState"));
+			}
+			}
+		}
+		while (HasWorkToDo());
 	}
 
-	return UpdatedTiles;
+	// we only hit this if we have processed too many tiles in a frame and we will already
+	// have calculated the tile regen duration, or if we have processed no tiles and we also
+	// don't need to calcualte the tile regen duration
+	return EndFunction(false);
 }
 
 //this code path is approx 10% faster than ProcessTileTasksSyncTimeSliced, however it spikes far worse for most use cases.
@@ -5376,7 +5461,7 @@ TArray<uint32> FRecastNavMeshGenerator::ProcessTileTasksSync(const int32 NumTask
 	// Release memory, list could be quite big after map load
 	if (PendingDirtyTiles.Num() == 0)
 	{
-		PendingDirtyTiles.Empty(32);
+		PendingDirtyTiles.Empty(64);
 	}
 
 	return UpdatedTiles;
@@ -5394,21 +5479,21 @@ TArray<uint32> FRecastNavMeshGenerator::ProcessTileTasks(const int32 NumTasksToP
 	UpdatedTiles = ProcessTileTasksAsync(NumTasksToProcess);
 #else
 	//only switch bTimeSliceRegen state if we are not time slicing or if we are but aren't part way through time slicing a tile
-	if (bTimeSliceRegenActive != bNextTimeSliceRegenActive)
+	if (SyncTimeSlicedData.bTimeSliceRegenActive != SyncTimeSlicedData.bNextTimeSliceRegenActive)
 	{
-		if (!bTimeSliceRegenActive)
+		if (!SyncTimeSlicedData.bTimeSliceRegenActive)
 		{
-			bTimeSliceRegenActive = bNextTimeSliceRegenActive;
+			SyncTimeSlicedData.bTimeSliceRegenActive = SyncTimeSlicedData.bNextTimeSliceRegenActive;
 		}
-		else if (!TimeSlicedTileGeneratorSync.IsValid())//test if we have finished processing a tile
+		else if (!SyncTimeSlicedData.TileGeneratorSync.IsValid())//test if we have finished processing a tile
 		{
-			bTimeSliceRegenActive = bNextTimeSliceRegenActive;
+			SyncTimeSlicedData.bTimeSliceRegenActive = SyncTimeSlicedData.bNextTimeSliceRegenActive;
 		}
 	}
 
-	if (bTimeSliceRegenActive)
+	if (SyncTimeSlicedData.bTimeSliceRegenActive)
 	{
-		UpdatedTiles = ProcessTileTasksSyncTimeSliced(NumTasksToProcess);
+		UpdatedTiles = ProcessTileTasksSyncTimeSliced();
 	}
 	else
 	{
@@ -5558,20 +5643,20 @@ bool FRecastNavMeshGenerator::IsBuildInProgress(bool bCheckDirtyToo) const
 {
 	return RunningDirtyTiles.Num()
 		|| (bCheckDirtyToo && PendingDirtyTiles.Num())
-		|| TimeSlicedTileGeneratorSync.Get();
+		|| SyncTimeSlicedData.TileGeneratorSync.Get();
 }
 
 int32 FRecastNavMeshGenerator::GetNumRemaningBuildTasks() const
 {
 	return RunningDirtyTiles.Num() 
 		+ PendingDirtyTiles.Num()
-		+ (TimeSlicedTileGeneratorSync.Get() ? 1 : 0);
+		+ (SyncTimeSlicedData.TileGeneratorSync.Get() ? 1 : 0);
 }
 
 int32 FRecastNavMeshGenerator::GetNumRunningBuildTasks() const
 {
 	return RunningDirtyTiles.Num()
-		+ (TimeSlicedTileGeneratorSync.Get() ? 1 : 0);
+		+ (SyncTimeSlicedData.TileGeneratorSync.Get() ? 1 : 0);
 }
 
 bool FRecastNavMeshGenerator::GatherGeometryOnGameThread() const 
@@ -5605,9 +5690,9 @@ uint32 FRecastNavMeshGenerator::LogMemUsed() const
 	for (const FRunningTileElement& Element : RunningDirtyTiles)
 	{
 		GeneratorsMem += Element.AsyncTask->GetTask().TileGenerator->UsedMemoryOnStartup;
-		if (TimeSlicedTileGeneratorSync.IsValid())
+		if (SyncTimeSlicedData.TileGeneratorSync.IsValid())
 		{
-			GeneratorsMem += TimeSlicedTileGeneratorSync->UsedMemoryOnStartup;
+			GeneratorsMem += SyncTimeSlicedData.TileGeneratorSync->UsedMemoryOnStartup;
 		}
 	}
 
diff --git a/Engine/Source/Runtime/NavigationSystem/Public/NavMesh/RecastNavMeshGenerator.h b/Engine/Source/Runtime/NavigationSystem/Public/NavMesh/RecastNavMeshGenerator.h
--- a/Engine/Source/Runtime/NavigationSystem/Public/NavMesh/RecastNavMeshGenerator.h
+++ b/Engine/Source/Runtime/NavigationSystem/Public/NavMesh/RecastNavMeshGenerator.h
@@ -13,6 +13,7 @@
 #include "AI/NavDataGenerator.h"
 #include "NavMesh/RecastHelpers.h"
 #include "NavDebugTypes.h"
+#include "Math/MovingWindowAverageFast.h"
 
 #if WITH_RECAST
 
@@ -209,7 +210,8 @@ public:
 	}
 
 	void SetTimeSliceDuration(double SliceDuration);
-	void ResetStartTime();
+	void StartTimeSlice();
+	double GetStartTime() const;
 	bool TestTimeSliceFinished() const;
 	bool IsTimeSliceFinishedCached() const;
 
@@ -592,8 +594,7 @@ public:
 
 private:
 	/** Prevent copying. */
-	FRecastNavMeshGenerator(FRecastNavMeshGenerator const& NoCopy) 
-		: bTimeSliceRegenActive(false) { check(0); };
+	FRecastNavMeshGenerator(FRecastNavMeshGenerator const& NoCopy) { check(0); };
 	FRecastNavMeshGenerator& operator=(FRecastNavMeshGenerator const& NoCopy) { check(0); return *this; }
 
 public:
@@ -665,9 +666,9 @@ public:
 	static void GetDebugGeometry(const FNavigationRelevantData& EncodedData, FNavDebugMeshData& DebugMeshData);
 #endif  // !UE_BUILD_SHIPPING
 
-	FTimeSlicer& GetTimeSlicer() { return TimeSlicer; }
+	FTimeSlicer& GetTimeSlicer() { return SyncTimeSlicedData.TimeSlicer; }
 	
-	void SetNextTimeSliceRegenActive(bool bRegenState) { bNextTimeSliceRegenActive = bRegenState; }
+	void SetNextTimeSliceRegenActive(bool bRegenState) { SyncTimeSlicedData.bNextTimeSliceRegenActive = bRegenState; }
 
 protected:
 	// Performs initial setup of member variables so that generator is ready to
@@ -707,7 +708,7 @@ protected:
 #else
 	TSharedRef<FRecastTileGenerator> CreateTileGeneratorFromPendingElement(FIntPoint &OutTileLocation);
 	/** Processes pending tile generation tasks Sync with option for time slicing currently an experimental feature. */
-	TArray<uint32> ProcessTileTasksSyncTimeSliced(const int32 NumTasksToProcess);
+	TArray<uint32> ProcessTileTasksSyncTimeSliced();
 	TArray<uint32> ProcessTileTasksSync(const int32 NumTasksToProcess);
 #endif
 	/** Processes pending tile generation tasks */
@@ -766,8 +767,12 @@ protected:
 	FVector BBoxGrowth;
 	
 	int32 NumActiveTiles;
-	/** the limit to number of asynchronous tile generators running at one time */
+	/** the limit to number of asynchronous tile generators running at one time,
+	 *  this is also used by the sync non time sliced regeneration ProcessTileTasksSync,
+	 *  but not by ProcessTileTasksSyncTimeSliced.
+	 */
 	int32 MaxTileGeneratorTasks;
+
 	float AvgLayersPerTile;
 
 	/** Total bounding box that includes all volumes, in unreal units. */
@@ -808,24 +813,55 @@ protected:
 	 *	like when navmesh size changes */
 	uint32 Version;
 
-	/** Start time sliced variables */
-	/** Do not null or Reset this directly instead call ResetTimeSlicedTileGeneratorSync(). The only exception currently is in ProcessTileTasksSyncTimeSliced */
-	TSharedPtr<FRecastTileGenerator> TimeSlicedTileGeneratorSync;
-	FTimeSlicer TimeSlicer;
-	TArray<uint32> ProcessTileTasksSyncTimeSlicedUpdatedTilesCache;
-	/** if we are currently using time sliced regen or not - currently an experimental feature.
-	 *  do not manipulate this value directly instead call SetNextTimeSliceRegenState() 
-	 */
-	bool bTimeSliceRegenActive;
-	bool bNextTimeSliceRegenActive;
+	/** Grouping all the member variables used by the time slicing code together for neatness */
+	struct FSyncTimeSlicedData
+	{
+		FSyncTimeSlicedData();
+
+		double CurrentTileRegenDuration;
+		double MinTimeSliceDuration;
+		double MaxTimeSliceDuration;
+		float RealTimeSecsLastCall; //no need to reset this
+		/** The max real world desired time to Regen all the tiles in PendingDirtyTiles,
+		 *  Note it could take longer than this, as the time slice is clamped per frame between
+		 *	MinTimeSliceDuration and MaxTimeSliceDuration.
+		 */
+		float MaxDesiredTileRegenDuration;
+
+		/** if we are currently using time sliced regen or not - currently an experimental feature.
+		 *  do not manipulate this value directly instead call SetNextTimeSliceRegenState()
+		 */
+		bool bTimeSliceRegenActive;
+		bool bNextTimeSliceRegenActive;
+		
+		/** Used by ProcessTileTasksSyncTimeSliced */
+		EProcessTileTasksSyncTimeSlicedState ProcessTileTasksSyncState;
+
+		/** Used by ProcessTileTasksSyncTimeSliced */
+		TArray<uint32> UpdatedTilesCache;
 
-	EProcessTileTasksSyncTimeSlicedState ProcessTileTasksSyncTimeSlicedState;
+		/** Used by AddGeneratedTilesTimeSliced */
+		TArray<uint32> ResultTileIndicesCached;
+
+		/** Used by AddGeneratedTilesTimeSliced */
+		TMap<int32, dtPolyRef> OldLayerTileIdMapCached;
+
+		/** Used by AddGeneratedTilesTimeSliced */
+		EAddGeneratedTilesTimeSlicedState AddGeneratedTilesState;
+
+		/** Used by AddGeneratedTilesTimeSliced */
+		int32 AddGenTilesLayerIndex;
+
+		/** Used to calculate the moving window average of the actual time spent inside functions used to regenerate a a tile, this is processing time not actual time over multiple frames */
+		FMovingWindowAverageFast<double, 128> MovingWindowTileRegenTime;
+		/** Used to calculate the actual moving window delta time */
+		FMovingWindowAverageFast<float, 128> MovingWindowDeltaTime;
+		/** Do not null or Reset this directly instead call ResetTimeSlicedTileGeneratorSync(). The only exception currently is in ProcessTileTasksSyncTimeSliced */
+		TSharedPtr<FRecastTileGenerator> TileGeneratorSync;
+		FTimeSlicer TimeSlicer;
+	};
 
-	TMap<int32, dtPolyRef> AddGenTilesTimeSlicedOldLayerTileIdMap;
-	TArray<uint32> AddGenTilesTimeSlicedResultTileIndices;
-	EAddGeneratedTilesTimeSlicedState AddGeneratedTilesTimeSlicedState;
-	int32 AddGenTilesTimeSlicedLayerIndex;
-	/** End time sliced variables */
+	FSyncTimeSlicedData SyncTimeSlicedData;
 };
 
 #endif // WITH_RECAST
