Build fix for IOS Client

diff --git a/Engine/Source/Runtime/NavigationSystem/Private/NavMesh/RecastNavMesh.cpp b/Engine/Source/Runtime/NavigationSystem/Private/NavMesh/RecastNavMesh.cpp
--- a/Engine/Source/Runtime/NavigationSystem/Private/NavMesh/RecastNavMesh.cpp
+++ b/Engine/Source/Runtime/NavigationSystem/Private/NavMesh/RecastNavMesh.cpp
@@ -2440,7 +2440,7 @@ void ARecastNavMesh::ConditionalConstructGenerator()
 		FRecastNavMeshGenerator* Generator = CreateGeneratorInstance();
 		if (Generator)
 		{
-			NavDataGenerator = MakeShareable(Generator);
+			NavDataGenerator = MakeShareable((FNavDataGenerator*)Generator);
 			Generator->Init();
 		}
 
diff --git a/Engine/Source/Runtime/NavigationSystem/Private/NavMesh/RecastNavMeshGenerator.cpp b/Engine/Source/Runtime/NavigationSystem/Private/NavMesh/RecastNavMeshGenerator.cpp
--- a/Engine/Source/Runtime/NavigationSystem/Private/NavMesh/RecastNavMeshGenerator.cpp
+++ b/Engine/Source/Runtime/NavigationSystem/Private/NavMesh/RecastNavMeshGenerator.cpp
@@ -1583,7 +1583,7 @@ FRecastTileGenerator::FRecastTileGenerator(FRecastNavMeshGenerator& ParentGenera
 	Version = ParentGenerator.GetVersion();
 	AdditionalCachedData = ParentGenerator.GetAdditionalCachedData();
 
-	ParentGeneratorWeakPtr = ParentGenerator.AsShared();
+	ParentGeneratorWeakPtr = ((FNavDataGenerator&)ParentGenerator).AsShared();
 
 	RasterizeGeomRecastState = ERasterizeGeomRecastTimeSlicedState::MarkWalkableTriangles;
 	RasterizeGeomState = ERasterizeGeomTimeSlicedState::RasterizeGeometryTransformCoords;
