NavigationGeomExec: command ExportNavigation now exports all navigation data, not only the default one.

#rnx

diff --git a/Engine/Source/Runtime/NavigationSystem/Private/NavMesh/RecastNavMeshGenerator.cpp b/Engine/Source/Runtime/NavigationSystem/Private/NavMesh/RecastNavMeshGenerator.cpp
--- a/Engine/Source/Runtime/NavigationSystem/Private/NavMesh/RecastNavMeshGenerator.cpp
+++ b/Engine/Source/Runtime/NavigationSystem/Private/NavMesh/RecastNavMeshGenerator.cpp
@@ -6089,43 +6089,40 @@ public:
 	/** Console commands, see embeded usage statement **/
 	virtual bool Exec( UWorld* InWorld, const TCHAR* Cmd, FOutputDevice& Ar ) override
 	{
+		bool bExported = false;
 #if ALLOW_DEBUG_FILES && !(UE_BUILD_SHIPPING || UE_BUILD_TEST)
-		bool bCorrectCmd = FParse::Command(&Cmd, TEXT("ExportNavigation"));
-		if (bCorrectCmd && !InWorld)
+		if (FParse::Command(&Cmd, TEXT("ExportNavigation")))
 		{
-			UE_LOG(LogNavigation, Error, TEXT("Failed to export navigation data due to missing UWorld"));
-		}
-		else if (InWorld && bCorrectCmd)
-		{
-			UNavigationSystemV1* NavSys = FNavigationSystem::GetCurrent<UNavigationSystemV1>(InWorld);
-			if (NavSys)
+			if (InWorld == nullptr)
+			{
+				UE_LOG(LogNavigation, Error, TEXT("Failed to export navigation data due to missing UWorld"));
+			}
+			else 
 			{
-				const ANavigationData* NavData = NavSys->GetDefaultNavDataInstance();
-				if (NavData)
+				UNavigationSystemV1* NavSys = FNavigationSystem::GetCurrent<UNavigationSystemV1>(InWorld);
+				if (NavSys)
 				{
-					if (const FNavDataGenerator* Generator = NavData->GetGenerator())
+					for (ANavigationData* NavData : NavSys->NavDataSet)
 					{
-						const FString Name = NavData->GetName();
-						Generator->ExportNavigationData( FString::Printf( TEXT("%s/%s"), *FPaths::ProjectSavedDir(), *Name ));
-						return true;
-					}
-					else
-					{
-						UE_LOG(LogNavigation, Error, TEXT("Failed to export navigation data due to missing generator"));
+						if (const FNavDataGenerator* Generator = NavData->GetGenerator())
+						{
+							Generator->ExportNavigationData(FString::Printf(TEXT("%s/%s"), *FPaths::ProjectSavedDir(), *NavData->GetName()));
+							bExported = true;
+						}
+						else
+						{
+							UE_LOG(LogNavigation, Error, TEXT("Failed to export navigation data %s due to missing generator"), *NavData->GetName());
+						}
 					}
 				}
 				else
 				{
-					UE_LOG(LogNavigation, Error, TEXT("Failed to export navigation data due to navigation data"));
+					UE_LOG(LogNavigation, Error, TEXT("Failed to export navigation data due to missing navigation system"));
 				}
 			}
-			else
-			{
-				UE_LOG(LogNavigation, Error, TEXT("Failed to export navigation data due to missing navigation system"));
-			}
 		}
-#endif // ALLOW_DEBUG_FILES && WITH_EDITOR
-		return false;
+#endif // ALLOW_DEBUG_FILES && !(UE_BUILD_SHIPPING || UE_BUILD_TEST)
+		return bExported;
 	}
 } NavigationGeomExec;
 
