NavigationSystemV1: GetNavDataForActor now returns any INavigationDataInterface found, not only classes in ARecastNavMesh class hierachy. It is up to the caller to cast and test for a specific sublass. (e.g. ANavigationData)

diff --git a/Engine/Source/Runtime/NavigationSystem/Private/NavigationSystem.cpp b/Engine/Source/Runtime/NavigationSystem/Private/NavigationSystem.cpp
--- a/Engine/Source/Runtime/NavigationSystem/Private/NavigationSystem.cpp
+++ b/Engine/Source/Runtime/NavigationSystem/Private/NavigationSystem.cpp
@@ -4226,8 +4226,7 @@ INavigationDataInterface* UNavigationSystemV1::GetNavDataForActor(const AActor&
 		NavData = NavSys->GetDefaultNavDataInstance(FNavigationSystem::DontCreate);
 	}
 
-	//  Only RecastNavMesh supported
-	return (INavigationDataInterface*)(Cast<ARecastNavMesh>(NavData));
+	return NavData;
 }
 
 int UNavigationSystemV1::GetNavigationBoundsForNavData(const ANavigationData& NavData, TArray<FBox>& OutBounds, ULevel* InLevel) const
