Chaos - Add NavMesh Physics Export Support for Convex and Trimesh

diff --git a/Engine/Source/Runtime/Engine/Public/AI/NavigationSystemHelpers.h b/Engine/Source/Runtime/Engine/Public/AI/NavigationSystemHelpers.h
--- a/Engine/Source/Runtime/Engine/Public/AI/NavigationSystemHelpers.h
+++ b/Engine/Source/Runtime/Engine/Public/AI/NavigationSystemHelpers.h
@@ -5,6 +5,7 @@
 #include "CoreMinimal.h"
 #include "EngineDefines.h"
 #include "AI/Navigation/NavigationTypes.h"
+#include "Chaos/Core.h"
 
 class AActor;
 class UBodySetup;
@@ -28,6 +29,9 @@ namespace Chaos
 {
 	template<typename T>
 	class THeightField;
+
+	template<typename T>
+	class TTriangleMeshImplicitObject;
 }
 #endif
 
@@ -41,6 +45,8 @@ struct FNavigableGeometryExport
 	virtual void ExportPxHeightField(physx::PxHeightField const * const HeightField, const FTransform& LocalToWorld) = 0;
 #endif // WITH_PHYSX
 #if WITH_CHAOS
+	virtual void ExportChaosTriMesh(const Chaos::FTriangleMeshImplicitObject* const TriMesh, const FTransform& LocalToWorld) = 0;
+	virtual void ExportChaosConvexMesh(const FKConvexElem* const Convex, const FTransform& LocalToWorld) = 0;
 	virtual void ExportChaosHeightField(const Chaos::THeightField<float>* const Heightfield, const FTransform& LocalToWorld) = 0;
 #endif
 	virtual void ExportHeightFieldSlice(const FNavHeightfieldSamples& PrefetchedHeightfieldSamples, const int32 NumRows, const int32 NumCols, const FTransform& LocalToWorld, const FBox& SliceBox) = 0;
diff --git a/Engine/Source/Runtime/Experimental/Chaos/Private/Chaos/TriangleMeshImplicitObject.cpp b/Engine/Source/Runtime/Experimental/Chaos/Private/Chaos/TriangleMeshImplicitObject.cpp
--- a/Engine/Source/Runtime/Experimental/Chaos/Private/Chaos/TriangleMeshImplicitObject.cpp
+++ b/Engine/Source/Runtime/Experimental/Chaos/Private/Chaos/TriangleMeshImplicitObject.cpp
@@ -700,6 +700,16 @@ uint16 FTriangleMeshImplicitObject::GetMaterialIndex(uint32 HintIndex) const
 	return 0;
 }
 
+const TParticles<FReal, 3>& FTriangleMeshImplicitObject::Particles() const
+{
+	return MParticles;
+}
+
+const TArray<TVector<int32, 3>> FTriangleMeshImplicitObject::Elements() const
+{
+	return MElements;
+}
+
 void Chaos::FTriangleMeshImplicitObject::RebuildBV()
 {
 	const int32 NumTris = MElements.Num();
diff --git a/Engine/Source/Runtime/Experimental/Chaos/Public/Chaos/TriangleMeshImplicitObject.h b/Engine/Source/Runtime/Experimental/Chaos/Public/Chaos/TriangleMeshImplicitObject.h
--- a/Engine/Source/Runtime/Experimental/Chaos/Public/Chaos/TriangleMeshImplicitObject.h
+++ b/Engine/Source/Runtime/Experimental/Chaos/Public/Chaos/TriangleMeshImplicitObject.h
@@ -115,6 +115,9 @@ namespace Chaos
 
 		virtual uint16 GetMaterialIndex(uint32 HintIndex) const override;
 
+		const TParticles<FReal, 3>& Particles() const;
+		const TArray<TVector<int32, 3>> Elements() const;
+
 	private:
 
 		void RebuildBV();
diff --git a/Engine/Source/Runtime/NavigationSystem/Private/NavMesh/RecastNavMeshGenerator.cpp b/Engine/Source/Runtime/NavigationSystem/Private/NavMesh/RecastNavMeshGenerator.cpp
--- a/Engine/Source/Runtime/NavigationSystem/Private/NavMesh/RecastNavMeshGenerator.cpp
+++ b/Engine/Source/Runtime/NavigationSystem/Private/NavMesh/RecastNavMeshGenerator.cpp
@@ -18,6 +18,7 @@
 #endif
 #if WITH_CHAOS
 	#include "Chaos/HeightField.h"
+	#include "Chaos/TriangleMeshImplicitObject.h"
 #endif
 #include "NavMesh/PImplRecastNavMesh.h"
 
@@ -209,6 +210,8 @@ struct FRecastGeometryExport : public FNavigableGeometryExport
 	virtual void ExportPxHeightField(physx::PxHeightField const * const HeightField, const FTransform& LocalToWorld) override;
 #endif // WITH_PHYSX
 #if WITH_CHAOS
+	virtual void ExportChaosTriMesh(const Chaos::FTriangleMeshImplicitObject* const TriMesh, const FTransform& LocalToWorld) override;
+	virtual void ExportChaosConvexMesh(const FKConvexElem* const Convex, const FTransform& LocalToWorld) override;
 	virtual void ExportChaosHeightField(const Chaos::THeightField<float>* const Heightfield, const FTransform& LocalToWorld) override;
 #endif
 	virtual void ExportHeightFieldSlice(const FNavHeightfieldSamples& PrefetchedHeightfieldSamples, const int32 NumRows, const int32 NumCols, const FTransform& LocalToWorld, const FBox& SliceBox) override;
@@ -535,6 +538,122 @@ void ExportPxHeightField(PxHeightField const * const HeightField, const FTransfo
 #endif // WITH_PHYSX
 
 #if WITH_CHAOS
+void ExportChaosTriMesh(const Chaos::FTriangleMeshImplicitObject* const TriMesh, const FTransform& LocalToWorld
+	, TNavStatArray<float>& VertexBuffer, TNavStatArray<int32>& IndexBuffer
+	, FBox& UnrealBounds)
+{
+	if (TriMesh == nullptr)
+	{
+		return;
+	}
+
+	using namespace Chaos;
+
+	int32 VertOffset = VertexBuffer.Num() / 3;
+	int32 NumTris = TriMesh->Elements().Num();
+	const TParticles<FReal, 3>& Vertices = TriMesh->Particles();
+	const TArray<TVector<int32, 3>>& Triangles = TriMesh->Elements();
+
+	VertexBuffer.Reserve(VertexBuffer.Num() + NumTris * 9);
+	IndexBuffer.Reserve(IndexBuffer.Num() + NumTris * 3);
+
+	const bool bFlipCullMode = (LocalToWorld.GetDeterminant() < 0.f);
+
+	const int32 IndexOrder[3] = { bFlipCullMode ? 0 : 2, 1, bFlipCullMode ? 2 : 0 };
+
+#if SHOW_NAV_EXPORT_PREVIEW
+	UWorld* DebugWorld = FindEditorWorld();
+#endif // SHOW_NAV_EXPORT_PREVIEW
+
+	for (int32 TriIdx = 0; TriIdx < NumTris; ++TriIdx)
+	{
+		for (int32 i = 0; i < 3; i++)
+		{
+			const FVector UnrealCoords = LocalToWorld.TransformPosition(Vertices.X(Triangles[TriIdx][i]));
+			UnrealBounds += UnrealCoords;
+
+			VertexBuffer.Add(UnrealCoords.X);
+			VertexBuffer.Add(UnrealCoords.Y);
+			VertexBuffer.Add(UnrealCoords.Z);
+		}
+
+		IndexBuffer.Add(VertOffset + IndexOrder[0]);
+		IndexBuffer.Add(VertOffset + IndexOrder[1]);
+		IndexBuffer.Add(VertOffset + IndexOrder[2]);
+
+#if SHOW_NAV_EXPORT_PREVIEW
+		if (DebugWorld)
+		{
+			FVector V0(VertexBuffer[(VertOffset + IndexOrder[0]) * 3 + 0], VertexBuffer[(VertOffset + IndexOrder[0]) * 3 + 1], VertexBuffer[(VertOffset + IndexOrder[0]) * 3 + 2]);
+			FVector V1(VertexBuffer[(VertOffset + IndexOrder[1]) * 3 + 0], VertexBuffer[(VertOffset + IndexOrder[1]) * 3 + 1], VertexBuffer[(VertOffset + IndexOrder[1]) * 3 + 2]);
+			FVector V2(VertexBuffer[(VertOffset + IndexOrder[2]) * 3 + 0], VertexBuffer[(VertOffset + IndexOrder[2]) * 3 + 1], VertexBuffer[(VertOffset + IndexOrder[2]) * 3 + 2]);
+
+			DrawDebugLine(DebugWorld, V0, V1, bFlipCullMode ? FColor::Red : FColor::Blue, true);
+			DrawDebugLine(DebugWorld, V1, V2, bFlipCullMode ? FColor::Red : FColor::Blue, true);
+			DrawDebugLine(DebugWorld, V2, V0, bFlipCullMode ? FColor::Red : FColor::Blue, true);
+		}
+#endif // SHOW_NAV_EXPORT_PREVIEW
+
+		VertOffset += 3;
+	}
+}
+
+
+
+void ExportChaosConvexMesh(const FKConvexElem* const Convex, const FTransform& LocalToWorld
+	, TNavStatArray<float>& VertexBuffer, TNavStatArray<int32>& IndexBuffer
+	, FBox& UnrealBounds)
+{
+	using namespace Chaos;
+
+	if (Convex == nullptr)
+	{
+		return;
+	}
+
+
+	QUICK_SCOPE_CYCLE_COUNTER(STAT_NavMesh_ExportChaosConvexMesh);
+
+	int32 VertOffset = VertexBuffer.Num() / 3;
+
+	VertexBuffer.Reserve(VertexBuffer.Num() + Convex->VertexData.Num() * 3);
+	IndexBuffer.Reserve(IndexBuffer.Num() + Convex->IndexData.Num());
+
+#if SHOW_NAV_EXPORT_PREVIEW
+	UWorld* DebugWorld = FindEditorWorld();
+#endif // SHOW_NAV_EXPORT_PREVIEW
+
+	for (const FVector& Vertex : Convex->VertexData)
+	{
+		const FVector UnrealCoord = LocalToWorld.TransformPosition(Vertex);
+		UnrealBounds += UnrealCoord;
+
+		VertexBuffer.Add(UnrealCoord.X);
+		VertexBuffer.Add(UnrealCoord.Y);
+		VertexBuffer.Add(UnrealCoord.Z);
+	}
+
+	for (int32 i = 0; i < Convex->IndexData.Num(); ++i)
+	{
+		IndexBuffer.Add(VertOffset + Convex->IndexData[i]);
+	}
+
+#if SHOW_NAV_EXPORT_PREVIEW
+	if (DebugWorld)
+	{
+		for (int32 Index = StartVertOffset; Index < VertexBuffer.Num(); Index += 3)
+		{
+			FVector V0(VertexBuffer[IndexBuffer[Index] * 3], VertexBuffer[IndexBuffer[Index] * 3 + 1], VertexBuffer[IndexBuffer[Index] * 3] + 2);
+			FVector V1(VertexBuffer[IndexBuffer[Index + 1] * 3], VertexBuffer[IndexBuffer[Index + 1] * 3 + 1], VertexBuffer[IndexBuffer[Index + 1] * 3] + 2);
+			FVector V2(VertexBuffer[IndexBuffer[Index + 2] * 3], VertexBuffer[IndexBuffer[Index + 2] * 3 + 1], VertexBuffer[IndexBuffer[Index + 2] * 3] + 2);
+
+			DrawDebugLine(DebugWorld, V0, V1, FColor::Blue, true);
+			DrawDebugLine(DebugWorld, V1, V2, FColor::Blue, true);
+			DrawDebugLine(DebugWorld, V2, V0, fColor::Blue, true);
+		}
+	#endif // SHOW_NAV_EXPORT_PREVIEW
+}
+
 void ExportChaosHeightField(const Chaos::THeightField<float>* const HeightField, const FTransform& LocalToWorld
 	, TNavStatArray<float>& VertexBuffer, TNavStatArray<int32>& IndexBuffer
 	, FBox& UnrealBounds)
@@ -780,19 +899,18 @@ FORCEINLINE_DEBUGGABLE void AddFacesToRecast(TArray<FVector, OtherAllocator>& In
 }
 
 FORCEINLINE_DEBUGGABLE void ExportRigidBodyConvexElements(UBodySetup& BodySetup, TNavStatArray<float>& VertexBuffer, TNavStatArray<int32>& IndexBuffer,
-														  TNavStatArray<int32>& ShapeBuffer, FBox& UnrealBounds, const FTransform& LocalToWorld)
+	TNavStatArray<int32>& ShapeBuffer, FBox& UnrealBounds, const FTransform& LocalToWorld)
 {
-#if WITH_PHYSX
 	const int32 ConvexCount = BodySetup.AggGeom.ConvexElems.Num();
-	FKConvexElem const * ConvexElem = BodySetup.AggGeom.ConvexElems.GetData();
-
+	FKConvexElem const* ConvexElem = BodySetup.AggGeom.ConvexElems.GetData();
 	const FTransform NegXScale(FQuat::Identity, FVector::ZeroVector, FVector(-1, 1, 1));
 
-	for(int32 i=0; i< ConvexCount; ++i, ++ConvexElem)
+	for (int32 i = 0; i < ConvexCount; ++i, ++ConvexElem)
 	{
 		// Store index of first vertex in shape buffer
 		ShapeBuffer.Add(VertexBuffer.Num() / 3);
 
+#if WITH_PHYSX && PHYSICS_INTERFACE_PHYSX
 		// Get verts/triangles from this hull.
 		if (!ConvexElem->GetConvexMesh() && ConvexElem->GetMirroredConvexMesh())
 		{
@@ -804,17 +922,23 @@ FORCEINLINE_DEBUGGABLE void ExportRigidBodyConvexElements(UBodySetup& BodySetup,
 			// Otherwise use the regular mesh in the case that both exist
 			ExportPxConvexMesh(ConvexElem->GetConvexMesh(), LocalToWorld, VertexBuffer, IndexBuffer, UnrealBounds);
 		}
+#elif WITH_CHAOS
+		if (ConvexElem->GetChaosConvexMesh())
+		{
+			// TODO use ConvexElem->GetTransform?() transform?
+			ExportChaosConvexMesh(ConvexElem, LocalToWorld, VertexBuffer, IndexBuffer, UnrealBounds);
+		}
+#endif
 	}
-#endif // WITH_PHYSX
 }
 
 FORCEINLINE_DEBUGGABLE void ExportRigidBodyTriMesh(UBodySetup& BodySetup, TNavStatArray<float>& VertexBuffer, TNavStatArray<int32>& IndexBuffer,
-												   FBox& UnrealBounds, const FTransform& LocalToWorld)
+	FBox& UnrealBounds, const FTransform& LocalToWorld)
 {
-#if WITH_PHYSX
 	if (BodySetup.GetCollisionTraceFlag() == CTF_UseComplexAsSimple)
 	{
-		for(PxTriangleMesh* TriMesh : BodySetup.TriMeshes)
+#if WITH_PHYSX && PHYSICS_INTERFACE_PHYSX
+		for (PxTriangleMesh* TriMesh : BodySetup.TriMeshes)
 		{
 			if (TriMesh->getTriangleMeshFlags() & PxTriangleMeshFlag::e16_BIT_INDICES)
 			{
@@ -825,8 +949,13 @@ FORCEINLINE_DEBUGGABLE void ExportRigidBodyTriMesh(UBodySetup& BodySetup, TNavSt
 				ExportPxTriMesh<PxU32>(TriMesh, LocalToWorld, VertexBuffer, IndexBuffer, UnrealBounds);
 			}
 		}
-	}
-#endif // WITH_PHYSX
+#elif WITH_CHAOS
+		for(const auto& TriMesh : BodySetup.ChaosTriMeshes)
+		{
+			ExportChaosTriMesh(TriMesh.Get(), LocalToWorld, VertexBuffer, IndexBuffer, UnrealBounds);
+		}
+#endif
+}
 }
 
 void ExportRigidBodyBoxElements(const FKAggregateGeom& AggGeom, TNavStatArray<float>& VertexBuffer, TNavStatArray<int32>& IndexBuffer,
@@ -1171,6 +1300,16 @@ void FRecastGeometryExport::ExportPxHeightField(physx::PxHeightField const * con
 #endif // WITH_PHYSX
 
 #if WITH_CHAOS
+void FRecastGeometryExport::ExportChaosTriMesh(const Chaos::FTriangleMeshImplicitObject* const TriMesh, const FTransform& LocalToWorld)
+{
+	RecastGeometryExport::ExportChaosTriMesh(TriMesh, LocalToWorld, VertexBuffer, IndexBuffer, Data->Bounds);
+}
+
+void FRecastGeometryExport::ExportChaosConvexMesh(const FKConvexElem* const Convex, const FTransform& LocalToWorld)
+{
+	RecastGeometryExport::ExportChaosConvexMesh(Convex, LocalToWorld, VertexBuffer, IndexBuffer, Data->Bounds);
+}
+
 void FRecastGeometryExport::ExportChaosHeightField(const Chaos::THeightField<float>* const Heightfield, const FTransform& LocalToWorld)
 {
 	RecastGeometryExport::ExportChaosHeightField(Heightfield, LocalToWorld, VertexBuffer, IndexBuffer, Data->Bounds);
