Addition of LogNavigationDataBuild to output in the logs a basic build summary

diff --git a/Engine/Source/Runtime/Engine/Classes/AI/NavigationSystemBase.h b/Engine/Source/Runtime/Engine/Classes/AI/NavigationSystemBase.h
--- a/Engine/Source/Runtime/Engine/Classes/AI/NavigationSystemBase.h
+++ b/Engine/Source/Runtime/Engine/Classes/AI/NavigationSystemBase.h
@@ -21,6 +21,7 @@ class AController;
 class UNavAreaBase;
 
 ENGINE_API DECLARE_LOG_CATEGORY_EXTERN(LogNavigation, Warning, All);
+ENGINE_API DECLARE_LOG_CATEGORY_EXTERN(LogNavigationDataBuild, Log, All);
 
 UENUM()
 enum class FNavigationSystemRunMode : uint8
diff --git a/Engine/Source/Runtime/Engine/Private/AI/NavigationSystemBase.cpp b/Engine/Source/Runtime/Engine/Private/AI/NavigationSystemBase.cpp
--- a/Engine/Source/Runtime/Engine/Private/AI/NavigationSystemBase.cpp
+++ b/Engine/Source/Runtime/Engine/Private/AI/NavigationSystemBase.cpp
@@ -12,6 +12,7 @@
 #include "AI/Navigation/NavigationDataChunk.h"
 
 DEFINE_LOG_CATEGORY(LogNavigation);
+DEFINE_LOG_CATEGORY(LogNavigationDataBuild);
 
 #if !UE_BUILD_SHIPPING
 #include "CoreGlobals.h"
diff --git a/Engine/Source/Runtime/NavigationSystem/Private/NavigationSystem.cpp b/Engine/Source/Runtime/NavigationSystem/Private/NavigationSystem.cpp
--- a/Engine/Source/Runtime/NavigationSystem/Private/NavigationSystem.cpp
+++ b/Engine/Source/Runtime/NavigationSystem/Private/NavigationSystem.cpp
@@ -2977,6 +2977,14 @@ void UNavigationSystemV1::GatherNavigationBounds()
 
 void UNavigationSystemV1::Build()
 {
+	UE_LOG(LogNavigationDataBuild, Log, TEXT("UNavigationSystemV1::Build started..."));
+#if WITH_PHYSX
+	UE_LOG(LogNavigationDataBuild, Log, TEXT("   Building navigation data using PHYSX."));
+#endif
+#if WITH_CHAOS
+	UE_LOG(LogNavigationDataBuild, Log, TEXT("   Building navigation data using CHAOS."));
+#endif
+
 	UWorld* World = GetWorld();
 	if (!World)
 	{
@@ -3031,7 +3039,8 @@ void UNavigationSystemV1::Build()
 	DefaultDirtyAreasController.bDirtyAreasReportedWhileAccumulationLocked = false;
 #endif // !UE_BUILD_SHIPPING
 
-	UE_LOG(LogNavigation, Display, TEXT("UNavigationSystemV1::Build total execution time: %.5f"), float(FPlatformTime::Seconds() - BuildStartTime));
+	UE_LOG(LogNavigationDataBuild, Log, TEXT("UNavigationSystemV1::Build total execution time: %.2fs"), float(FPlatformTime::Seconds() - BuildStartTime));
+	UE_LOG(LogNavigation, Display, TEXT("UNavigationSystemV1::Build total execution time: %.5fs"), float(FPlatformTime::Seconds() - BuildStartTime));
 }
 
 void UNavigationSystemV1::CancelBuild()
@@ -3272,6 +3281,8 @@ void UNavigationSystemV1::RebuildAll(bool bIsLoadTime)
 				
 		if (NavData && (!bIsLoadTime || NavData->NeedsRebuildOnLoad()) && (!bIsInGame || NavData->SupportsRuntimeGeneration()))
 		{
+			UE_LOG(LogNavigationDataBuild, Log, TEXT("   Building NavData:  %s."), *NavData->GetConfig().GetDescription());
+
 			NavData->RebuildAll();
 		}
 	}
