Fixing several navmesh issues caused by incorrect collision data in the DDC
Addition of a NAVCOLLISION_FORMAT for Chaos
Rotating the DDC string version for the FDerivedDataNavCollisionCooker
Adding new logs to LogNavigationDataBuild

diff --git a/Engine/Source/Runtime/NavigationSystem/Private/NavCollision.cpp b/Engine/Source/Runtime/NavigationSystem/Private/NavCollision.cpp
--- a/Engine/Source/Runtime/NavigationSystem/Private/NavCollision.cpp
+++ b/Engine/Source/Runtime/NavigationSystem/Private/NavCollision.cpp
@@ -36,7 +36,11 @@ namespace NavCollisionCookStats
 }
 #endif
 
+#if WITH_PHYSX && PHYSICS_INTERFACE_PHYSX
 static const FName NAVCOLLISION_FORMAT = TEXT("NavCollision_X");
+#elif WITH_CHAOS
+static const FName NAVCOLLISION_FORMAT = TEXT("NavCollision_Chaos");
+#endif
 
 class FNavCollisionDataReader
 {
@@ -90,7 +94,7 @@ public:
 
 	virtual const TCHAR* GetVersionString() const override
 	{
-		return TEXT("B89838347A4348138EE337A847529C5C");
+		return TEXT("7F071F64D317460BB568645D8521B30A");
 	}
 
 	virtual FString GetPluginSpecificCacheKeySuffix() const override
diff --git a/Engine/Source/Runtime/NavigationSystem/Private/NavigationSystem.cpp b/Engine/Source/Runtime/NavigationSystem/Private/NavigationSystem.cpp
--- a/Engine/Source/Runtime/NavigationSystem/Private/NavigationSystem.cpp
+++ b/Engine/Source/Runtime/NavigationSystem/Private/NavigationSystem.cpp
@@ -2978,11 +2978,14 @@ void UNavigationSystemV1::GatherNavigationBounds()
 void UNavigationSystemV1::Build()
 {
 	UE_LOG(LogNavigationDataBuild, Log, TEXT("UNavigationSystemV1::Build started..."));
+#if PHYSICS_INTERFACE_PHYSX
+	UE_LOG(LogNavigationDataBuild, Log, TEXT("   Building navigation data using PHYSICS_INTERFACE_PHYSX."));
+#endif
 #if WITH_PHYSX
-	UE_LOG(LogNavigationDataBuild, Log, TEXT("   Building navigation data using PHYSX."));
+	UE_LOG(LogNavigationDataBuild, Log, TEXT("   Building navigation data using WITH_PHYSX."));
 #endif
 #if WITH_CHAOS
-	UE_LOG(LogNavigationDataBuild, Log, TEXT("   Building navigation data using CHAOS."));
+	UE_LOG(LogNavigationDataBuild, Log, TEXT("   Building navigation data using WITH_CHAOS."));
 #endif
 
 	UWorld* World = GetWorld();
