LogNavigationDataBuild -> change verbosity from Log to Display

diff --git a/Engine/Source/Runtime/NavigationSystem/Private/NavigationSystem.cpp b/Engine/Source/Runtime/NavigationSystem/Private/NavigationSystem.cpp
--- a/Engine/Source/Runtime/NavigationSystem/Private/NavigationSystem.cpp
+++ b/Engine/Source/Runtime/NavigationSystem/Private/NavigationSystem.cpp
@@ -2977,15 +2977,15 @@ void UNavigationSystemV1::GatherNavigationBounds()
 
 void UNavigationSystemV1::Build()
 {
-	UE_LOG(LogNavigationDataBuild, Log, TEXT("UNavigationSystemV1::Build started..."));
+	UE_LOG(LogNavigationDataBuild, Display, TEXT("UNavigationSystemV1::Build started..."));
 #if PHYSICS_INTERFACE_PHYSX
-	UE_LOG(LogNavigationDataBuild, Log, TEXT("   Building navigation data using PHYSICS_INTERFACE_PHYSX."));
+	UE_LOG(LogNavigationDataBuild, Display, TEXT("   Building navigation data using PHYSICS_INTERFACE_PHYSX."));
 #endif
 #if WITH_PHYSX
-	UE_LOG(LogNavigationDataBuild, Log, TEXT("   Building navigation data using WITH_PHYSX."));
+	UE_LOG(LogNavigationDataBuild, Display, TEXT("   Building navigation data using WITH_PHYSX."));
 #endif
 #if WITH_CHAOS
-	UE_LOG(LogNavigationDataBuild, Log, TEXT("   Building navigation data using WITH_CHAOS."));
+	UE_LOG(LogNavigationDataBuild, Display, TEXT("   Building navigation data using WITH_CHAOS."));
 #endif
 
 	UWorld* World = GetWorld();
@@ -3042,7 +3042,7 @@ void UNavigationSystemV1::Build()
 	DefaultDirtyAreasController.bDirtyAreasReportedWhileAccumulationLocked = false;
 #endif // !UE_BUILD_SHIPPING
 
-	UE_LOG(LogNavigationDataBuild, Log, TEXT("UNavigationSystemV1::Build total execution time: %.2fs"), float(FPlatformTime::Seconds() - BuildStartTime));
+	UE_LOG(LogNavigationDataBuild, Display, TEXT("UNavigationSystemV1::Build total execution time: %.2fs"), float(FPlatformTime::Seconds() - BuildStartTime));
 	UE_LOG(LogNavigation, Display, TEXT("UNavigationSystemV1::Build total execution time: %.5fs"), float(FPlatformTime::Seconds() - BuildStartTime));
 }
 
@@ -3284,7 +3284,7 @@ void UNavigationSystemV1::RebuildAll(bool bIsLoadTime)
 				
 		if (NavData && (!bIsLoadTime || NavData->NeedsRebuildOnLoad()) && (!bIsInGame || NavData->SupportsRuntimeGeneration()))
 		{
-			UE_LOG(LogNavigationDataBuild, Log, TEXT("   Building NavData:  %s."), *NavData->GetConfig().GetDescription());
+			UE_LOG(LogNavigationDataBuild, Display, TEXT("   Building NavData:  %s."), *NavData->GetConfig().GetDescription());
 
 			NavData->RebuildAll();
 		}
