Fix Chaos convex mesh indices winding to have triangle normals pointing outward (for the nav collision)
Rotate navcollision DDC GUI
This fixes "Can't generate Navmesh in 4.25 builds"
Also fixes the lakes and river not generating water navmesh on their surface

diff --git a/Engine/Source/Runtime/NavigationSystem/Private/NavCollision.cpp b/Engine/Source/Runtime/NavigationSystem/Private/NavCollision.cpp
--- a/Engine/Source/Runtime/NavigationSystem/Private/NavCollision.cpp
+++ b/Engine/Source/Runtime/NavigationSystem/Private/NavCollision.cpp
@@ -94,7 +94,7 @@ public:
 
 	virtual const TCHAR* GetVersionString() const override
 	{
-		return TEXT("7F071F64D317460BB568645D8521B30A");
+		return TEXT("2EE6412E1E884D1AAA7F79FF760363BC");
 	}
 
 	virtual FString GetPluginSpecificCacheKeySuffix() const override
diff --git a/Engine/Source/Runtime/NavigationSystem/Private/NavMesh/RecastNavMeshGenerator.cpp b/Engine/Source/Runtime/NavigationSystem/Private/NavMesh/RecastNavMeshGenerator.cpp
--- a/Engine/Source/Runtime/NavigationSystem/Private/NavMesh/RecastNavMeshGenerator.cpp
+++ b/Engine/Source/Runtime/NavigationSystem/Private/NavMesh/RecastNavMeshGenerator.cpp
@@ -652,9 +652,11 @@ void ExportChaosConvexMesh(const FKConvexElem* const Convex, const FTransform& L
 		ensure(Convex->IndexData.Num());
 	}
 
-	for (int32 i = 0; i < Convex->IndexData.Num(); ++i)
+	for (int32 i = 0; i < Convex->IndexData.Num(); i += 3)
 	{
 		IndexBuffer.Add(VertOffset + Convex->IndexData[i]);
+		IndexBuffer.Add(VertOffset + Convex->IndexData[i + 2]);
+		IndexBuffer.Add(VertOffset + Convex->IndexData[i + 1]);
 	}
 
 #if SHOW_NAV_EXPORT_PREVIEW
