Addition of an ensure that the index data is a multiple of 3 in ExportChaosConvexMesh()

diff --git a/Engine/Source/Runtime/NavigationSystem/Private/NavMesh/RecastNavMeshGenerator.cpp b/Engine/Source/Runtime/NavigationSystem/Private/NavMesh/RecastNavMeshGenerator.cpp
--- a/Engine/Source/Runtime/NavigationSystem/Private/NavMesh/RecastNavMeshGenerator.cpp
+++ b/Engine/Source/Runtime/NavigationSystem/Private/NavMesh/RecastNavMeshGenerator.cpp
@@ -652,11 +652,14 @@ void ExportChaosConvexMesh(const FKConvexElem* const Convex, const FTransform& L
 		ensure(Convex->IndexData.Num());
 	}
 
-	for (int32 i = 0; i < Convex->IndexData.Num(); i += 3)
+	if (ensureMsgf(Convex->IndexData.Num() % 3 == 0, TEXT("Invalid index data in '%s'."), *Convex->GetName().ToString()))
 	{
-		IndexBuffer.Add(VertOffset + Convex->IndexData[i]);
-		IndexBuffer.Add(VertOffset + Convex->IndexData[i + 2]);
-		IndexBuffer.Add(VertOffset + Convex->IndexData[i + 1]);
+		for (int32 i = 0; i < Convex->IndexData.Num(); i += 3)
+		{
+			IndexBuffer.Add(VertOffset + Convex->IndexData[i]);
+			IndexBuffer.Add(VertOffset + Convex->IndexData[i + 2]);
+			IndexBuffer.Add(VertOffset + Convex->IndexData[i + 1]);
+		}
 	}
 
 #if SHOW_NAV_EXPORT_PREVIEW
