Fix leak in new timesliced navmesh generator code

#rnx

diff --git a/Engine/Source/Runtime/NavigationSystem/Private/NavMesh/RecastNavMeshGenerator.cpp b/Engine/Source/Runtime/NavigationSystem/Private/NavMesh/RecastNavMeshGenerator.cpp
--- a/Engine/Source/Runtime/NavigationSystem/Private/NavMesh/RecastNavMeshGenerator.cpp
+++ b/Engine/Source/Runtime/NavigationSystem/Private/NavMesh/RecastNavMeshGenerator.cpp
@@ -1850,9 +1850,9 @@ FRecastTileGenerator::FRecastTileGenerator(FRecastNavMeshGenerator& ParentGenera
 
 FRecastTileGenerator::~FRecastTileGenerator()
 {
-	GenNavDataTimeSlicedGenerationContext.Release();
-	GenNavDataTimeSlicedAllocator.Release();
-	GenCompressedlayersTimeSlicedRasterContext.Release();
+	GenNavDataTimeSlicedGenerationContext.Reset();
+	GenNavDataTimeSlicedAllocator.Reset();
+	GenCompressedlayersTimeSlicedRasterContext.Reset();
 }
 
 void FRecastTileGenerator::Setup(const FRecastNavMeshGenerator& ParentGenerator, const TArray<FBox>& DirtyAreas)
