Log NavData registration results
The motivation is to provide a minimum of feedback when failing to register navigation data (ex: if navmesh settings do not fit the project navmesh settings) .

diff --git a/Engine/Source/Runtime/NavigationSystem/Private/NavigationSystem.cpp b/Engine/Source/Runtime/NavigationSystem/Private/NavigationSystem.cpp
--- a/Engine/Source/Runtime/NavigationSystem/Private/NavigationSystem.cpp
+++ b/Engine/Source/Runtime/NavigationSystem/Private/NavigationSystem.cpp
@@ -783,6 +783,7 @@ void UNavigationSystemV1::OnWorldInitDone(FNavigationSystemRunMode Mode)
 					if (NavData != NULL)
 					{
 						const ERegistrationResult Result = RegisterNavData(NavData);
+						LogNavDataRegistrationResult(Result);
 
 						if (Result == RegistrationSuccessful)
 						{
@@ -1577,7 +1578,8 @@ ANavigationData* UNavigationSystemV1::GetDefaultNavDataInstance(FNavigationSyste
 #endif // WITH_RECAST
 		// either way make sure it's registered. Registration stores unique
 		// navmeshes, so we have nothing to lose
-		RegisterNavData(MainNavData);
+		const ERegistrationResult Result = RegisterNavData(MainNavData);
+		LogNavDataRegistrationResult(Result);
 	}
 
 	return MainNavData;
@@ -1804,6 +1806,7 @@ void UNavigationSystemV1::ProcessRegistrationCandidates()
 		if (OwningLevel && OwningLevel->bIsVisible)
 		{
 			const ERegistrationResult Result = RegisterNavData(NavDataPtr);
+			LogNavDataRegistrationResult(Result);
 
 			if (Result != RegistrationSuccessful && Result != RegistrationFailed_DataPendingKill)
 			{
@@ -3637,6 +3640,34 @@ ERuntimeGenerationType UNavigationSystemV1::GetRuntimeGenerationType() const
 	return RuntimeGenerationType;
 }
 
+void UNavigationSystemV1::LogNavDataRegistrationResult(ERegistrationResult InResult)
+{
+	switch (InResult)
+	{
+	case UNavigationSystemV1::RegistrationError:
+		UE_VLOG_UELOG(this, LogNavigation, Warning, TEXT("NavData RegistrationError, could not be registered."));
+		break;
+	case UNavigationSystemV1::RegistrationFailed_DataPendingKill:
+		UE_VLOG_UELOG(this, LogNavigation, Warning, TEXT("NavData RegistrationFailed_DataPendingKill."));
+		break;
+	case UNavigationSystemV1::RegistrationFailed_AgentAlreadySupported:
+		UE_VLOG_UELOG(this, LogNavigation, Warning, TEXT("NavData RegistrationFailed_AgentAlreadySupported, specified agent type already has its navmesh implemented."));
+		break;
+	case UNavigationSystemV1::RegistrationFailed_AgentNotValid:
+		UE_VLOG_UELOG(this, LogNavigation, Warning, TEXT("NavData RegistrationFailed_AgentNotValid, NavData instance contains navmesh that doesn't support any of expected agent types."));
+		break;
+	case UNavigationSystemV1::RegistrationFailed_NotSuitable:
+		UE_VLOG_UELOG(this, LogNavigation, Warning, TEXT("NavData RegistrationFailed_NotSuitable."));
+		break;
+	case UNavigationSystemV1::RegistrationSuccessful:
+		UE_VLOG_UELOG(this, LogNavigation, Verbose, TEXT("NavData RegistrationSuccessful."));
+		break;
+	default:
+		UE_VLOG_UELOG(this, LogNavigation, Warning, TEXT("Registration not successful default warning."));
+		break;
+	}
+}
+
 //----------------------------------------------------------------------//
 // Blueprint functions
 //----------------------------------------------------------------------//
diff --git a/Engine/Source/Runtime/NavigationSystem/Public/NavigationSystem.h b/Engine/Source/Runtime/NavigationSystem/Public/NavigationSystem.h
--- a/Engine/Source/Runtime/NavigationSystem/Public/NavigationSystem.h
+++ b/Engine/Source/Runtime/NavigationSystem/Public/NavigationSystem.h
@@ -1000,6 +1000,8 @@ protected:
 	 *  Depends on runtime generation settings of each navigation data, always ERuntimeGenerationType::Dynamic in the editor world
 	 */
 	ERuntimeGenerationType GetRuntimeGenerationType() const;
+
+	void LogNavDataRegistrationResult(ERegistrationResult);
 	
 	//----------------------------------------------------------------------//
 	// new stuff
