Fix warning: don't try to register a null NavData with RegisterNavData()

diff --git a/Engine/Source/Runtime/NavigationSystem/Private/NavigationSystem.cpp b/Engine/Source/Runtime/NavigationSystem/Private/NavigationSystem.cpp
--- a/Engine/Source/Runtime/NavigationSystem/Private/NavigationSystem.cpp
+++ b/Engine/Source/Runtime/NavigationSystem/Private/NavigationSystem.cpp
@@ -1578,8 +1578,11 @@ ANavigationData* UNavigationSystemV1::GetDefaultNavDataInstance(FNavigationSyste
 #endif // WITH_RECAST
 		// either way make sure it's registered. Registration stores unique
 		// navmeshes, so we have nothing to lose
-		const ERegistrationResult Result = RegisterNavData(MainNavData);
-		LogNavDataRegistrationResult(Result);
+		if (MainNavData != nullptr)
+		{
+			const ERegistrationResult Result = RegisterNavData(MainNavData);
+			LogNavDataRegistrationResult(Result);
+		}
 	}
 
 	return MainNavData;
