Fix non-determinism in cooked navmesh assets
When initializing a new world, reset navlinks id counter to make sure each map has its own set of unique ids (else, when cooking map by map, ids may change depending on order of cooking).

diff --git a/Engine/Source/Runtime/Engine/Classes/AI/NavigationSystemBase.h b/Engine/Source/Runtime/Engine/Classes/AI/NavigationSystemBase.h
--- a/Engine/Source/Runtime/Engine/Classes/AI/NavigationSystemBase.h
+++ b/Engine/Source/Runtime/Engine/Classes/AI/NavigationSystemBase.h
@@ -22,6 +22,7 @@ class UNavAreaBase;
 
 ENGINE_API DECLARE_LOG_CATEGORY_EXTERN(LogNavigation, Warning, All);
 ENGINE_API DECLARE_LOG_CATEGORY_EXTERN(LogNavigationDataBuild, Log, All);
+ENGINE_API DECLARE_LOG_CATEGORY_EXTERN(LogNavLink, Warning, All);
 
 UENUM()
 enum class FNavigationSystemRunMode : uint8
diff --git a/Engine/Source/Runtime/Engine/Private/AI/NavigationSystemBase.cpp b/Engine/Source/Runtime/Engine/Private/AI/NavigationSystemBase.cpp
--- a/Engine/Source/Runtime/Engine/Private/AI/NavigationSystemBase.cpp
+++ b/Engine/Source/Runtime/Engine/Private/AI/NavigationSystemBase.cpp
@@ -13,6 +13,7 @@
 
 DEFINE_LOG_CATEGORY(LogNavigation);
 DEFINE_LOG_CATEGORY(LogNavigationDataBuild);
+DEFINE_LOG_CATEGORY(LogNavLink);
 
 #if !UE_BUILD_SHIPPING
 #include "CoreGlobals.h"
diff --git a/Engine/Source/Runtime/NavigationSystem/Private/NavLinkCustomComponent.cpp b/Engine/Source/Runtime/NavigationSystem/Private/NavLinkCustomComponent.cpp
--- a/Engine/Source/Runtime/NavigationSystem/Private/NavLinkCustomComponent.cpp
+++ b/Engine/Source/Runtime/NavigationSystem/Private/NavLinkCustomComponent.cpp
@@ -81,6 +81,7 @@ void UNavLinkCustomComponent::PostEditImport()
 	Super::PostEditImport();
 
 	NavLinkUserId = INavLinkCustomInterface::GetUniqueId();
+	UE_LOG(LogNavLink, VeryVerbose, TEXT("%s new navlink id %u."), ANSI_TO_TCHAR(__FUNCTION__), NavLinkUserId);
 }
 #endif
 
@@ -169,6 +170,7 @@ void UNavLinkCustomComponent::OnRegister()
 	if (NavLinkUserId == 0)
 	{
 		NavLinkUserId = INavLinkCustomInterface::GetUniqueId();
+		UE_LOG(LogNavLink, VeryVerbose, TEXT("%s new navlink id %u [%s]."), ANSI_TO_TCHAR(__FUNCTION__), NavLinkUserId, *GetFullName());
 	}
 
 	UNavigationSystemV1::RequestCustomLinkRegistering(*this, this);
diff --git a/Engine/Source/Runtime/NavigationSystem/Private/NavigationSystem.cpp b/Engine/Source/Runtime/NavigationSystem/Private/NavigationSystem.cpp
--- a/Engine/Source/Runtime/NavigationSystem/Private/NavigationSystem.cpp
+++ b/Engine/Source/Runtime/NavigationSystem/Private/NavigationSystem.cpp
@@ -2018,6 +2018,7 @@ void UNavigationSystemV1::RegisterCustomLink(INavLinkCustomInterface& CustomLink
 	if (CustomLinksMap.Contains(LinkId))
 	{
 		LinkId = INavLinkCustomInterface::GetUniqueId();
+		UE_LOG(LogNavLink, VeryVerbose, TEXT("%s new navlink id %u."), ANSI_TO_TCHAR(__FUNCTION__), LinkId);
 		CustomLink.UpdateLinkId(LinkId);
 
 		UObject* CustomLinkOb = CustomLink.GetLinkOwner();
@@ -3585,6 +3586,7 @@ void UNavigationSystemV1::CleanUp(FNavigationSystem::ECleanupMode Mode)
 
 		if (MyWorld->WorldType == EWorldType::Game || MyWorld->WorldType == EWorldType::Editor)
 		{
+			UE_LOG(LogNavLink, VeryVerbose, TEXT("Reset navlink id on cleanup."));
 			INavLinkCustomInterface::NextUniqueId = 1;
 		}
 	}
diff --git a/Engine/Source/Runtime/NavigationSystem/Private/NavigationSystemModule.cpp b/Engine/Source/Runtime/NavigationSystem/Private/NavigationSystemModule.cpp
--- a/Engine/Source/Runtime/NavigationSystem/Private/NavigationSystemModule.cpp
+++ b/Engine/Source/Runtime/NavigationSystem/Private/NavigationSystemModule.cpp
@@ -5,6 +5,7 @@
 #include "AI/NavigationSystemBase.h"
 #include "Templates/SubclassOf.h"
 #include "Engine/World.h"
+#include "NavLinkCustomInterface.h"
 
 #define LOCTEXT_NAMESPACE "NavigationSystem"
 
@@ -18,21 +19,34 @@ class FNavigationSystemModule : public INavSysModule
 	
 	//virtual UNavigationSystemBase* CreateNavigationSystemInstance(UWorld& World) override;
 	// End IModuleInterface
+
+private:
+#if WITH_EDITOR 
+	static FDelegateHandle OnPreWorldInitializationHandle;
+#endif
 };
 
 IMPLEMENT_MODULE(FNavigationSystemModule, NavigationSystem)
 
+#if WITH_EDITOR 
+FDelegateHandle FNavigationSystemModule::OnPreWorldInitializationHandle;
+#endif
+
 void FNavigationSystemModule::StartupModule()
 { 
 	// mz@todo bind to all the delegates in FNavigationSystem
-
-	// also, to 
+#if WITH_EDITOR 
+	OnPreWorldInitializationHandle = FWorldDelegates::OnPreWorldInitialization.AddStatic(&INavLinkCustomInterface::OnPreWorldInitialization);
+#endif
 }
 
 void FNavigationSystemModule::ShutdownModule()
 {
 	// This function may be called during shutdown to clean up your module.  For modules that support dynamic reloading,
 	// we call this function before unloading the module.
+#if WITH_EDITOR
+	FWorldDelegates::OnPreWorldInitialization.Remove(OnPreWorldInitializationHandle);
+#endif
 }
 
 //UNavigationSystemBase* FNavigationSystemModule::CreateNavigationSystemInstance(UWorld& World)
diff --git a/Engine/Source/Runtime/NavigationSystem/Private/NavigationSystemTypes.cpp b/Engine/Source/Runtime/NavigationSystem/Private/NavigationSystemTypes.cpp
--- a/Engine/Source/Runtime/NavigationSystem/Private/NavigationSystemTypes.cpp
+++ b/Engine/Source/Runtime/NavigationSystem/Private/NavigationSystemTypes.cpp
@@ -296,11 +296,13 @@ UObject* INavLinkCustomInterface::GetLinkOwner() const
 
 uint32 INavLinkCustomInterface::GetUniqueId()
 {
+	UE_LOG(LogNavLink, VeryVerbose, TEXT("%s id: %u."), ANSI_TO_TCHAR(__FUNCTION__), NextUniqueId);
 	return NextUniqueId++;
 }
 
 void INavLinkCustomInterface::UpdateUniqueId(uint32 AlreadyUsedId)
 {
+	UE_CLOG(AlreadyUsedId + 1 > NextUniqueId, LogNavLink, VeryVerbose, TEXT("%s, updating NextUniqueId to: %u."), ANSI_TO_TCHAR(__FUNCTION__), AlreadyUsedId + 1)
 	NextUniqueId = FMath::Max(NextUniqueId, AlreadyUsedId + 1);
 }
 
@@ -317,3 +319,14 @@ FNavigationLink INavLinkCustomInterface::GetModifier(const INavLinkCustomInterfa
 
 	return LinkMod;
 }
+
+void INavLinkCustomInterface::OnPreWorldInitialization(UWorld* World, const UWorld::InitializationValues IVS)
+{
+	ResetUniqueId();
+}
+
+void INavLinkCustomInterface::ResetUniqueId()
+{
+	UE_LOG(LogNavLink, VeryVerbose, TEXT("Reset navlink id."));
+	NextUniqueId = 1;
+}
diff --git a/Engine/Source/Runtime/NavigationSystem/Public/NavLinkCustomInterface.h b/Engine/Source/Runtime/NavigationSystem/Public/NavLinkCustomInterface.h
--- a/Engine/Source/Runtime/NavigationSystem/Public/NavLinkCustomInterface.h
+++ b/Engine/Source/Runtime/NavigationSystem/Public/NavLinkCustomInterface.h
@@ -76,5 +76,8 @@ class NAVIGATIONSYSTEM_API INavLinkCustomInterface
 	/** Helper function: create modifier for navigation data export */
 	static FNavigationLink GetModifier(const INavLinkCustomInterface* CustomNavLink);
 
+	static void OnPreWorldInitialization(UWorld* World, const UWorld::InitializationValues IVS);
+	static void ResetUniqueId();
+
 	static uint32 NextUniqueId;
 };
