Build fix, missing world.h include

diff --git a/Engine/Source/Runtime/NavigationSystem/Public/NavLinkCustomInterface.h b/Engine/Source/Runtime/NavigationSystem/Public/NavLinkCustomInterface.h
--- a/Engine/Source/Runtime/NavigationSystem/Public/NavLinkCustomInterface.h
+++ b/Engine/Source/Runtime/NavigationSystem/Public/NavLinkCustomInterface.h
@@ -8,6 +8,7 @@
 #include "UObject/Interface.h"
 #include "NavAreas/NavArea.h"
 #include "AI/Navigation/NavLinkDefinition.h"
+#include "Engine/World.h"
 #include "NavLinkCustomInterface.generated.h"
 
 /** 
