Fixed nav ensure firing and failing builds when convexes don't generate index buffers.

diff --git a/Engine/Source/Runtime/NavigationSystem/Private/NavMesh/RecastNavMeshGenerator.cpp b/Engine/Source/Runtime/NavigationSystem/Private/NavMesh/RecastNavMeshGenerator.cpp
--- a/Engine/Source/Runtime/NavigationSystem/Private/NavMesh/RecastNavMeshGenerator.cpp
+++ b/Engine/Source/Runtime/NavigationSystem/Private/NavMesh/RecastNavMeshGenerator.cpp
@@ -637,6 +637,21 @@ void ExportChaosConvexMesh(const FKConvexElem* const Convex, const FTransform& L
 	UWorld* DebugWorld = FindEditorWorld();
 #endif // SHOW_NAV_EXPORT_PREVIEW
 
+	if (Convex->VertexData.Num())
+	{
+		if(Convex->IndexData.Num() == 0)
+		{
+			UE_LOG(LogNavigation, Verbose, TEXT("Zero indices in convex."));
+			return;
+		}
+
+		if(Convex->IndexData.Num() % 3 != 0)
+		{
+			UE_LOG(LogNavigation, Verbose, TEXT("Invalid indices in convex."));
+			return;
+		}
+	}
+
 	for (const FVector& Vertex : Convex->VertexData)
 	{
 		const FVector UnrealCoord = LocalToWorld.TransformPosition(Vertex);
@@ -647,12 +662,7 @@ void ExportChaosConvexMesh(const FKConvexElem* const Convex, const FTransform& L
 		VertexBuffer.Add(UnrealCoord.Z);
 	}
 
-	if (Convex->VertexData.Num())
-	{
-		ensure(Convex->IndexData.Num());
-	}
-
-	if (ensureMsgf(Convex->IndexData.Num() % 3 == 0, TEXT("Invalid index data in '%s'."), *Convex->GetName().ToString()))
+	if (Convex->IndexData.Num() % 3 == 0)
 	{
 		for (int32 i = 0; i < Convex->IndexData.Num(); i += 3)
 		{
