Fixed a really dumb mistake done in the original implementation of ARecastNavMesh::ReplaceAreaInTileBounds

diff --git a/Engine/Source/Runtime/NavigationSystem/Private/NavMesh/RecastNavMesh.cpp b/Engine/Source/Runtime/NavigationSystem/Private/NavMesh/RecastNavMesh.cpp
--- a/Engine/Source/Runtime/NavigationSystem/Private/NavMesh/RecastNavMesh.cpp
+++ b/Engine/Source/Runtime/NavigationSystem/Private/NavMesh/RecastNavMesh.cpp
@@ -1507,7 +1507,7 @@ int32 ARecastNavMesh::ReplaceAreaInTileBounds(const FBox& Bounds, TSubclassOf<UN
 
 		const int32 OldAreaID = GetAreaID(OldArea);
 		ensure(OldAreaID != INDEX_NONE);
-		const int32 NewAreaID = 34;// GetAreaID(NewArea);
+		const int32 NewAreaID = GetAreaID(NewArea);
 		ensure(NewAreaID != INDEX_NONE);
 		ensure(NewAreaID != OldAreaID);
 
