PR #6437: Improve visibility of Nav Mesh smart link state (Contributed by [REDACTED_PERSON])

diff --git a/Engine/Source/Runtime/NavigationSystem/Private/NavMesh/NavMeshRenderingComponent.cpp b/Engine/Source/Runtime/NavigationSystem/Private/NavMesh/NavMeshRenderingComponent.cpp
--- a/Engine/Source/Runtime/NavigationSystem/Private/NavMesh/NavMeshRenderingComponent.cpp
+++ b/Engine/Source/Runtime/NavigationSystem/Private/NavMesh/NavMeshRenderingComponent.cpp
@@ -159,6 +159,16 @@ namespace FNavMeshRenderingHelpers
 		return FColor(((Col >> 1) & 0x007f7f7f) | (Col & 0xff000000));
 	}
 
+	FColor SemiDarkenColor(const FColor& Base)
+	{
+		const uint32 Col = Base.DWColor();
+		// 1/2 can be too dark(DarkenColor) - use 3/4 here.
+		const uint32 QuarterCol = ((Col >> 2) & 0x003f3f3f);
+		const uint32 HalfCol = ((Col >> 1) & 0x007f7f7f);
+		const uint32 ThreeQuarterCol = QuarterCol + HalfCol;
+		return FColor(ThreeQuarterCol | (Col & 0xff000000));
+	}
+
 	void AddVertex(FNavMeshSceneProxyData::FDebugMeshData& MeshData, const FVector& Pos, const FColor Color)
 	{
 		const int32 VertexIndex = MeshData.Vertices.Num();
@@ -498,7 +508,7 @@ void FNavMeshSceneProxyData::GatherData(const ARecastNavMesh* NavMesh, int32 InN
 				{
 					const FVector V0 = Link.Left + NavMeshDrawOffset;
 					const FVector V1 = Link.Right + NavMeshDrawOffset;
-					const FColor LinkColor = ((Link.Direction && Link.ValidEnds) || (Link.ValidEnds & FRecastDebugGeometry::OMLE_Left)) ? FNavMeshRenderingHelpers::DarkenColor(NavMeshColors[Link.AreaID]) : NavMeshRenderColor_OffMeshConnectionInvalid;
+					const FColor LinkColor = ((Link.Direction && Link.ValidEnds) || (Link.ValidEnds & FRecastDebugGeometry::OMLE_Left)) ? FNavMeshRenderingHelpers::SemiDarkenColor(NavMeshColors[Link.AreaID]) : NavMeshRenderColor_OffMeshConnectionInvalid;
 
 					FNavMeshRenderingHelpers::CacheLink(NavLinkLines, V0, V1, LinkColor, Link.Direction);
 
diff --git a/Engine/Source/Runtime/NavigationSystem/Private/NavMesh/RecastNavMesh.cpp b/Engine/Source/Runtime/NavigationSystem/Private/NavMesh/RecastNavMesh.cpp
--- a/Engine/Source/Runtime/NavigationSystem/Private/NavMesh/RecastNavMesh.cpp
+++ b/Engine/Source/Runtime/NavigationSystem/Private/NavMesh/RecastNavMesh.cpp
@@ -1411,6 +1411,10 @@ void ARecastNavMesh::UpdateCustomLink(const INavLinkCustomInterface* CustomLink)
 
 		RecastNavMeshImpl->UpdateNavigationLinkArea(UserId, AreaId, PolyFlags);
 		RecastNavMeshImpl->UpdateSegmentLinkArea(UserId, AreaId, PolyFlags);
+
+#if !UE_BUILD_SHIPPING
+		RequestDrawingUpdate(false);
+#endif
 	}
 }
 
