PR #6464: Fix bug in AI navigation system stopping DefaultMaxSearchNodes from working (Contributed by [REDACTED_PERSON])

diff --git a/Engine/Source/Runtime/NavigationSystem/Private/NavFilters/NavigationQueryFilter.cpp b/Engine/Source/Runtime/NavigationSystem/Private/NavFilters/NavigationQueryFilter.cpp
--- a/Engine/Source/Runtime/NavigationSystem/Private/NavFilters/NavigationQueryFilter.cpp
+++ b/Engine/Source/Runtime/NavigationSystem/Private/NavFilters/NavigationQueryFilter.cpp
@@ -38,12 +38,12 @@ FSharedConstNavQueryFilter UNavigationQueryFilter::GetQueryFilter(const ANavigat
 	FSharedConstNavQueryFilter SharedFilter = bInstantiateForQuerier ? nullptr : NavData.GetQueryFilter(GetClass());
 	if (!SharedFilter.IsValid())
 	{
-		FNavigationQueryFilter* NavFilter = new FNavigationQueryFilter();
-		NavFilter->SetFilterImplementation(NavData.GetDefaultQueryFilterImpl());
+		// Clone the default filter so we get search nodes and other fields
+		FSharedNavQueryFilter NavFilter = NavData.GetDefaultQueryFilter()->GetCopy();
 
-		InitializeFilter(NavData, Querier, *NavFilter);
+		InitializeFilter(NavData, Querier, *NavFilter.Get());
 
-		SharedFilter = MakeShareable(NavFilter);
+		SharedFilter = NavFilter;
 		if (!bInstantiateForQuerier)
 		{
 			const_cast<ANavigationData&>(NavData).StoreQueryFilter(GetClass(), SharedFilter);
