Fixed a bug where the NavigationSystem would rebuild Static and DynamicModifiersOnlt navdata in PIE and game if bAutoCreateNavigationData was set to false.

diff --git a/Engine/Source/Runtime/NavigationSystem/Private/NavigationSystem.cpp b/Engine/Source/Runtime/NavigationSystem/Private/NavigationSystem.cpp
--- a/Engine/Source/Runtime/NavigationSystem/Private/NavigationSystem.cpp
+++ b/Engine/Source/Runtime/NavigationSystem/Private/NavigationSystem.cpp
@@ -772,6 +772,8 @@ void UNavigationSystemV1::OnWorldInitDone(FNavigationSystemRunMode Mode)
 
 			if (GetDefaultNavDataInstance(FNavigationSystem::DontCreate) != NULL)
 			{
+				const bool bIsInGame = World->IsGameWorld();
+
 				// trigger navmesh update
 				for (TActorIterator<ANavigationData> It(World); It; ++It)
 				{
@@ -782,7 +784,7 @@ void UNavigationSystemV1::OnWorldInitDone(FNavigationSystemRunMode Mode)
 
 						if (Result == RegistrationSuccessful)
 						{
-							if (bAllowRebuild)
+							if (bAllowRebuild && (!bIsInGame || NavData->SupportsRuntimeGeneration()))
 							{
 								NavData->RebuildAll();
 							}
