Extended nav path debug drawing via NavTestRenderingComponent to draw arrow heads to indicate path direction

diff --git a/Engine/Source/Runtime/NavigationSystem/Private/NavMesh/NavTestRenderingComponent.cpp b/Engine/Source/Runtime/NavigationSystem/Private/NavMesh/NavTestRenderingComponent.cpp
--- a/Engine/Source/Runtime/NavigationSystem/Private/NavMesh/NavTestRenderingComponent.cpp
+++ b/Engine/Source/Runtime/NavigationSystem/Private/NavMesh/NavTestRenderingComponent.cpp
@@ -93,7 +93,8 @@ void FNavTestSceneProxy::GetDynamicMeshElements(const TArray<const FSceneView*>&
 			{
 				for (int32 PointIndex = 1; PointIndex < PathPoints.Num(); PointIndex++)
 				{
-					PDI->DrawLine(PathPoints[PointIndex-1], PathPoints[PointIndex], FLinearColor::Red, SDPG_World, 2.0f, 0.0f, true);
+					PDI->DrawLine(PathPoints[PointIndex-1], PathPoints[PointIndex], FLinearColor::Red, SDPG_World, 2.0f, 0.0f);
+					DrawArrowHead(PDI, PathPoints[PointIndex], PathPoints[PointIndex - 1], 25.f, FLinearColor::Red, SDPG_World, 2.0f);
 				}
 			}
 
