Experimental navmesh debug tools (first pass)
Allow to display internal debug data at different navmesh generation steps for a selected tile
Ignore height filtering while displaying the navmesh around camera position

diff --git a/Engine/Source/Runtime/NavigationSystem/Private/NavMesh/NavMeshRenderingComponent.cpp b/Engine/Source/Runtime/NavigationSystem/Private/NavMesh/NavMeshRenderingComponent.cpp
--- a/Engine/Source/Runtime/NavigationSystem/Private/NavMesh/NavMeshRenderingComponent.cpp
+++ b/Engine/Source/Runtime/NavigationSystem/Private/NavMesh/NavMeshRenderingComponent.cpp
@@ -1,7 +1,7 @@
 // Copyright Epic Games, Inc. All Rights Reserved.
 
 /*=============================================================================
-	NavMeshRenderingComponent.cpp: A component that renders a nav mesh.
+	NavMeshRenderingComponent.cpp: A component that renders a navmesh.
  =============================================================================*/
 
 #include "NavMesh/NavMeshRenderingComponent.h"
@@ -18,6 +18,10 @@
 #include "SceneManagement.h"
 #include "TimerManager.h"
 
+#if RECAST_INTERNAL_DEBUG_DATA
+#include "NavMesh/RecastInternalDebugData.h"
+#endif
+
 #if WITH_EDITOR
 #include "Editor.h"
 #include "EditorViewportClient.h"
@@ -41,10 +45,10 @@ static const float ClusterLinkLines_LineThickness = 2.0f;
 
 namespace FNavMeshRenderingHelpers
 {
-	bool LineInView(const FVector& Start, const FVector& End, const FSceneView* View, bool bUseDistanceCheck)
+	bool LineInView(const FVector& Start, const FVector& End, const FSceneView* View)
 	{
-		if (FVector::DistSquared(Start, View->ViewMatrices.GetViewOrigin()) > ARecastNavMesh::GetDrawDistanceSq() ||
-			FVector::DistSquared(End, View->ViewMatrices.GetViewOrigin()) > ARecastNavMesh::GetDrawDistanceSq())
+		if (FVector::DistSquaredXY(Start, View->ViewMatrices.GetViewOrigin()) > ARecastNavMesh::GetDrawDistanceSq() ||
+			FVector::DistSquaredXY(End,   View->ViewMatrices.GetViewOrigin()) > ARecastNavMesh::GetDrawDistanceSq())
 		{
 			return false;
 		}
@@ -64,8 +68,8 @@ namespace FNavMeshRenderingHelpers
 	bool LineInCorrectDistance(const FVector& Start, const FVector& End, const FSceneView* View, float CorrectDistance = -1)
 	{
 		const float MaxDistanceSq = (CorrectDistance > 0) ? FMath::Square(CorrectDistance) : ARecastNavMesh::GetDrawDistanceSq();
-		return FVector::DistSquared(Start, View->ViewMatrices.GetViewOrigin()) < MaxDistanceSq &&
-			FVector::DistSquared(End, View->ViewMatrices.GetViewOrigin()) < MaxDistanceSq;
+		return	FVector::DistSquaredXY(Start, View->ViewMatrices.GetViewOrigin()) < MaxDistanceSq &&
+				FVector::DistSquaredXY(End, View->ViewMatrices.GetViewOrigin()) < MaxDistanceSq;
 	}
 
 	FVector EvalArc(const FVector& Org, const FVector& Dir, const float h, const float u)
@@ -182,7 +186,7 @@ namespace FNavMeshRenderingHelpers
 		Vertex->Color = Color;
 	}
 
-	void AddTriangle(FNavMeshSceneProxyData::FDebugMeshData& MeshData, int32 V0, int32 V1, int32 V2)
+	void AddTriangleIndices(FNavMeshSceneProxyData::FDebugMeshData& MeshData, int32 V0, int32 V1, int32 V2)
 	{
 		MeshData.Indices.Add(V0);
 		MeshData.Indices.Add(V1);
@@ -203,7 +207,7 @@ namespace FNavMeshRenderingHelpers
 		}
 		for (int32 TriIdx = 0; TriIdx < MeshIndices.Num(); TriIdx += 3)
 		{
-			FNavMeshRenderingHelpers::AddTriangle(DebugMeshData, MeshIndices[TriIdx], MeshIndices[TriIdx + 1], MeshIndices[TriIdx + 2]);
+			FNavMeshRenderingHelpers::AddTriangleIndices(DebugMeshData, MeshIndices[TriIdx], MeshIndices[TriIdx + 1], MeshIndices[TriIdx + 2]);
 		}
 
 		DebugMeshData.ClusterColor = Color;
@@ -635,6 +639,18 @@ void FNavMeshSceneProxyData::GatherData(const ARecastNavMesh* NavMesh, int32 InN
 			}
 		}
 
+#if RECAST_INTERNAL_DEBUG_DATA
+		// Display internal debug data
+		for (const FIntPoint& point : NavMeshGeometry.TilesToDisplayInternalData)
+		{
+			if (NavMesh->GetDebugDataMap())
+			{
+				const FRecastInternalDebugData& DebugData = NavMesh->GetDebugDataMap()->FindRef(point);
+				AddMeshForInternalData(DebugData);
+			}
+		}
+#endif
+
 		NavMesh->FinishBatchQuery();
 
 		// Draw Mesh
@@ -659,7 +675,7 @@ void FNavMeshSceneProxyData::GatherData(const ARecastNavMesh* NavMesh, int32 InN
 					}
 					for (int32 TriIdx = 0; TriIdx < MeshIndices.Num(); TriIdx += 3)
 					{
-						FNavMeshRenderingHelpers::AddTriangle(DebugMeshData, MeshIndices[TriIdx], MeshIndices[TriIdx + 1], MeshIndices[TriIdx + 2]);
+						FNavMeshRenderingHelpers::AddTriangleIndices(DebugMeshData, MeshIndices[TriIdx], MeshIndices[TriIdx + 1], MeshIndices[TriIdx + 2]);
 					}
 
 					MeshBuilders.Add(DebugMeshData);
@@ -821,10 +837,10 @@ void FNavMeshSceneProxyData::GatherData(const ARecastNavMesh* NavMesh, int32 InN
 						FNavMeshRenderingHelpers::AddVertex(DebugMeshData, Pt0, (ArcIdx == NumArcPoints) ? ColB : FColor::White);
 						FNavMeshRenderingHelpers::AddVertex(DebugMeshData, Pt1, (ArcIdx == NumArcPoints) ? ColB : FColor::White);
 
-						FNavMeshRenderingHelpers::AddTriangle(DebugMeshData, VertBase + 0, VertBase + 2, VertBase + 1);
-						FNavMeshRenderingHelpers::AddTriangle(DebugMeshData, VertBase + 2, VertBase + 3, VertBase + 1);
-						FNavMeshRenderingHelpers::AddTriangle(DebugMeshData, VertBase + 0, VertBase + 1, VertBase + 2);
-						FNavMeshRenderingHelpers::AddTriangle(DebugMeshData, VertBase + 2, VertBase + 1, VertBase + 3);
+						FNavMeshRenderingHelpers::AddTriangleIndices(DebugMeshData, VertBase + 0, VertBase + 2, VertBase + 1);
+						FNavMeshRenderingHelpers::AddTriangleIndices(DebugMeshData, VertBase + 2, VertBase + 3, VertBase + 1);
+						FNavMeshRenderingHelpers::AddTriangleIndices(DebugMeshData, VertBase + 0, VertBase + 1, VertBase + 2);
+						FNavMeshRenderingHelpers::AddTriangleIndices(DebugMeshData, VertBase + 2, VertBase + 1, VertBase + 3);
 
 						VertBase += 2;
 						Prev0 = Pt0;
@@ -844,6 +860,38 @@ void FNavMeshSceneProxyData::GatherData(const ARecastNavMesh* NavMesh, int32 InN
 	}
 }
 
+#if RECAST_INTERNAL_DEBUG_DATA
+void FNavMeshSceneProxyData::AddMeshForInternalData(const FRecastInternalDebugData& InInternalData)
+{
+	// Only supports triangles for now
+	for (const FRecastInternalDebugData::Command& cmd : InInternalData.Commands)
+	{
+		if (cmd.Prim != DU_DRAW_TRIS)
+		{
+			UE_LOG(LogNavigation, Warning, TEXT("Debug data display command not supported yet."));
+			return;
+		}
+	}
+
+	// Add triangles
+	const TArray<FVector>& Verts = InInternalData.Vertices;
+	const TArray<uint32>&  VertsColors = InInternalData.VertexColors;
+	const TArray<uint32>&  Indices = InInternalData.Indices;
+	if (ensure(Verts.Num() == VertsColors.Num()) && ensure(Indices.Num() % 3 == 0))
+	{
+		MeshBuilders.AddDefaulted();
+		FDebugMeshData& MeshData = MeshBuilders.Last();
+		MeshData.Vertices.Reserve(Verts.Num());
+		for (int32 VertIdx = 0; VertIdx < Verts.Num(); ++VertIdx)
+		{
+			FNavMeshRenderingHelpers::AddVertex(MeshData, Verts[VertIdx], FColor(VertsColors[VertIdx]));
+		}
+		MeshData.Indices = Indices;
+		MeshData.ClusterColor = FColor(60, 200, 200, 100); //find the proper material to use the vert colors instead
+	}
+}
+#endif //RECAST_INTERNAL_DEBUG_DATA
+
 #endif
 
 //////////////////////////////////////////////////////////////////////////
@@ -1040,7 +1088,7 @@ void FNavMeshSceneProxy::GetDynamicMeshElements(const TArray<const FSceneView*>&
 			for (int32 Index = 0; Index < Num; ++Index)
 			{
 				const FDebugLine &Line = ProxyData.NavMeshEdgeLines[Index];
-				if (FNavMeshRenderingHelpers::LineInView(Line.Start, Line.End, View, false))
+				if (FNavMeshRenderingHelpers::LineInView(Line.Start, Line.End, View))
 				{
 					if (FNavMeshRenderingHelpers::LineInCorrectDistance(Line.Start, Line.End, View))
 					{
@@ -1059,7 +1107,7 @@ void FNavMeshSceneProxy::GetDynamicMeshElements(const TArray<const FSceneView*>&
 			for (int32 Index = 0; Index < Num; ++Index)
 			{
 				const FDebugLine &Line = ProxyData.ClusterLinkLines[Index];
-				if (FNavMeshRenderingHelpers::LineInView(Line.Start, Line.End, View, false))
+				if (FNavMeshRenderingHelpers::LineInView(Line.Start, Line.End, View))
 				{
 					if (FNavMeshRenderingHelpers::LineInCorrectDistance(Line.Start, Line.End, View))
 					{
@@ -1078,7 +1126,7 @@ void FNavMeshSceneProxy::GetDynamicMeshElements(const TArray<const FSceneView*>&
 			for (int32 Index = 0; Index < Num; ++Index)
 			{
 				const FDebugLine &Line = ProxyData.TileEdgeLines[Index];
-				if (FNavMeshRenderingHelpers::LineInView(Line.Start, Line.End, View, false))
+				if (FNavMeshRenderingHelpers::LineInView(Line.Start, Line.End, View))
 				{
 					if (FNavMeshRenderingHelpers::LineInCorrectDistance(Line.Start, Line.End, View))
 					{
@@ -1097,7 +1145,7 @@ void FNavMeshSceneProxy::GetDynamicMeshElements(const TArray<const FSceneView*>&
 			for (int32 Index = 0; Index < Num; ++Index)
 			{
 				const FDebugLine &Line = ProxyData.NavLinkLines[Index];
-				if (FNavMeshRenderingHelpers::LineInView(Line.Start, Line.End, View, false))
+				if (FNavMeshRenderingHelpers::LineInView(Line.Start, Line.End, View))
 				{
 					if (FNavMeshRenderingHelpers::LineInCorrectDistance(Line.Start, Line.End, View))
 					{
@@ -1115,7 +1163,7 @@ void FNavMeshSceneProxy::GetDynamicMeshElements(const TArray<const FSceneView*>&
 			for (int32 Index = 0; Index < Num; ++Index)
 			{
 				const auto &Line = ProxyData.ThickLineItems[Index];
-				if (FNavMeshRenderingHelpers::LineInView(Line.Start, Line.End, View, false))
+				if (FNavMeshRenderingHelpers::LineInView(Line.Start, Line.End, View))
 				{
 					if (FNavMeshRenderingHelpers::LineInCorrectDistance(Line.Start, Line.End, View))
 					{
diff --git a/Engine/Source/Runtime/NavigationSystem/Private/NavMesh/PImplRecastNavMesh.cpp b/Engine/Source/Runtime/NavigationSystem/Private/NavMesh/PImplRecastNavMesh.cpp
--- a/Engine/Source/Runtime/NavigationSystem/Private/NavMesh/PImplRecastNavMesh.cpp
+++ b/Engine/Source/Runtime/NavigationSystem/Private/NavMesh/PImplRecastNavMesh.cpp
@@ -281,8 +281,11 @@ void FPImplRecastNavMesh::ReleaseDetourNavMesh()
 	}
 	DetourNavMesh = nullptr;
 	
-	//
 	CompressedTileCacheLayers.Empty();
+
+#if RECAST_INTERNAL_DEBUG_DATA
+	DebugDataMap.Empty();
+#endif
 }
 
 /**
@@ -1525,7 +1528,7 @@ NavNodeRef FPImplRecastNavMesh::FindNearestPoly(FVector const& Loc, FVector cons
 
 bool FPImplRecastNavMesh::GetPolysWithinPathingDistance(FVector const& StartLoc, const float PathingDistance, 
 	const FNavigationQueryFilter& Filter, const UObject* Owner,
-	TArray<NavNodeRef>& FoundPolys, FRecastDebugPathfindingData* DebugData) const
+	TArray<NavNodeRef>& FoundPolys, FRecastDebugPathfindingData* OutDebugData) const
 {
 	ensure(PathingDistance > 0.0f && "PathingDistance <= 0 doesn't make sense");
 	
@@ -1569,9 +1572,9 @@ bool FPImplRecastNavMesh::GetPolysWithinPathingDistance(FVector const& StartLoc,
 	NavQuery.findPolysInPathDistance(StartPolyID, RecastStartPos, PathingDistance, QueryFilter, FoundPolys.GetData(), &NumPolys, MaxSearchNodes);
 	FoundPolys.RemoveAt(NumPolys, FoundPolys.Num() - NumPolys);
 
-	if (DebugData)
+	if (OutDebugData)
 	{
-		StorePathfindingDebugData(NavQuery, DetourNavMesh, *DebugData);
+		StorePathfindingDebugData(NavQuery, DetourNavMesh, *OutDebugData);
 	}
 
 	return FoundPolys.Num() > 0;
@@ -2276,7 +2279,11 @@ int32 FPImplRecastNavMesh::GetTilesDebugGeometry(const FRecastNavMeshGenerator*
 	check(NavMeshOwner && DetourNavMesh);
 	dtMeshHeader const* const Header = Tile.header;
 	check(Header);
-	
+
+#if RECAST_INTERNAL_DEBUG_DATA
+	OutGeometry.TilesToDisplayInternalData.Push(FIntPoint(Header->x, Header->y));
+#endif
+
 	const bool bIsBeingBuilt = Generator != nullptr && !!NavMeshOwner->bDistinctlyDrawTilesBeingBuilt
 		&& Generator->IsTileChanged(TileIdx == INDEX_NONE ? DetourNavMesh->decodePolyIdTile(DetourNavMesh->getTileRef(&Tile)) : TileIdx);
 
@@ -2288,6 +2295,7 @@ int32 FPImplRecastNavMesh::GetTilesDebugGeometry(const FRecastNavMeshGenerator*
 		OutGeometry.MeshVerts.Add(VertPos);
 		F += 3;
 	}
+
 	int32 const DetailVertIndexBase = Header->vertCount;
 	// add the detail verts
 	F = Tile.detailVerts;
@@ -2756,7 +2764,11 @@ void FPImplRecastNavMesh::RemoveTileCacheLayer(int32 TileX, int32 TileY, int32 L
 		
 		if (ExistingLayersList->Num() == 0)
 		{
-			CompressedTileCacheLayers.Remove(FIntPoint(TileX, TileY));
+			RemoveTileCacheLayers(TileX, TileY);
+
+#if RECAST_INTERNAL_DEBUG_DATA
+			NavMeshOwner->RemoveTileDebugData(TileX, TileY);
+#endif
 		}
 	}
 }
diff --git a/Engine/Source/Runtime/NavigationSystem/Private/NavMesh/RecastNavMesh.cpp b/Engine/Source/Runtime/NavigationSystem/Private/NavMesh/RecastNavMesh.cpp
--- a/Engine/Source/Runtime/NavigationSystem/Private/NavMesh/RecastNavMesh.cpp
+++ b/Engine/Source/Runtime/NavigationSystem/Private/NavMesh/RecastNavMesh.cpp
@@ -932,6 +932,24 @@ void ARecastNavMesh::AddTileCacheLayers(int32 TileX, int32 TileY, const TArray<F
 	}
 }
 
+#if RECAST_INTERNAL_DEBUG_DATA
+void ARecastNavMesh::RemoveTileDebugData(int32 TileX, int32 TileY)
+{
+	if (RecastNavMeshImpl)
+	{
+		RecastNavMeshImpl->DebugDataMap.Remove(FIntPoint(TileX, TileY));
+	}
+}
+
+void ARecastNavMesh::AddTileDebugData(int32 TileX, int32 TileY, const FRecastInternalDebugData& InTileDebugData)
+{
+	if (RecastNavMeshImpl)
+	{
+		RecastNavMeshImpl->DebugDataMap.Add(FIntPoint(TileX, TileY), InTileDebugData);
+	}
+}
+#endif //RECAST_INTERNAL_DEBUG_DATA
+
 void ARecastNavMesh::MarkEmptyTileCacheLayers(int32 TileX, int32 TileY)
 {
 	if (RecastNavMeshImpl && bStoreEmptyTileLayers)
@@ -2688,6 +2706,17 @@ void ARecastNavMesh::RebuildTile(const TArray<FIntPoint>& Tiles)
 	}
 }
 
+#if RECAST_INTERNAL_DEBUG_DATA
+const TMap<FIntPoint, struct FRecastInternalDebugData>* ARecastNavMesh::GetDebugDataMap() const
+{
+	if (RecastNavMeshImpl)
+	{
+		return &RecastNavMeshImpl->DebugDataMap;
+	}
+	return nullptr;
+}
+#endif //RECAST_INTERNAL_DEBUG_DATA
+
 //----------------------------------------------------------------------//
 // FRecastNavMeshCachedData
 //----------------------------------------------------------------------//
diff --git a/Engine/Source/Runtime/NavigationSystem/Private/NavMesh/RecastNavMeshGenerator.cpp b/Engine/Source/Runtime/NavigationSystem/Private/NavMesh/RecastNavMeshGenerator.cpp
--- a/Engine/Source/Runtime/NavigationSystem/Private/NavMesh/RecastNavMeshGenerator.cpp
+++ b/Engine/Source/Runtime/NavigationSystem/Private/NavMesh/RecastNavMeshGenerator.cpp
@@ -32,6 +32,11 @@
 #include "PhysicsEngine/ConvexElem.h"
 #include "PhysicsEngine/BodySetup.h"
 
+#if RECAST_INTERNAL_DEBUG_DATA
+#include "DebugUtils/DebugDraw.h"
+#include "DebugUtils/RecastDebugDraw.h"
+#endif //RECAST_INTERNAL_DEBUG_DATA
+
 #ifndef OUTPUT_NAV_TILE_LAYER_COMPRESSION_DATA
 	#define OUTPUT_NAV_TILE_LAYER_COMPRESSION_DATA 0
 #endif
@@ -53,6 +58,19 @@ CSV_DEFINE_CATEGORY(NAVREGEN, false);
 
 struct dtTileCacheAlloc;
 
+//Experimental debug tools
+static int32 GNavmeshSynchronousTileGeneration = 0;
+static FAutoConsoleVariableRef NavmeshVarSynchronous(TEXT("n.GNavmeshSynchronousTileGeneration"), GNavmeshSynchronousTileGeneration, TEXT(""), ECVF_Default);
+
+#if RECAST_INTERNAL_DEBUG_DATA
+static int32 GNavmeshDisplayStep = 0;
+static int32 GNavmeshDebugTileX = 1;
+static int32 GNavmeshDebugTileY = 1;
+static FAutoConsoleVariableRef NavmeshVarDisplayStep(TEXT("n.GNavmeshDisplayStep"), GNavmeshDisplayStep, TEXT(""), ECVF_Default);
+static FAutoConsoleVariableRef NavmeshVarDebugTileX(TEXT("n.GNavmeshDebugTileX"), GNavmeshDebugTileX, TEXT(""), ECVF_Default);
+static FAutoConsoleVariableRef NavmeshVarDebugTileY(TEXT("n.GNavmeshDebugTileY"), GNavmeshDebugTileY, TEXT(""), ECVF_Default);
+#endif //RECAST_INTERNAL_DEBUG_DATA
+
 FORCEINLINE bool DoesBoxContainOrOverlapVector(const FBox& BigBox, const FVector& In)
 {
 	return (In.X >= BigBox.Min.X) && (In.X <= BigBox.Max.X) 
@@ -1623,10 +1641,18 @@ protected:
 class FNavMeshBuildContext : public rcContext, public dtTileCacheLogContext
 {
 public:
-	FNavMeshBuildContext()
+	FNavMeshBuildContext(FRecastTileGenerator& InTileGenerator)
 		: rcContext(true)
+#if RECAST_INTERNAL_DEBUG_DATA
+		, InternalDebugData(InTileGenerator.GetMutableDebugData())
+#endif
 	{
 	}
+
+#if RECAST_INTERNAL_DEBUG_DATA
+	FRecastInternalDebugData& InternalDebugData;
+#endif
+
 protected:
 	/// Logs a message.
 	///  @param[in]		category	The category of the message.
@@ -2539,7 +2565,7 @@ void FRecastTileGenerator::AppendGeometry(const TNavStatArray<uint8>& RawCollisi
 
 ETimeSliceWorkResult FRecastTileGenerator::GenerateTileTimeSliced()
 {
-	FNavMeshBuildContext BuildContext;
+	FNavMeshBuildContext BuildContext(*this);
 	ETimeSliceWorkResult WorkResult = ETimeSliceWorkResult::Succeeded;
 
 	switch (GenerateTileTimeSlicedState)
@@ -2602,7 +2628,7 @@ ETimeSliceWorkResult FRecastTileGenerator::GenerateTileTimeSliced()
 
 bool FRecastTileGenerator::GenerateTile()
 {
-	FNavMeshBuildContext BuildContext;
+	FNavMeshBuildContext BuildContext(*this);
 	bool bSuccess = true;
 
 	if (bRegenerateCompressedLayers)
@@ -3339,24 +3365,59 @@ bool FRecastTileGenerator::GenerateCompressedLayers(FNavMeshBuildContext& BuildC
 		return true;
 	}
 
+#if RECAST_INTERNAL_DEBUG_DATA
+	if (GNavmeshDisplayStep == 10 && IsTileToDebug())
+	{
+		duDebugDrawHeightfieldSolid(&BuildContext.InternalDebugData, *RasterContext.SolidHF);
+	}
+#endif
+
 	// Reject voxels outside generation boundaries
 	if (TileConfig.bPerformVoxelFiltering && !bFullyEncapsulatedByInclusionBounds)
 	{
 		ApplyVoxelFilter(RasterContext.SolidHF, TileConfig.walkableRadius);
 	}
 
+#if RECAST_INTERNAL_DEBUG_DATA
+	if (GNavmeshDisplayStep == 20 && IsTileToDebug())
+	{
+		duDebugDrawHeightfieldSolid(&BuildContext.InternalDebugData, *RasterContext.SolidHF);
+	}
+#endif
+
 	GenerateRecastFilter(BuildContext, RasterContext);
 
+#if RECAST_INTERNAL_DEBUG_DATA
+	if (GNavmeshDisplayStep == 30 && IsTileToDebug())
+	{
+		duDebugDrawHeightfieldSolid(&BuildContext.InternalDebugData, *RasterContext.SolidHF);
+	}
+#endif
+
 	if (!BuildCompactHeightField(BuildContext, RasterContext))
 	{
 		return false;
 	}
 
+#if RECAST_INTERNAL_DEBUG_DATA
+	if (GNavmeshDisplayStep == 40 && IsTileToDebug())
+	{
+		duDebugDrawCompactHeightfieldSolid(&BuildContext.InternalDebugData, *RasterContext.CompactHF);
+	}
+#endif
+
 	if (!RecastErodeWalkable(BuildContext, RasterContext))
 	{
 		return false;
 	}
 
+#if RECAST_INTERNAL_DEBUG_DATA
+	if (GNavmeshDisplayStep == 50 && IsTileToDebug())
+	{
+		duDebugDrawCompactHeightfieldSolid(&BuildContext.InternalDebugData, *RasterContext.CompactHF);
+	}
+#endif
+
 	if (!RecastBuildLayers(BuildContext, RasterContext))
 	{
 		return false;
@@ -4784,6 +4845,10 @@ TArray<uint32> FRecastNavMeshGenerator::RemoveTileLayers(const int32 TileX, cons
 
 		// Remove compressed tile cache layers
 		DestNavMesh->RemoveTileCacheLayers(TileX, TileY);
+
+#if RECAST_INTERNAL_DEBUG_DATA
+		DestNavMesh->RemoveTileDebugData(TileX, TileY);
+#endif
 	}
 
 	return UpdatedIndices;
@@ -5369,6 +5434,13 @@ void FRecastNavMeshGenerator::StoreCompressedTileCacheLayers(const FRecastTileGe
 	}
 }
 
+#if RECAST_INTERNAL_DEBUG_DATA
+void FRecastNavMeshGenerator::StoreDebugData(const FRecastTileGenerator& TileGenerator, int32 TileX, int32 TileY)
+{
+	DestNavMesh->AddTileDebugData(TileX, TileY, TileGenerator.GetDebugData());
+}
+#endif
+
 #if RECAST_ASYNC_REBUILDING
 TArray<uint32> FRecastNavMeshGenerator::ProcessTileTasksAsync(const int32 NumTasksToProcess)
 {
@@ -5396,7 +5468,15 @@ TArray<uint32> FRecastNavMeshGenerator::ProcessTileTasksAsync(const int32 NumTas
 			if (TileTask->GetTask().TileGenerator->HasDataToBuild())
 			{
 				RunningElement.AsyncTask = TileTask.Release();
-				RunningElement.AsyncTask->StartBackgroundTask();
+
+				if (!GNavmeshSynchronousTileGeneration)
+				{
+					RunningElement.AsyncTask->StartBackgroundTask();
+				}
+				else
+				{
+					RunningElement.AsyncTask->StartSynchronousTask();
+				}
 			
 				RunningDirtyTiles.Add(RunningElement);
 			}
@@ -5435,6 +5515,10 @@ TArray<uint32> FRecastNavMeshGenerator::ProcessTileTasksAsync(const int32 NumTas
 				UpdatedTiles.Append(UpdatedTileIndices);
 			
 				StoreCompressedTileCacheLayers(TileGenerator, Element.Coord.X, Element.Coord.Y);
+
+#if RECAST_INTERNAL_DEBUG_DATA
+				StoreDebugData(TileGenerator, Element.Coord.X, Element.Coord.Y);
+#endif
 			}
 
 			{
@@ -6388,3 +6472,10 @@ public:
 } NavigationGeomExec;
 
 #endif // WITH_RECAST
+
+#if RECAST_INTERNAL_DEBUG_DATA
+bool FRecastTileGenerator::IsTileToDebug()
+{
+	return TileX == GNavmeshDebugTileX && TileY == GNavmeshDebugTileY;
+}
+#endif //RECAST_INTERNAL_DEBUG_DATA
\ No newline at end of file
diff --git a/Engine/Source/Runtime/NavigationSystem/Public/NavMesh/NavMeshRenderingComponent.h b/Engine/Source/Runtime/NavigationSystem/Public/NavMesh/NavMeshRenderingComponent.h
--- a/Engine/Source/Runtime/NavigationSystem/Public/NavMesh/NavMeshRenderingComponent.h
+++ b/Engine/Source/Runtime/NavigationSystem/Public/NavMesh/NavMeshRenderingComponent.h
@@ -15,6 +15,7 @@
 #include "LocalVertexFactory.h"
 #include "Math/GenericOctree.h"
 #include "StaticMeshResources.h"
+#include "NavigationSystemTypes.h"
 #include "NavMeshRenderingComponent.generated.h"
 
 class APlayerController;
@@ -91,6 +92,11 @@ struct NAVIGATIONSYSTEM_API FNavMeshSceneProxyData : public TSharedFromThis<FNav
 #if WITH_RECAST
 	int32 GetDetailFlags(const ARecastNavMesh* NavMesh) const;
 	void GatherData(const ARecastNavMesh* NavMesh, int32 InNavDetailFlags, const TArray<int32>& TileSet);
+
+#if RECAST_INTERNAL_DEBUG_DATA
+	void AddMeshForInternalData(const struct FRecastInternalDebugData& InInternalData);
+#endif //RECAST_INTERNAL_DEBUG_DATA
+
 #endif
 };
 
@@ -221,5 +227,5 @@ namespace FNavMeshRenderingHelpers
 {
 	NAVIGATIONSYSTEM_API void AddVertex(FNavMeshSceneProxyData::FDebugMeshData& MeshData, const FVector& Pos, const FColor Color = FColor::White);
 
-	NAVIGATIONSYSTEM_API void AddTriangle(FNavMeshSceneProxyData::FDebugMeshData& MeshData, int32 V0, int32 V1, int32 V2);
+	NAVIGATIONSYSTEM_API void AddTriangleIndices(FNavMeshSceneProxyData::FDebugMeshData& MeshData, int32 V0, int32 V1, int32 V2);
 }
diff --git a/Engine/Source/Runtime/NavigationSystem/Public/NavMesh/PImplRecastNavMesh.h b/Engine/Source/Runtime/NavigationSystem/Public/NavMesh/PImplRecastNavMesh.h
--- a/Engine/Source/Runtime/NavigationSystem/Public/NavMesh/PImplRecastNavMesh.h
+++ b/Engine/Source/Runtime/NavigationSystem/Public/NavMesh/PImplRecastNavMesh.h
@@ -15,6 +15,10 @@
 #include "NavMesh/RecastNavMesh.h"
 #include "NavMesh/RecastQueryFilter.h"
 
+#if RECAST_INTERNAL_DEBUG_DATA
+#include "NavMesh/RecastInternalDebugData.h"
+#endif
+
 #if WITH_RECAST
 #include "Detour/DetourNavMesh.h"
 #include "Detour/DetourNavMeshQuery.h"
@@ -212,6 +216,10 @@ public:
 	/** Compressed layers data, can be reused for tiles generation */
 	TMap<FIntPoint, TArray<FNavMeshTileData> > CompressedTileCacheLayers;
 
+#if RECAST_INTERNAL_DEBUG_DATA
+	TMap<FIntPoint, FRecastInternalDebugData> DebugDataMap;
+#endif
+
 	/** query used for searching data on game thread */
 	mutable dtNavMeshQuery SharedNavQuery;
 
diff --git a/Engine/Source/Runtime/NavigationSystem/Public/NavMesh/RecastInternalDebugData.cpp b/Engine/Source/Runtime/NavigationSystem/Public/NavMesh/RecastInternalDebugData.cpp
new file mode 100644
--- /dev/null
+++ b/Engine/Source/Runtime/NavigationSystem/Public/NavMesh/RecastInternalDebugData.cpp
@@ -0,0 +1,46 @@
+// Copyright Epic Games, Inc. All Rights Reserved.
+
+#include "NavMesh/RecastInternalDebugData.h"
+#include "DebugUtils/DebugDraw.h"
+
+void FRecastInternalDebugData::vertex(const float x, const float y, const float z, unsigned int color, const float u, const float v)
+{
+	float pos[3] = { x,y,z };
+	Vertices.Push(Recast2UnrealPoint(pos));
+	VertexColors.Push(color);
+}
+
+void FRecastInternalDebugData::end()
+{
+	if (ensure(Commands.Num() > 0))
+	{
+		Command& Cmd = Commands.Last();
+		const int FirstVertex = Cmd.Offset;
+
+		Cmd.Offset = Indices.Num();
+
+		// Turns quads to triangles
+		if (Cmd.Prim == DU_DRAW_QUADS)
+		{
+			Cmd.Prim = DU_DRAW_TRIS;
+			for (int i = FirstVertex; i < Vertices.Num(); i += 4)
+			{
+				ensure(i + 3 < Vertices.Num());
+				Indices.Push(i + 0);
+				Indices.Push(i + 1);
+				Indices.Push(i + 3);
+
+				Indices.Push(i + 3);
+				Indices.Push(i + 1);
+				Indices.Push(i + 2);
+			}
+		}
+		else
+		{
+			for (int i = FirstVertex; i < Vertices.Num(); i++)
+			{
+				Indices.Push(i);
+			}
+		}
+	}
+}
diff --git a/Engine/Source/Runtime/NavigationSystem/Public/NavMesh/RecastInternalDebugData.h b/Engine/Source/Runtime/NavigationSystem/Public/NavMesh/RecastInternalDebugData.h
new file mode 100644
--- /dev/null
+++ b/Engine/Source/Runtime/NavigationSystem/Public/NavMesh/RecastInternalDebugData.h
@@ -0,0 +1,56 @@
+// Copyright Epic Games, Inc. All Rights Reserved.
+
+#pragma once
+
+#include "NavMesh/RecastHelpers.h"
+#include "DebugUtils/DebugDraw.h"
+
+struct FRecastInternalDebugData : public duDebugDraw
+{
+	struct Command
+	{
+		duDebugDrawPrimitives Prim;
+		int Offset; // Index offset
+		int Count;  // Index count
+	};
+
+	TArray<Command> Commands;
+	TArray<uint32> Indices;
+	TArray<FVector> Vertices;
+	TArray<uint32>  VertexColors;
+
+	FRecastInternalDebugData() {}
+	virtual ~FRecastInternalDebugData() override {}
+
+	virtual void depthMask(bool state) override { /*unused*/ };
+	virtual void texture(bool state) override { /*unused*/ };
+
+	virtual void begin(duDebugDrawPrimitives prim, float size = 1.0f) override
+	{
+		Command Cmd;
+		Cmd.Prim = prim;
+		Cmd.Offset = Vertices.Num();    // Misuse to store initial state.
+		Cmd.Count = 0;
+		Commands.Push(Cmd);
+	}
+
+	virtual void vertex(const float* pos, unsigned int color) override
+	{
+		vertex(pos[0], pos[1], pos[2], color, 0.0f, 0.0f);
+	}
+
+	virtual void vertex(const float x, const float y, const float z, unsigned int color) override
+	{
+		vertex(x, y, z, color, 0.0f, 0.0f);
+	}
+
+	virtual void vertex(const float* pos, unsigned int color, const float* uv) override
+	{
+		vertex(pos[0], pos[1], pos[2], color, uv[0], uv[1]);
+	}
+
+	virtual void vertex(const float x, const float y, const float z, unsigned int color, const float u, const float v) override;
+
+	virtual void end() override;
+};
+
diff --git a/Engine/Source/Runtime/NavigationSystem/Public/NavMesh/RecastNavMesh.h b/Engine/Source/Runtime/NavigationSystem/Public/NavMesh/RecastNavMesh.h
--- a/Engine/Source/Runtime/NavigationSystem/Public/NavMesh/RecastNavMesh.h
+++ b/Engine/Source/Runtime/NavigationSystem/Public/NavMesh/RecastNavMesh.h
@@ -177,6 +177,10 @@ struct FRecastDebugGeometry
 	TArray<FOffMeshSegment> OffMeshSegments;
 	TArray<int32> OffMeshSegmentAreas[RECAST_MAX_AREAS];
 
+#if RECAST_INTERNAL_DEBUG_DATA
+	TArray<FIntPoint> TilesToDisplayInternalData;
+#endif
+
 	int32 bGatherPolyEdges : 1;
 	int32 bGatherNavMeshEdges : 1;
 	int32 bMarkForbiddenPolys : 1;
@@ -775,6 +779,11 @@ public:
 	/** Stores compressed tile data for given tile coord */
 	void AddTileCacheLayers(int32 TileX, int32 TileY, const TArray<FNavMeshTileData>& InLayers);
 
+#if RECAST_INTERNAL_DEBUG_DATA
+	void RemoveTileDebugData(int32 TileX, int32 TileY);
+	void AddTileDebugData(int32 TileX, int32 TileY, const struct FRecastInternalDebugData& InTileDebugData);
+#endif
+
 	/** Marks tile coord as rebuild and empty */
 	void MarkEmptyTileCacheLayers(int32 TileX, int32 TileY);
 	
@@ -1019,6 +1028,10 @@ public:
 	virtual void RemoveTiles(const TArray<FIntPoint>& Tiles);
 	void RebuildTile(const TArray<FIntPoint>& Tiles);
 
+#if RECAST_INTERNAL_DEBUG_DATA
+	const TMap<FIntPoint, struct FRecastInternalDebugData>* GetDebugDataMap() const;
+#endif
+
 protected:
 
 	void UpdatePolyRefBitsPreview();
diff --git a/Engine/Source/Runtime/NavigationSystem/Public/NavMesh/RecastNavMeshGenerator.h b/Engine/Source/Runtime/NavigationSystem/Public/NavMesh/RecastNavMeshGenerator.h
--- a/Engine/Source/Runtime/NavigationSystem/Public/NavMesh/RecastNavMeshGenerator.h
+++ b/Engine/Source/Runtime/NavigationSystem/Public/NavMesh/RecastNavMeshGenerator.h
@@ -20,6 +20,10 @@
 #include "Recast/Recast.h"
 #include "Detour/DetourNavMesh.h"
 
+#if RECAST_INTERNAL_DEBUG_DATA
+#include "NavMesh/RecastInternalDebugData.h"
+#endif
+
 class UBodySetup;
 class ARecastNavMesh;
 class FNavigationOctree;
@@ -322,6 +326,11 @@ public:
 	virtual void AddReferencedObjects(FReferenceCollector& Collector) override;
 	virtual FString GetReferencerName() const override;
 	// FGCObject end
+
+#if RECAST_INTERNAL_DEBUG_DATA
+	const FRecastInternalDebugData& GetDebugData() const { return DebugData; }
+	FRecastInternalDebugData& GetMutableDebugData() { return DebugData; }
+#endif
 		
 protected:
 	/** Does the actual TimeSliced tile generation. 
@@ -397,6 +406,10 @@ protected:
 
 	void DumpAsyncData();
 
+#if RECAST_INTERNAL_DEBUG_DATA
+	bool IsTileToDebug();
+#endif
+
 protected:
 	uint32 bRegenerateCompressedLayers : 1;
 	uint32 bFullyEncapsulatedByInclusionBounds : 1;
@@ -462,6 +475,10 @@ protected:
 	TNavStatArray<TSharedRef<FNavigationRelevantData, ESPMode::ThreadSafe> > NavigationRelevantData;
 	TSharedPtr<FNavigationOctree, ESPMode::ThreadSafe> NavOctree; 
 	FNavDataConfig NavDataConfig;
+
+#if RECAST_INTERNAL_DEBUG_DATA
+	FRecastInternalDebugData DebugData;
+#endif
 };
 
 struct NAVIGATIONSYSTEM_API FRecastTileGeneratorWrapper : public FNonAbandonableTask
@@ -702,6 +719,11 @@ protected:
 	void RemoveLayers(const FIntPoint& Tile, TArray<uint32>& UpdatedTiles);
 	
 	void StoreCompressedTileCacheLayers(const FRecastTileGenerator& TileGenerator, int32 TileX, int32 TileY);
+
+#if RECAST_INTERNAL_DEBUG_DATA
+	void StoreDebugData(const FRecastTileGenerator& TileGenerator, int32 TileX, int32 TileY);
+#endif
+
 #if RECAST_ASYNC_REBUILDING
 	/** Processes pending tile generation tasks Async*/
 	TArray<uint32> ProcessTileTasksAsync(const int32 NumTasksToProcess);
diff --git a/Engine/Source/Runtime/NavigationSystem/Public/NavigationSystemTypes.h b/Engine/Source/Runtime/NavigationSystem/Public/NavigationSystemTypes.h
--- a/Engine/Source/Runtime/NavigationSystem/Public/NavigationSystemTypes.h
+++ b/Engine/Source/Runtime/NavigationSystem/Public/NavigationSystemTypes.h
@@ -9,9 +9,10 @@
 #include "Math/GenericOctreePublic.h"
 #include "AI/NavigationModifier.h"
 
-
 #define NAVSYS_DEBUG (0 && UE_BUILD_DEBUG)
 
+#define RECAST_INTERNAL_DEBUG_DATA (!UE_BUILD_SHIPPING)
+
 class UBodySetup;
 class UNavCollision;
 struct FKAggregateGeom;
diff --git a/Engine/Source/Runtime/Navmesh/Private/DebugUtils/DebugDraw.cpp b/Engine/Source/Runtime/Navmesh/Private/DebugUtils/DebugDraw.cpp
--- a/Engine/Source/Runtime/Navmesh/Private/DebugUtils/DebugDraw.cpp
+++ b/Engine/Source/Runtime/Navmesh/Private/DebugUtils/DebugDraw.cpp
@@ -20,6 +20,7 @@
 //
 
 #include "DebugUtils/DebugDraw.h"
+#include "Recast/RecastAlloc.h" // UE4
 #define _USE_MATH_DEFINES
 
 
@@ -574,22 +575,22 @@ duDisplayList::duDisplayList(int cap) :
 
 duDisplayList::~duDisplayList()
 {
-	delete [] m_pos;
-	delete [] m_color;
+	rcFree(m_pos); // UE4
+	rcFree(m_color); // UE4
 }
 
 void duDisplayList::resize(int cap)
 {
-	float* newPos = new float[cap*3];
+	float* newPos = (float*)rcAlloc(sizeof(float)*3*cap, RC_ALLOC_PERM); // UE4
 	if (m_size)
 		memcpy(newPos, m_pos, sizeof(float)*3*m_size);
-	delete [] m_pos;
+	rcFree(m_pos); // UE4
 	m_pos = newPos;
 
-	unsigned int* newColor = new unsigned int[cap];
+	unsigned int* newColor = (unsigned int*)rcAlloc(sizeof(unsigned int)*cap, RC_ALLOC_PERM); // UE4
 	if (m_size)
 		memcpy(newColor, m_color, sizeof(unsigned int)*m_size);
-	delete [] m_color;
+	rcFree(m_color); // UE4
 	m_color = newColor;
 	
 	m_cap = cap;
diff --git a/Engine/Source/Runtime/Navmesh/Public/DebugUtils/DebugDraw.h b/Engine/Source/Runtime/Navmesh/Public/DebugUtils/DebugDraw.h
--- a/Engine/Source/Runtime/Navmesh/Public/DebugUtils/DebugDraw.h
+++ b/Engine/Source/Runtime/Navmesh/Public/DebugUtils/DebugDraw.h
@@ -38,7 +38,7 @@ enum duDebugDrawPrimitives
 /// Abstract debug draw interface.
 struct duDebugDraw
 {
-	virtual ~duDebugDraw() = 0;
+	NAVMESH_API virtual ~duDebugDraw() = 0;
 	
 	virtual void depthMask(bool state) = 0;
 
@@ -213,7 +213,7 @@ class duDisplayList : public duDebugDraw
 	
 public:
 	duDisplayList(int cap = 512);
-	~duDisplayList();
+	virtual ~duDisplayList() override;
 	virtual void depthMask(bool state);
 	virtual void begin(duDebugDrawPrimitives prim, float size = 1.0f);
 	virtual void vertex(const float x, const float y, const float z, unsigned int color);
diff --git a/Engine/Source/Runtime/Navmesh/Public/DebugUtils/RecastDebugDraw.h b/Engine/Source/Runtime/Navmesh/Public/DebugUtils/RecastDebugDraw.h
--- a/Engine/Source/Runtime/Navmesh/Public/DebugUtils/RecastDebugDraw.h
+++ b/Engine/Source/Runtime/Navmesh/Public/DebugUtils/RecastDebugDraw.h
@@ -26,28 +26,27 @@
 
 struct duDebugDraw;
 
-void duDebugDrawTriMesh(struct duDebugDraw* dd, const float* verts, int nverts, const int* tris, const float* normals, int ntris, const unsigned char* flags, const float texScale);
-void duDebugDrawTriMeshSlope(struct duDebugDraw* dd, const float* verts, int nverts, const int* tris, const float* normals, int ntris, const float walkableSlopeAngle, const float texScale);
+NAVMESH_API void duDebugDrawTriMesh(struct duDebugDraw* dd, const float* verts, int nverts, const int* tris, const float* normals, int ntris, const unsigned char* flags, const float texScale);
+NAVMESH_API void duDebugDrawTriMeshSlope(struct duDebugDraw* dd, const float* verts, int nverts, const int* tris, const float* normals, int ntris, const float walkableSlopeAngle, const float texScale);
 
-void duDebugDrawHeightfieldSolid(struct duDebugDraw* dd, const struct rcHeightfield& hf);
-void duDebugDrawHeightfieldWalkable(struct duDebugDraw* dd, const struct rcHeightfield& hf);
+NAVMESH_API void duDebugDrawHeightfieldSolid(struct duDebugDraw* dd, const struct rcHeightfield& hf);
+NAVMESH_API void duDebugDrawHeightfieldWalkable(struct duDebugDraw* dd, const struct rcHeightfield& hf);
 
-void duDebugDrawCompactHeightfieldSolid(struct duDebugDraw* dd, const struct rcCompactHeightfield& chf);
-void duDebugDrawCompactHeightfieldRegions(struct duDebugDraw* dd, const struct rcCompactHeightfield& chf);
-void duDebugDrawCompactHeightfieldDistance(struct duDebugDraw* dd, const struct rcCompactHeightfield& chf);
+NAVMESH_API void duDebugDrawCompactHeightfieldSolid(struct duDebugDraw* dd, const struct rcCompactHeightfield& chf);
+NAVMESH_API void duDebugDrawCompactHeightfieldRegions(struct duDebugDraw* dd, const struct rcCompactHeightfield& chf);
+NAVMESH_API void duDebugDrawCompactHeightfieldDistance(struct duDebugDraw* dd, const struct rcCompactHeightfield& chf);
 
-void duDebugDrawHeightfieldLayer(duDebugDraw* dd, const struct rcHeightfieldLayer& layer, const int idx);
-void duDebugDrawHeightfieldLayers(duDebugDraw* dd, const struct rcHeightfieldLayerSet& lset);
-void duDebugDrawHeightfieldLayersRegions(duDebugDraw* dd, const struct rcHeightfieldLayerSet& lset);
+NAVMESH_API void duDebugDrawHeightfieldLayer(duDebugDraw* dd, const struct rcHeightfieldLayer& layer, const int idx);
+NAVMESH_API void duDebugDrawHeightfieldLayers(duDebugDraw* dd, const struct rcHeightfieldLayerSet& lset);
+NAVMESH_API void duDebugDrawHeightfieldLayersRegions(duDebugDraw* dd, const struct rcHeightfieldLayerSet& lset);
 
-void duDebugDrawLayerContours(duDebugDraw* dd, const struct rcLayerContourSet& lcset);
-void duDebugDrawLayerPolyMesh(duDebugDraw* dd, const struct rcLayerPolyMesh& lmesh);
+NAVMESH_API void duDebugDrawLayerContours(duDebugDraw* dd, const struct rcLayerContourSet& lcset);
+NAVMESH_API void duDebugDrawLayerPolyMesh(duDebugDraw* dd, const struct rcLayerPolyMesh& lmesh);
 
-
-void duDebugDrawRegionConnections(struct duDebugDraw* dd, const struct rcContourSet& cset, const float alpha = 1.0f);
-void duDebugDrawRawContours(struct duDebugDraw* dd, const struct rcContourSet& cset, const float alpha = 1.0f);
-void duDebugDrawContours(struct duDebugDraw* dd, const struct rcContourSet& cset, const float alpha = 1.0f);
-void duDebugDrawPolyMesh(struct duDebugDraw* dd, const struct rcPolyMesh& mesh);
-void duDebugDrawPolyMeshDetail(struct duDebugDraw* dd, const struct rcPolyMeshDetail& dmesh);
+NAVMESH_API void duDebugDrawRegionConnections(struct duDebugDraw* dd, const struct rcContourSet& cset, const float alpha = 1.0f);
+NAVMESH_API void duDebugDrawRawContours(struct duDebugDraw* dd, const struct rcContourSet& cset, const float alpha = 1.0f);
+NAVMESH_API void duDebugDrawContours(struct duDebugDraw* dd, const struct rcContourSet& cset, const float alpha = 1.0f);
+NAVMESH_API void duDebugDrawPolyMesh(struct duDebugDraw* dd, const struct rcPolyMesh& mesh);
+NAVMESH_API void duDebugDrawPolyMeshDetail(struct duDebugDraw* dd, const struct rcPolyMeshDetail& dmesh);
 
 #endif // RECAST_DEBUGDRAW_H
