Simple change to allow better logging of locks and rebuild all to the nav mesh

diff --git a/Engine/Source/Runtime/NavigationSystem/Private/NavigationSystem.cpp b/Engine/Source/Runtime/NavigationSystem/Private/NavigationSystem.cpp
--- a/Engine/Source/Runtime/NavigationSystem/Private/NavigationSystem.cpp
+++ b/Engine/Source/Runtime/NavigationSystem/Private/NavigationSystem.cpp
@@ -3259,6 +3259,13 @@ void UNavigationSystemV1::OnPIEEnd()
 	}
 }
 
+void UNavigationSystemV1::AddNavigationBuildLock(uint8 Flags)
+{
+	NavBuildingLockFlags |= Flags;
+
+	UE_LOG(LogNavigation, Verbose, TEXT("UNavigationSystemV1::AddNavigationBuildLock IsLocked=%s"), IsNavigationBuildingLocked() ? TEXT("true") : TEXT("false"));
+}
+
 void UNavigationSystemV1::RemoveNavigationBuildLock(uint8 Flags, bool bSkipRebuildInEditor)
 {
 	const bool bWasLocked = IsNavigationBuildingLocked();
@@ -3267,14 +3274,27 @@ void UNavigationSystemV1::RemoveNavigationBuildLock(uint8 Flags, bool bSkipRebui
 
 	const bool bIsLocked = IsNavigationBuildingLocked();
 	const bool bSkipRebuild = (OperationMode == FNavigationSystemRunMode::EditorMode) && bSkipRebuildInEditor;
+
+	UE_LOG(LogNavigation, Verbose, TEXT("UNavigationSystemV1::RemoveNavigationBuildLock bWasLocked=%s IsLocked=%s"), bWasLocked ? TEXT("true") : TEXT("false"), bIsLocked ? TEXT("true") : TEXT("false"));
+
 	if (bWasLocked && !bIsLocked && !bSkipRebuild)
 	{
 		RebuildAll();
 	}
 }
 
+void UNavigationSystemV1::SetNavigationOctreeLock(bool bLock)
+{
+	UE_LOG(LogNavigation, Verbose, TEXT("UNavigationSystemV1::SetNavigationOctreeLock IsLocked=%s"), bLock ? TEXT("true") : TEXT("false"));
+
+	DefaultOctreeController.SetNavigationOctreeLock(bLock);
+}
+
+
 void UNavigationSystemV1::RebuildAll(bool bIsLoadTime)
 {
+	UE_LOG(LogNavigation, Verbose, TEXT("UNavigationSystemV1::RebuildAll"));
+
 	const bool bIsInGame = GetWorld()->IsGameWorld();
 	
 	GatherNavigationBounds();
diff --git a/Engine/Source/Runtime/NavigationSystem/Public/NavigationSystem.h b/Engine/Source/Runtime/NavigationSystem/Public/NavigationSystem.h
--- a/Engine/Source/Runtime/NavigationSystem/Public/NavigationSystem.h
+++ b/Engine/Source/Runtime/NavigationSystem/Public/NavigationSystem.h
@@ -779,10 +779,10 @@ public:
 	/** adds BSP collisions of currently streamed in levels to octree */
 	void InitializeLevelCollisions();
 
-	FORCEINLINE void AddNavigationBuildLock(uint8 Flags) { NavBuildingLockFlags |= Flags; }
+	void AddNavigationBuildLock(uint8 Flags);
 	void RemoveNavigationBuildLock(uint8 Flags, bool bSkipRebuildInEditor = false);
 
-	void SetNavigationOctreeLock(bool bLock) { DefaultOctreeController.SetNavigationOctreeLock(bLock); }
+	void SetNavigationOctreeLock(bool bLock);
 
 	/** checks if auto-rebuilding navigation data is enabled. Defaults to bNavigationAutoUpdateEnabled
 	*	value, but can be overridden per nav sys instance */
