NavigationSystem: Added CSV stat to track navigation tasks (running + remaining)

diff --git a/Engine/Source/Runtime/NavigationSystem/Private/NavigationSystem.cpp b/Engine/Source/Runtime/NavigationSystem/Private/NavigationSystem.cpp
--- a/Engine/Source/Runtime/NavigationSystem/Private/NavigationSystem.cpp
+++ b/Engine/Source/Runtime/NavigationSystem/Private/NavigationSystem.cpp
@@ -120,6 +120,7 @@ DEFINE_STAT(STAT_Navigation_ObservedPathsCount);
 DEFINE_STAT(STAT_Navigation_RecastMemory);
 
 CSV_DEFINE_CATEGORY(NavigationSystem, false);
+CSV_DEFINE_CATEGORY(NavTasks, true);
 
 //----------------------------------------------------------------------//
 // consts
@@ -984,6 +985,9 @@ void UNavigationSystemV1::Tick(float DeltaSeconds)
 		}
 	}
 
+	CSV_CUSTOM_STAT(NavTasks, NumRemainingTasks, GetNumRemainingBuildTasks(), ECsvCustomStatOp::Set);
+	CSV_CUSTOM_STAT(NavTasks, NumRunningTasks, GetNumRunningBuildTasks(), ECsvCustomStatOp::Set);
+
 	if (AsyncPathFindingQueries.Num() > 0)
 	{
 		SCOPE_CYCLE_COUNTER(STAT_Navigation_TickAsyncPathfinding);
