Fix NaN in GetPolysInBox()

diff --git a/Engine/Source/Runtime/NavigationSystem/Private/NavMesh/RecastNavMesh.cpp b/Engine/Source/Runtime/NavigationSystem/Private/NavMesh/RecastNavMesh.cpp
--- a/Engine/Source/Runtime/NavigationSystem/Private/NavMesh/RecastNavMesh.cpp
+++ b/Engine/Source/Runtime/NavigationSystem/Private/NavMesh/RecastNavMesh.cpp
@@ -1066,7 +1066,8 @@ bool ARecastNavMesh::GetRandomPointInNavigableRadius(const FVector& Origin, floa
 		const float RadiusSq = FMath::Square(Radius);
 		TArray<FNavPoly> Polys;
 		const FVector FallbackExtent(Radius, Radius, BIG_NUMBER);
-		const FBox Box(Origin - FallbackExtent, Origin + FallbackExtent);
+		const FVector BoxOrigin(Origin.X, Origin.Y, 0.f);
+		const FBox Box(BoxOrigin - FallbackExtent, BoxOrigin + FallbackExtent);
 		GetPolysInBox(Box, Polys, Filter, Querier);
 	
 		// @todo extremely naive implementation, barely random. To be improved
diff --git a/Engine/Source/Runtime/NavigationSystem/Public/NavMesh/RecastNavMesh.h b/Engine/Source/Runtime/NavigationSystem/Public/NavMesh/RecastNavMesh.h
--- a/Engine/Source/Runtime/NavigationSystem/Public/NavMesh/RecastNavMesh.h
+++ b/Engine/Source/Runtime/NavigationSystem/Public/NavMesh/RecastNavMesh.h
@@ -862,7 +862,7 @@ public:
 	/** Returns query extent including adjustments for voxelization error compensation */
 	FVector GetModifiedQueryExtent(const FVector& QueryExtent) const
 	{
-		return FVector(QueryExtent.X, QueryExtent.Y, QueryExtent.Z + FMath::Max(0.0f, VerticalDeviationFromGroundCompensation));
+		return FVector(QueryExtent.X, QueryExtent.Y, QueryExtent.Z >= BIG_NUMBER ? BIG_NUMBER : (QueryExtent.Z + FMath::Max(0.0f, VerticalDeviationFromGroundCompensation)));
 	}
 
 	//----------------------------------------------------------------------//
