Fix NaN when using FBox by using HALF_WORLD_MAX instead of BIG_NUMBER

diff --git a/Engine/Source/Runtime/NavigationSystem/Private/NavMesh/RecastNavMesh.cpp b/Engine/Source/Runtime/NavigationSystem/Private/NavMesh/RecastNavMesh.cpp
--- a/Engine/Source/Runtime/NavigationSystem/Private/NavMesh/RecastNavMesh.cpp
+++ b/Engine/Source/Runtime/NavigationSystem/Private/NavMesh/RecastNavMesh.cpp
@@ -1065,7 +1065,7 @@ bool ARecastNavMesh::GetRandomPointInNavigableRadius(const FVector& Origin, floa
 	{
 		const float RadiusSq = FMath::Square(Radius);
 		TArray<FNavPoly> Polys;
-		const FVector FallbackExtent(Radius, Radius, BIG_NUMBER);
+		const FVector FallbackExtent(Radius, Radius, HALF_WORLD_MAX); //Using HALF_WORLD_MAX instead of BIG_NUMBER, else the box size will be NaN.
 		const FVector BoxOrigin(Origin.X, Origin.Y, 0.f);
 		const FBox Box(BoxOrigin - FallbackExtent, BoxOrigin + FallbackExtent);
 		GetPolysInBox(Box, Polys, Filter, Querier);
diff --git a/Engine/Source/Runtime/NavigationSystem/Public/NavMesh/RecastNavMesh.h b/Engine/Source/Runtime/NavigationSystem/Public/NavMesh/RecastNavMesh.h
--- a/Engine/Source/Runtime/NavigationSystem/Public/NavMesh/RecastNavMesh.h
+++ b/Engine/Source/Runtime/NavigationSystem/Public/NavMesh/RecastNavMesh.h
@@ -862,7 +862,8 @@ public:
 	/** Returns query extent including adjustments for voxelization error compensation */
 	FVector GetModifiedQueryExtent(const FVector& QueryExtent) const
 	{
-		return FVector(QueryExtent.X, QueryExtent.Y, QueryExtent.Z >= BIG_NUMBER ? BIG_NUMBER : (QueryExtent.Z + FMath::Max(0.0f, VerticalDeviationFromGroundCompensation)));
+		// Using HALF_WORLD_MAX instead of BIG_NUMBER, else using the extent for a box will result in NaN.
+		return FVector(QueryExtent.X, QueryExtent.Y, QueryExtent.Z >= HALF_WORLD_MAX ? HALF_WORLD_MAX : (QueryExtent.Z + FMath::Max(0.0f, VerticalDeviationFromGroundCompensation)));
 	}
 
 	//----------------------------------------------------------------------//
