Marking map as dirty when navmesh tiles are built

diff --git a/Engine/Source/Runtime/NavigationSystem/Private/NavigationSystem.cpp b/Engine/Source/Runtime/NavigationSystem/Private/NavigationSystem.cpp
--- a/Engine/Source/Runtime/NavigationSystem/Private/NavigationSystem.cpp
+++ b/Engine/Source/Runtime/NavigationSystem/Private/NavigationSystem.cpp
@@ -3360,6 +3360,13 @@ bool UNavigationSystemV1::IsNavigationBuildInProgress(bool bCheckDirtyToo)
 void UNavigationSystemV1::OnNavigationGenerationFinished(ANavigationData& NavData)
 {
 	OnNavigationGenerationFinishedDelegate.Broadcast(&NavData);
+
+#if WITH_EDITOR
+	if (!GetWorld()->IsGameWorld())
+	{
+		NavData.MarkPackageDirty();
+	}
+#endif //WITH_EDITOR
 }
 
 int32 UNavigationSystemV1::GetNumRemainingBuildTasks() const
