NavigationDirtyArea: added error reporting for large area dirtying the navigation data
+ added advanced property "DirtyAreaWarningSizeThreshold" in NavigationSystem to report warning logs when tiles larger than the specified threshold are pushed to the dirty queue.
+ error reporting is disabled as long as navigation lock is effective
+ centralized area validations and error reporting inside method FNavigationDirtyAreasController::AddArea and added new optional parameter to retrieve source object

diff --git a/Engine/Source/Runtime/Engine/Private/Components/PrimitiveComponent.cpp b/Engine/Source/Runtime/Engine/Private/Components/PrimitiveComponent.cpp
--- a/Engine/Source/Runtime/Engine/Private/Components/PrimitiveComponent.cpp
+++ b/Engine/Source/Runtime/Engine/Private/Components/PrimitiveComponent.cpp
@@ -2440,7 +2440,8 @@ bool UPrimitiveComponent::IsNavigationRelevant() const
 
 FBox UPrimitiveComponent::GetNavigationBounds() const
 {
-	return Bounds.GetBox();
+	// Return invalid box when retrieving NavigationBounds before they are being computed at component registration
+	return bRegistered ? Bounds.GetBox() : FBox(ForceInit);
 }
 
 //////////////////////////////////////////////////////////////////////////
diff --git a/Engine/Source/Runtime/NavigationSystem/Private/NavigationDataHandler.cpp b/Engine/Source/Runtime/NavigationSystem/Private/NavigationDataHandler.cpp
--- a/Engine/Source/Runtime/NavigationSystem/Private/NavigationDataHandler.cpp
+++ b/Engine/Source/Runtime/NavigationSystem/Private/NavigationDataHandler.cpp
@@ -27,7 +27,7 @@ void FNavigationDataHandler::RemoveNavOctreeElementId(const FOctreeElementId& El
 		const FNavigationOctreeElement& ElementData = OctreeController.NavOctree->GetElementById(ElementId);
 		const int32 DirtyFlag = GetDirtyFlagHelper(UpdateFlags, ElementData.Data->GetDirtyFlag());
 		// mark area occupied by given actor as dirty
-		DirtyAreasController.AddArea(ElementData.Bounds.GetBox(), DirtyFlag);
+		DirtyAreasController.AddArea(ElementData.Bounds.GetBox(), DirtyFlag, [&ElementData] { return ElementData.Data->SourceObject.Get(); });
 		OctreeController.NavOctree->RemoveNode(ElementId);
 	}
 }
@@ -96,7 +96,7 @@ void FNavigationDataHandler::AddElementToNavOctree(const FNavigationDirtyElement
 	{
 		if (DirtyElement.bHasPrevData)
 		{
-			DirtyAreasController.AddArea(DirtyElement.PrevBounds, DirtyElement.PrevFlags);
+			DirtyAreasController.AddArea(DirtyElement.PrevBounds, DirtyElement.PrevFlags, [&DirtyElement] { return DirtyElement.Owner.Get(); });
 		}
 
 		return;
@@ -143,13 +143,10 @@ void FNavigationDataHandler::AddElementToNavOctree(const FNavigationDirtyElement
 		OctreeController.NavOctree->AddNode(ElementOwner, DirtyElement.NavInterface, ElementBounds, GeneratedData);
 	}
 
-	const FBox BBox = GeneratedData.Bounds.GetBox();
-	const bool bValidBBox = BBox.IsValid && !BBox.GetSize().IsNearlyZero();
-
-	if (bValidBBox && !GeneratedData.IsEmpty())
+	if (!GeneratedData.IsEmpty())
 	{
 		const int32 DirtyFlag = DirtyElement.FlagsOverride ? DirtyElement.FlagsOverride : GeneratedData.Data->GetDirtyFlag();
-		DirtyAreasController.AddArea(BBox, DirtyFlag);
+		DirtyAreasController.AddArea(GeneratedData.Bounds.GetBox(), DirtyFlag, [&ElementOwner] { return ElementOwner; });
 	}
 }
 
@@ -345,7 +342,7 @@ bool FNavigationDataHandler::UpdateNavOctreeElementBounds(UActorComponent& Comp,
 			if (ElementId != nullptr && ensure(OctreeController.IsValidElement(*ElementId)))
 			{
 				FNavigationOctreeElement& ElementData = OctreeController.NavOctree->GetElementById(*ElementId);
-				DirtyAreasController.AddArea(DirtyArea, ElementData.Data->GetDirtyFlag());
+				DirtyAreasController.AddArea(DirtyArea, ElementData.Data->GetDirtyFlag(), [&Comp] { return &Comp; });
 			}
 		}
 
@@ -435,7 +432,7 @@ void FNavigationDataHandler::AddLevelCollisionToOctree(ULevel& Level)
 			if (!Bounds.GetExtent().IsNearlyZero())
 			{
 				OctreeController.NavOctree->AddNode(&Level, nullptr, Bounds, BSPElem);
-				DirtyAreasController.AddArea(Bounds, ENavigationDirtyFlag::All);
+				DirtyAreasController.AddArea(Bounds, ENavigationDirtyFlag::All, [&Level] { return &Level; });
 
 				UE_LOG(LogNavOctree, Log, TEXT("ADD %s"), *Level.GetName());
 			}
@@ -457,7 +454,7 @@ void FNavigationDataHandler::RemoveLevelCollisionFromOctree(ULevel& Level)
 			{
 				// mark area occupied by given actor as dirty
 				FNavigationOctreeElement& ElementData = OctreeController.NavOctree->GetElementById(*ElementId);
-				DirtyAreasController.AddArea(ElementData.Bounds.GetBox(), ENavigationDirtyFlag::All);
+				DirtyAreasController.AddArea(ElementData.Bounds.GetBox(), ENavigationDirtyFlag::All, [&Level] { return &Level; });
 			}
 
 			OctreeController.NavOctree->RemoveNode(*ElementId);
@@ -502,4 +499,4 @@ void FNavigationDataHandler::ProcessPendingOctreeUpdates()
 		}
 	}
 	OctreeController.PendingOctreeUpdates.Empty(32);
-}
+}
\ No newline at end of file
diff --git a/Engine/Source/Runtime/NavigationSystem/Private/NavigationDirtyAreasController.cpp b/Engine/Source/Runtime/NavigationSystem/Private/NavigationDirtyAreasController.cpp
--- a/Engine/Source/Runtime/NavigationSystem/Private/NavigationDirtyAreasController.cpp
+++ b/Engine/Source/Runtime/NavigationSystem/Private/NavigationDirtyAreasController.cpp
@@ -3,6 +3,8 @@
 #include "NavigationDirtyAreasController.h"
 #include "NavigationData.h"
 
+DEFINE_LOG_CATEGORY_STATIC(LogNavigationDirtyArea, Warning, All);
+
 //----------------------------------------------------------------------//
 // FNavigationDirtyAreasController
 //----------------------------------------------------------------------//
@@ -10,6 +12,8 @@ FNavigationDirtyAreasController::FNavigationDirtyAreasController()
 	: bCanAccumulateDirtyAreas(true)
 #if !UE_BUILD_SHIPPING
 	, bDirtyAreasReportedWhileAccumulationLocked(false)
+	, bCanReportOversizedDirtyArea(false)
+	, bNavigationBuildLocked(false)
 #endif // !UE_BUILD_SHIPPING
 {
 
@@ -41,17 +45,76 @@ void FNavigationDirtyAreasController::Tick(const float DeltaSeconds, const TArra
 	}
 }
 
-void FNavigationDirtyAreasController::AddArea(const FBox& NewArea, int32 Flags)
+void FNavigationDirtyAreasController::AddArea(const FBox& NewArea, const int32 Flags, const TFunction<UObject*()>& ObjectProviderFunc /*= nullptr*/)
 {
-	if (Flags > 0 && bCanAccumulateDirtyAreas && NewArea.IsValid)
+#if !UE_BUILD_SHIPPING
+	// always keep track of reported areas even when filtered out by invalid area as long as flags are valid
+	bDirtyAreasReportedWhileAccumulationLocked = bDirtyAreasReportedWhileAccumulationLocked || (Flags > 0 && !bCanAccumulateDirtyAreas);
+#endif // !UE_BUILD_SHIPPING
+
+	if (!NewArea.IsValid)
+	{
+		UE_LOG(LogNavigationDirtyArea, Warning, TEXT("Skipping dirty area creation because of invalid bounds (object: %s)"), *GetFullNameSafe(ObjectProviderFunc ? ObjectProviderFunc() : nullptr));
+		return;
+	}
+
+	const FVector2D BoundsSize(NewArea.GetSize());
+	if (BoundsSize.IsNearlyZero())
+	{
+		UE_LOG(LogNavigationDirtyArea, Warning, TEXT("Skipping dirty area creation because of empty bounds (object: %s)"), *GetFullNameSafe(ObjectProviderFunc ? ObjectProviderFunc() : nullptr));
+		return;
+	}
+
+#if !UE_BUILD_SHIPPING
+	UE_CLOG(ShouldReportOversizedDirtyArea() && BoundsSize.GetMax() > DirtyAreaWarningSizeThreshold,
+		LogNavigationDirtyArea, Warning, TEXT("Adding an oversized dirty area (object:%s size:%s threshold:%.2f)"),
+		*GetFullNameSafe(ObjectProviderFunc ? ObjectProviderFunc() : nullptr),
+		*BoundsSize.ToString(),
+		DirtyAreaWarningSizeThreshold);
+#endif // !UE_BUILD_SHIPPING
+
+	if (Flags > 0 && bCanAccumulateDirtyAreas)
 	{
 		DirtyAreas.Add(FNavigationDirtyArea(NewArea, Flags));
 	}
+}
+
+void FNavigationDirtyAreasController::OnNavigationBuildLocked()
+{
 #if !UE_BUILD_SHIPPING
-	bDirtyAreasReportedWhileAccumulationLocked = bDirtyAreasReportedWhileAccumulationLocked || (Flags > 0 && !bCanAccumulateDirtyAreas);
+	bNavigationBuildLocked = true;
+#endif // !UE_BUILD_SHIPPING
+}
+
+void FNavigationDirtyAreasController::OnNavigationBuildUnlocked()
+{
+#if !UE_BUILD_SHIPPING
+	bNavigationBuildLocked = false;
+#endif // !UE_BUILD_SHIPPING
+}
+
+void FNavigationDirtyAreasController::SetDirtyAreaWarningSizeThreshold(const float Threshold)
+{
+#if !UE_BUILD_SHIPPING
+	DirtyAreaWarningSizeThreshold = Threshold;
 #endif // !UE_BUILD_SHIPPING
 }
 
+void FNavigationDirtyAreasController::SetCanReportOversizedDirtyArea(bool bCanReport)
+{
+#if !UE_BUILD_SHIPPING
+	bCanReportOversizedDirtyArea = bCanReport;
+#endif // !UE_BUILD_SHIPPING
+}
+
+#if !UE_BUILD_SHIPPING
+bool FNavigationDirtyAreasController::ShouldReportOversizedDirtyArea() const
+{ 
+	return bNavigationBuildLocked == false && bCanReportOversizedDirtyArea && DirtyAreaWarningSizeThreshold >= 0.0f;
+}
+#endif // !UE_BUILD_SHIPPING
+
+
 void FNavigationDirtyAreasController::Reset()
 {
 	// discard all pending dirty areas, we are going to rebuild navmesh anyway 
diff --git a/Engine/Source/Runtime/NavigationSystem/Private/NavigationSystem.cpp b/Engine/Source/Runtime/NavigationSystem/Private/NavigationSystem.cpp
--- a/Engine/Source/Runtime/NavigationSystem/Private/NavigationSystem.cpp
+++ b/Engine/Source/Runtime/NavigationSystem/Private/NavigationSystem.cpp
@@ -258,6 +258,7 @@ UNavigationSystemV1::UNavigationSystemV1(const FObjectInitializer& ObjectInitial
 	, bTickWhilePaused(false)
 	, bWholeWorldNavigable(false)
 	, bSkipAgentHeightCheckWhenPickingNavData(false)
+	, DirtyAreaWarningSizeThreshold(-1.0f)
 	, OperationMode(FNavigationSystemRunMode::InvalidMode)
 	, NavBuildingLockFlags(0)
 	, InitialNavBuildingLockFlags(0)
@@ -555,6 +556,8 @@ void UNavigationSystemV1::PostInitProperties()
 			FNavDataConfig& SupportedAgentConfig = SupportedAgents[AgentIndex];
 			SetSupportedAgentsNavigationClass(AgentIndex, SupportedAgentConfig.GetNavDataClass<ANavigationData>());
 		}
+
+		DefaultDirtyAreasController.SetDirtyAreaWarningSizeThreshold(DirtyAreaWarningSizeThreshold);
 	
 		if (bInitialBuildingLocked)
 		{
@@ -715,7 +718,6 @@ void UNavigationSystemV1::OnBeginTearingDown()
 
 void UNavigationSystemV1::OnWorldInitDone(FNavigationSystemRunMode Mode)
 {
-	static const bool bSkipRebuildInEditor = true;
 	OperationMode = Mode;
 	DoInitialSetup();
 	
@@ -747,7 +749,7 @@ void UNavigationSystemV1::OnWorldInitDone(FNavigationSystemRunMode Mode)
 
 		if (OperationMode == FNavigationSystemRunMode::EditorMode)
 		{
-			RemoveNavigationBuildLock(InitialNavBuildingLockFlags, bSkipRebuildInEditor);
+			RemoveNavigationBuildLock(InitialNavBuildingLockFlags, ELockRemovalRebuildAction::RebuildIfNotInEditor);
 		}
 	}
 	else
@@ -809,7 +811,7 @@ void UNavigationSystemV1::OnWorldInitDone(FNavigationSystemRunMode Mode)
 		if (OperationMode == FNavigationSystemRunMode::EditorMode)
 		{
 			// don't lock navigation building in editor
-			RemoveNavigationBuildLock(InitialNavBuildingLockFlags, bSkipRebuildInEditor);
+			RemoveNavigationBuildLock(InitialNavBuildingLockFlags, ELockRemovalRebuildAction::RebuildIfNotInEditor);
 		}
 
 		// See if any of registered navigation data needs NavOctree
@@ -866,6 +868,9 @@ void UNavigationSystemV1::OnWorldInitDone(FNavigationSystemRunMode Mode)
 		DefaultDirtyAreasController.DirtyAreas.Empty();
 	}
 
+	// Report oversized dirty areas only in game mode
+	DefaultDirtyAreasController.SetCanReportOversizedDirtyArea(Mode == FNavigationSystemRunMode::GameMode);
+
 	bWorldInitDone = true;
 	OnNavigationInitDone.Broadcast();
 }
@@ -1030,8 +1035,7 @@ void UNavigationSystemV1::SetNavigationAutoUpdateEnabled(bool bNewEnable, UNavig
 
 			if (bCurrentIsEnabled)
 			{
-				const bool bSkipRebuildsInEditor = false;
-				NavSystem->RemoveNavigationBuildLock(ENavigationBuildLock::NoUpdateInEditor, bSkipRebuildsInEditor);
+				NavSystem->RemoveNavigationBuildLock(ENavigationBuildLock::NoUpdateInEditor);
 			}
 			else
 			{
@@ -3259,37 +3263,50 @@ void UNavigationSystemV1::OnPIEEnd()
 	{
 		bAsyncBuildPaused = false;
 		// there's no need to request while navigation rebuilding just because PIE has ended
-		RemoveNavigationBuildLock(ENavigationBuildLock::NoUpdateInEditor, /*bSkipRebuildInEditor=*/true);
+		RemoveNavigationBuildLock(ENavigationBuildLock::NoUpdateInEditor, ELockRemovalRebuildAction::RebuildIfNotInEditor);
 	}
 }
 
 void UNavigationSystemV1::AddNavigationBuildLock(uint8 Flags)
 {
+	const bool bWasLocked = IsNavigationBuildingLocked();
+
 	NavBuildingLockFlags |= Flags;
 
-	UE_LOG(LogNavigation, Verbose, TEXT("UNavigationSystemV1::AddNavigationBuildLock IsLocked=%s"), IsNavigationBuildingLocked() ? TEXT("true") : TEXT("false"));
+	const bool bIsLocked = IsNavigationBuildingLocked();
+	UE_LOG(LogNavigation, Verbose, TEXT("UNavigationSystemV1::AddNavigationBuildLock WasLocked=%s IsLocked=%s"), *LexToString(bWasLocked), *LexToString(bIsLocked));
+	if (!bWasLocked && bIsLocked)
+	{
+		DefaultDirtyAreasController.OnNavigationBuildLocked();
+	}
 }
 
-void UNavigationSystemV1::RemoveNavigationBuildLock(uint8 Flags, bool bSkipRebuildInEditor)
+void UNavigationSystemV1::RemoveNavigationBuildLock(uint8 Flags, const ELockRemovalRebuildAction RebuildAction /*= ELockRemovalRebuildAction::Rebuild*/)
 {
 	const bool bWasLocked = IsNavigationBuildingLocked();
 
 	NavBuildingLockFlags &= ~Flags;
 
 	const bool bIsLocked = IsNavigationBuildingLocked();
-	const bool bSkipRebuild = (OperationMode == FNavigationSystemRunMode::EditorMode) && bSkipRebuildInEditor;
-
-	UE_LOG(LogNavigation, Verbose, TEXT("UNavigationSystemV1::RemoveNavigationBuildLock bWasLocked=%s IsLocked=%s"), bWasLocked ? TEXT("true") : TEXT("false"), bIsLocked ? TEXT("true") : TEXT("false"));
-
-	if (bWasLocked && !bIsLocked && !bSkipRebuild)
+	UE_LOG(LogNavigation, Verbose, TEXT("UNavigationSystemV1::RemoveNavigationBuildLock WasLocked=%s IsLocked=%s"), *LexToString(bWasLocked), *LexToString(bIsLocked));
+	if (bWasLocked && !bIsLocked)
 	{
-		RebuildAll();
+		DefaultDirtyAreasController.OnNavigationBuildUnlocked();
+
+		const bool bRebuild = 
+			(RebuildAction == ELockRemovalRebuildAction::RebuildIfNotInEditor && (OperationMode != FNavigationSystemRunMode::EditorMode)) || 
+			(RebuildAction == ELockRemovalRebuildAction::Rebuild);
+		
+		if (bRebuild)
+		{
+			RebuildAll();
+		}
 	}
 }
 
 void UNavigationSystemV1::SetNavigationOctreeLock(bool bLock)
 {
-	UE_LOG(LogNavigation, Verbose, TEXT("UNavigationSystemV1::SetNavigationOctreeLock IsLocked=%s"), bLock ? TEXT("true") : TEXT("false"));
+	UE_LOG(LogNavigation, Verbose, TEXT("UNavigationSystemV1::SetNavigationOctreeLock IsLocked=%s"), *LexToString(bLock));
 
 	DefaultOctreeController.SetNavigationOctreeLock(bLock);
 }
@@ -3550,7 +3567,7 @@ void UNavigationSystemV1::OnHotReload(bool bWasTriggeredAutomatically)
 
 		if (bInitialBuildingLocked)
 		{
-			RemoveNavigationBuildLock(ENavigationBuildLock::InitialLock, /*bSkipRebuildInEditor=*/true);
+			RemoveNavigationBuildLock(ENavigationBuildLock::InitialLock, ELockRemovalRebuildAction::RebuildIfNotInEditor);
 		}
 	}
 }
diff --git a/Engine/Source/Runtime/NavigationSystem/Public/NavigationDirtyAreasController.h b/Engine/Source/Runtime/NavigationSystem/Public/NavigationDirtyAreasController.h
--- a/Engine/Source/Runtime/NavigationSystem/Public/NavigationDirtyAreasController.h
+++ b/Engine/Source/Runtime/NavigationSystem/Public/NavigationDirtyAreasController.h
@@ -21,8 +21,17 @@ struct NAVIGATIONSYSTEM_API FNavigationDirtyAreasController
 	uint8 bCanAccumulateDirtyAreas : 1;
 #if !UE_BUILD_SHIPPING
 	uint8 bDirtyAreasReportedWhileAccumulationLocked : 1;
+private:
+	uint8 bCanReportOversizedDirtyArea : 1;
+	uint8 bNavigationBuildLocked : 1;
+
+	/** -1 by default, if set to a positive value dirty area with bounds size over that threshold will be logged */
+	float DirtyAreaWarningSizeThreshold = -1.f;
+
+	bool ShouldReportOversizedDirtyArea() const;
 #endif // !UE_BUILD_SHIPPING
 
+public:
 	FNavigationDirtyAreasController();
 
 	void Reset();
@@ -31,11 +40,25 @@ struct NAVIGATIONSYSTEM_API FNavigationDirtyAreasController
 	void ForceRebuildOnNextTick();
 
 	void Tick(float DeltaSeconds, const TArray<ANavigationData*>& NavDataSet, bool bForceRebuilding = false);
-	void AddArea(const FBox& NewArea, int32 Flags);
+
+	/** Add a dirty area to the queue based on the provided bounds and flags.
+	 * Bounds must be valid and non empty otherwise the request will be ignored and a warning reported.
+	 * Accumulation must be allowed and flags valid otherwise the add is ignored.
+	 *	@param NewArea Bounding box of the affected area
+	 *	@param Flags Indicates the type of modification applied to the area
+	 *	@param ObjectProviderFunc Optional function to retrieve source object that can be use for error reporting
+	 */
+	void AddArea(const FBox& NewArea, const int32 Flags, const TFunction<UObject*()>& ObjectProviderFunc = nullptr);
 	
 	bool IsDirty() const { return GetNumDirtyAreas() > 0; }
 	int32 GetNumDirtyAreas() const { return DirtyAreas.Num(); }
 
+	void OnNavigationBuildLocked();
+	void OnNavigationBuildUnlocked();
+
+	void SetCanReportOversizedDirtyArea(const bool bCanReport);
+	void SetDirtyAreaWarningSizeThreshold(const float Threshold);
+
 #if !UE_BUILD_SHIPPING
 	bool HadDirtyAreasReportedWhileAccumulationLocked() const { return bCanAccumulateDirtyAreas == false && bDirtyAreasReportedWhileAccumulationLocked; }
 #endif // UE_BUILD_SHIPPING
diff --git a/Engine/Source/Runtime/NavigationSystem/Public/NavigationSystem.h b/Engine/Source/Runtime/NavigationSystem/Public/NavigationSystem.h
--- a/Engine/Source/Runtime/NavigationSystem/Public/NavigationSystem.h
+++ b/Engine/Source/Runtime/NavigationSystem/Public/NavigationSystem.h
@@ -168,9 +168,6 @@ public:
 	uint32 bSkipAgentHeightCheckWhenPickingNavData:1;
 
 protected:
-	
-	UPROPERTY(EditDefaultsOnly, Category = "NavigationSystem", config)
-	ENavDataGatheringModeConfig DataGatheringMode;
 
 	/** If set to true navigation will be generated only around registered "navigation enforcers"
 	*	This has a range of consequences (including how navigation octree operates) so it needs to
@@ -179,12 +176,19 @@ protected:
 	*	@see RegisterNavigationInvoker
 	*/
 	UPROPERTY(EditDefaultsOnly, Category = "Navigation Enforcing", config)
-	uint32 bGenerateNavigationOnlyAroundNavigationInvokers : 1;
-	
+	uint32 bGenerateNavigationOnlyAroundNavigationInvokers:1;
+
 	/** Minimal time, in seconds, between active tiles set update */
 	UPROPERTY(EditAnywhere, Category = "Navigation Enforcing", meta = (ClampMin = "0.1", UIMin = "0.1", EditCondition = "bGenerateNavigationOnlyAroundNavigationInvokers"), config)
 	float ActiveTilesUpdateInterval;
 
+	UPROPERTY(EditDefaultsOnly, Category = "NavigationSystem", config)
+	ENavDataGatheringModeConfig DataGatheringMode;
+
+	/** -1 by default, if set to a positive value dirty areas with any dimensions in 2d over the threshold created at runtime will be logged */
+	UPROPERTY(config, EditAnywhere, AdvancedDisplay, Category = NavigationSystem, meta = (ClampMin = "-1.0", UIMin = "-1.0"))
+	float DirtyAreaWarningSizeThreshold;
+
 	UPROPERTY(config, EditAnywhere, Category = Agents)
 	TArray<FNavDataConfig> SupportedAgents;
 
@@ -779,8 +783,14 @@ public:
 	/** adds BSP collisions of currently streamed in levels to octree */
 	void InitializeLevelCollisions();
 
+	enum class ELockRemovalRebuildAction
+	{
+		Rebuild,
+		RebuildIfNotInEditor,
+		NoRebuild
+	};
 	void AddNavigationBuildLock(uint8 Flags);
-	void RemoveNavigationBuildLock(uint8 Flags, bool bSkipRebuildInEditor = false);
+	void RemoveNavigationBuildLock(uint8 Flags, const ELockRemovalRebuildAction RebuildAction = ELockRemovalRebuildAction::Rebuild);
 
 	void SetNavigationOctreeLock(bool bLock);
 
@@ -1074,6 +1084,9 @@ public:
 	UE_DEPRECATED(4.24, "This member is deprecated and no longer used.  Please access DirtyAreasController instead")
 	uint8 bDirtyAreasReportedWhileAccumulationLocked : 1;
 #endif // !UE_BUILD_SHIPPING
+
+	UE_DEPRECATED(4.26, "This version of RemoveNavigationBuildLock is deprecated. Please use the new version")
+	void RemoveNavigationBuildLock(uint8 Flags, bool bSkipRebuildInEditor) { RemoveNavigationBuildLock(Flags, bSkipRebuildInEditor ? ELockRemovalRebuildAction::RebuildIfNotInEditor : ELockRemovalRebuildAction::Rebuild);}
 };
 
 //----------------------------------------------------------------------//
