Fix issue where AI would generate straight path on Android and Switch (ignoring obstacles)
This was caused by the unwanted and undefined behavior of passing fltmax to some int flags (causing flags to SkipStringPulling)

Integration from CL 11616045

diff --git a/Engine/Source/Runtime/NavigationSystem/Private/NavigationData.cpp b/Engine/Source/Runtime/NavigationSystem/Private/NavigationData.cpp
--- a/Engine/Source/Runtime/NavigationSystem/Private/NavigationData.cpp
+++ b/Engine/Source/Runtime/NavigationSystem/Private/NavigationData.cpp
@@ -20,7 +20,7 @@
 // FPathFindingQuery
 //----------------------------------------------------------------------//
 FPathFindingQuery::FPathFindingQuery(const UObject* InOwner, const ANavigationData& InNavData, const FVector& Start, const FVector& End, FSharedConstNavQueryFilter SourceQueryFilter, FNavPathSharedPtr InPathInstanceToFill, const float CostLimit) :
-	FPathFindingQueryData(InOwner, Start, End, SourceQueryFilter, CostLimit),
+	FPathFindingQueryData(InOwner, Start, End, SourceQueryFilter, 0 /*InNavDataFlags*/, true /*bInAllowPartialPaths*/, CostLimit),
 	NavData(&InNavData), PathInstanceToFill(InPathInstanceToFill), NavAgentProperties(InNavData.GetConfig())
 {
 	if (!QueryFilter.IsValid() && NavData.IsValid())
@@ -30,7 +30,7 @@ FPathFindingQuery::FPathFindingQuery(const UObject* InOwner, const ANavigationDa
 }
 
 FPathFindingQuery::FPathFindingQuery(const INavAgentInterface& InNavAgent, const ANavigationData& InNavData, const FVector& Start, const FVector& End, FSharedConstNavQueryFilter SourceQueryFilter, FNavPathSharedPtr InPathInstanceToFill, const float CostLimit) :
-	FPathFindingQueryData(Cast<UObject>(&InNavAgent), Start, End, SourceQueryFilter, CostLimit),
+	FPathFindingQueryData(Cast<UObject>(&InNavAgent), Start, End, SourceQueryFilter, 0 /*InNavDataFlags*/, true /*bInAllowPartialPaths*/, CostLimit),
 	NavData(&InNavData), PathInstanceToFill(InPathInstanceToFill), NavAgentProperties(InNavAgent.GetNavAgentPropertiesRef())
 {
 	if (!QueryFilter.IsValid() && NavData.IsValid())
