Work to properly account for time slicing accross multiple generators. The time slicing is now calculated in the navigation system.

The general idea here is to tick any non time sliced generators once per frame. Time sliced generators we aim to tick one per frame and move to the next, next frame.

In the case where one time sliced generator doesn't use the whole time slice we move to the next time sliced generator. That generator will only be considered to have a full frames processing if either it runs out of work or uses a large % of the time slice. Depending we either tick it again next frame or go to the next time sliced generator (next frame).

diff --git a/Engine/Source/Runtime/Engine/Public/AI/NavDataGenerator.h b/Engine/Source/Runtime/Engine/Public/AI/NavDataGenerator.h
--- a/Engine/Source/Runtime/Engine/Public/AI/NavDataGenerator.h
+++ b/Engine/Source/Runtime/Engine/Public/AI/NavDataGenerator.h
@@ -19,7 +19,10 @@ public:
 	/** Cancels build, may block until current running async tasks are finished */
 	virtual void CancelBuild() {};
 
-	/**  */
+	/** Ticks navigation build
+	 *  If the generator is set to time sliced rebuild then this function will only get called when
+	 *  there is sufficient time (effectively roughly once in n frames where n is the number of time sliced nav data / generators currently building)
+	 */
 	virtual void TickAsyncBuild(float DeltaSeconds) {};
 	
 	/**  */
@@ -31,6 +34,8 @@ public:
 	/** determines whether this generator is performing navigation building actions at the moment*/
 	virtual bool IsBuildInProgress(bool bCheckDirtyToo = false) const { return false; }	
 
+	virtual bool GetTimeSliceData(int32& OutNumRemainingBuildTasks, double& OutCurrentBuildTaskDuration) const { OutNumRemainingBuildTasks = 0; OutCurrentBuildTaskDuration = 0.; return false; }
+
 	/** Returns number of remaining tasks till build is complete
 	 */
 	virtual int32 GetNumRemaningBuildTasks() const { return 0; };
diff --git a/Engine/Source/Runtime/NavigationSystem/Private/NavMesh/RecastNavMeshGenerator.cpp b/Engine/Source/Runtime/NavigationSystem/Private/NavMesh/RecastNavMeshGenerator.cpp
--- a/Engine/Source/Runtime/NavigationSystem/Private/NavMesh/RecastNavMeshGenerator.cpp
+++ b/Engine/Source/Runtime/NavigationSystem/Private/NavMesh/RecastNavMeshGenerator.cpp
@@ -1822,41 +1822,13 @@ static FBox CalculateTileBounds(int32 X, int32 Y, const FVector& RcNavMeshOrigin
 	return TileBox;
 }
 
-void FTimeSlicer::SetTimeSliceDuration(double SliceDuration)
-{
-	TimeSliceDuration = SliceDuration;
-}
-
-void FTimeSlicer::StartTimeSlice()
-{
-	TimeSliceStartTime = FPlatformTime::Seconds();
-	bTimeSliceFinishedCached = false;
-}
-
-double FTimeSlicer::GetStartTime() const
-{
-	return TimeSliceStartTime;
-}
-
-bool FTimeSlicer::TestTimeSliceFinished() const
-{
-	ensureMsgf(!bTimeSliceFinishedCached, TEXT("Testing time slice is finished when we have already confirmed that!"));
-
-	bTimeSliceFinishedCached = FPlatformTime::Seconds() - TimeSliceStartTime >= TimeSliceDuration;
-	return bTimeSliceFinishedCached;
-}
-
-bool FTimeSlicer::IsTimeSliceFinishedCached() const
-{
-	return bTimeSliceFinishedCached;
-}
 
 //----------------------------------------------------------------------//
 // FRecastTileGenerator
 //----------------------------------------------------------------------//
 
 FRecastTileGenerator::FRecastTileGenerator(FRecastNavMeshGenerator& ParentGenerator, const FIntPoint& Location)
-	: TimeSlicer(ParentGenerator.GetTimeSlicer())
+	: TimeSliceManager(ParentGenerator.GetTimeSliceManager())
 {
 	bUpdateGeometry = true;
 	bHasLowAreaModifiers = false;
@@ -1983,6 +1955,8 @@ ETimeSliceWorkResult FRecastTileGenerator::DoWorkTimeSliced()
 	TSharedPtr<FNavDataGenerator, ESPMode::ThreadSafe> ParentGenerator = ParentGeneratorWeakPtr.Pin();
 	ETimeSliceWorkResult WorkResult = ETimeSliceWorkResult::Succeeded;
 
+	check(TimeSliceManager);
+
 	if (ParentGenerator.IsValid())
 	{
 		switch (DoWorkTimeSlicedState)
@@ -2002,7 +1976,7 @@ ETimeSliceWorkResult FRecastTileGenerator::DoWorkTimeSliced()
 				const bool bHadNavigationRelevantData = DoAsyncGeometryGathering();
 
 				//check bHadNavigationRelevantData as an optimization so we don't call TestTimeSliceFinished when its unnecessary 
-				if (bHadNavigationRelevantData && TimeSlicer.TestTimeSliceFinished())
+				if (bHadNavigationRelevantData && TimeSliceManager->GetTimeSlicer().TestTimeSliceFinished())
 				{
 					return ETimeSliceWorkResult::CallAgainNextTimeSlice;
 				}
@@ -2566,6 +2540,8 @@ ETimeSliceWorkResult FRecastTileGenerator::GenerateTileTimeSliced()
 	FNavMeshBuildContext BuildContext(*this);
 	ETimeSliceWorkResult WorkResult = ETimeSliceWorkResult::Succeeded;
 
+	check(TimeSliceManager);
+
 	switch (GenerateTileTimeSlicedState)
 	{
 	case EGenerateTileTimeSlicedState::Invalid:
@@ -2592,7 +2568,7 @@ ETimeSliceWorkResult FRecastTileGenerator::GenerateTileTimeSliced()
 				return ETimeSliceWorkResult::Failed;
 			}
 
-			if (TimeSlicer.IsTimeSliceFinishedCached())
+			if (TimeSliceManager->GetTimeSlicer().IsTimeSliceFinishedCached())
 			{
 				return ETimeSliceWorkResult::CallAgainNextTimeSlice;
 			}
@@ -2716,6 +2692,8 @@ ETimeSliceWorkResult FRecastTileGenerator::RasterizeGeometryRecastTimeSliced(FNa
 {
 	QUICK_SCOPE_CYCLE_COUNTER(STAT_Navigation_RasterizeGeometryRecast);
 
+	check(TimeSliceManager);
+
 	const int32 NumFaces = Indices.Num() / 3;
 	const int32 NumVerts = Coords.Num() / 3;
 
@@ -2733,7 +2711,7 @@ ETimeSliceWorkResult FRecastTileGenerator::RasterizeGeometryRecastTimeSliced(FNa
 
 		RasterizeGeomRecastState = ERasterizeGeomRecastTimeSlicedState::RasterizeTriangles;
 
-		if (TimeSlicer.TestTimeSliceFinished())
+		if (TimeSliceManager->GetTimeSlicer().TestTimeSliceFinished())
 		{
 			return ETimeSliceWorkResult::CallAgainNextTimeSlice;
 		}
@@ -2752,7 +2730,7 @@ ETimeSliceWorkResult FRecastTileGenerator::RasterizeGeometryRecastTimeSliced(FNa
 		//reset this so next call we start by marking walkable triangles
 		RasterizeGeomRecastState = ERasterizeGeomRecastTimeSlicedState::MarkWalkableTriangles;
 
-		TimeSlicer.TestTimeSliceFinished();
+		TimeSliceManager->GetTimeSlicer().TestTimeSliceFinished();
 	}
 	break;
 	default:
@@ -2817,6 +2795,9 @@ void FRecastTileGenerator::RasterizeGeometryTransformCoords(const TArray<float>&
 ETimeSliceWorkResult FRecastTileGenerator::RasterizeGeometryTimeSliced(FNavMeshBuildContext& BuildContext, const TArray<float>& Coords, const TArray<int32>& Indices, const FTransform& LocalToWorld, const rcRasterizationFlags RasterizationFlags, FTileRasterizationContext& RasterContext)
 {
 	QUICK_SCOPE_CYCLE_COUNTER(STAT_Navigation_RasterizeGeometry);
+
+	check(TimeSliceManager);
+
 	ETimeSliceWorkResult WorkResult = ETimeSliceWorkResult::Succeeded;
 
 	switch (RasterizeGeomState)
@@ -2827,7 +2808,7 @@ ETimeSliceWorkResult FRecastTileGenerator::RasterizeGeometryTimeSliced(FNavMeshB
 
 		RasterizeGeomState = ERasterizeGeomTimeSlicedState::RasterizeGeometryRecast;
 
-		if (TimeSlicer.TestTimeSliceFinished())
+		if (TimeSliceManager->GetTimeSlicer().TestTimeSliceFinished())
 		{
 			return ETimeSliceWorkResult::CallAgainNextTimeSlice;
 		}
@@ -2865,6 +2846,8 @@ ETimeSliceWorkResult FRecastTileGenerator::RasterizeTrianglesTimeSliced(FNavMesh
 	// Rasterize geometry
 	SCOPE_CYCLE_COUNTER(STAT_Navigation_RecastRasterizeTriangles)
 
+	check(TimeSliceManager);
+
 	while (RasterizeTrianglesTimeSlicedRawGeomIdx < RawGeometry.Num())
 	{
 		const FRecastRawGeometryElement& Element = RawGeometry[RasterizeTrianglesTimeSlicedRawGeomIdx];
@@ -2877,7 +2860,7 @@ ETimeSliceWorkResult FRecastTileGenerator::RasterizeTrianglesTimeSliced(FNavMesh
 			
 				//the original code just kept calling the RasterizeGeometry() functions and had no return type, 
 				//so we will process the next layer (if we are not needing to process this layer again next time slice) 
-				if (TimeSlicer.IsTimeSliceFinishedCached())
+				if (TimeSliceManager->GetTimeSlicer().IsTimeSliceFinishedCached())
 				{
 					if (WorkResult != ETimeSliceWorkResult::CallAgainNextTimeSlice)
 					{
@@ -2896,7 +2879,7 @@ ETimeSliceWorkResult FRecastTileGenerator::RasterizeTrianglesTimeSliced(FNavMesh
 		{
 			const ETimeSliceWorkResult WorkResult = RasterizeGeometryRecastTimeSliced(BuildContext, Element.GeomCoords, Element.GeomIndices, Element.RasterizationFlags, RasterContext);
 	
-			if (TimeSlicer.IsTimeSliceFinishedCached())
+			if (TimeSliceManager->GetTimeSlicer().IsTimeSliceFinishedCached())
 			{
 				if (WorkResult != ETimeSliceWorkResult::CallAgainNextTimeSlice)
 				{
@@ -3174,6 +3157,8 @@ ETimeSliceWorkResult FRecastTileGenerator::GenerateCompressedLayersTimeSliced(FN
 {
 	SCOPE_CYCLE_COUNTER(STAT_Navigation_RecastBuildCompressedLayers);
 
+	check(TimeSliceManager);
+
 	FTileRasterizationContext* RasterContext = GenCompressedlayersTimeSlicedRasterContext.Get();
 
 	switch (GenCompressedLayersTimeSlicedState)
@@ -3203,7 +3188,7 @@ ETimeSliceWorkResult FRecastTileGenerator::GenerateCompressedLayersTimeSliced(FN
 
 		GenCompressedLayersTimeSlicedState = EGenerateCompressedLayersTimeSliced::RasterizeTriangles;
 
-		if (TimeSlicer.TestTimeSliceFinished())
+		if (TimeSliceManager->GetTimeSlicer().TestTimeSliceFinished())
 		{
 			return ETimeSliceWorkResult::CallAgainNextTimeSlice;
 		}
@@ -3218,7 +3203,7 @@ ETimeSliceWorkResult FRecastTileGenerator::GenerateCompressedLayersTimeSliced(FN
 			GenCompressedLayersTimeSlicedState = EGenerateCompressedLayersTimeSliced::EmptyLayers;
 		}
 
-		if (TimeSlicer.IsTimeSliceFinishedCached())
+		if (TimeSliceManager->GetTimeSlicer().IsTimeSliceFinishedCached())
 		{
 			return ETimeSliceWorkResult::CallAgainNextTimeSlice;
 		}
@@ -3245,7 +3230,7 @@ ETimeSliceWorkResult FRecastTileGenerator::GenerateCompressedLayersTimeSliced(FN
 		{
 			ApplyVoxelFilter(RasterContext->SolidHF, TileConfig.walkableRadius);
 
-			if (TimeSlicer.TestTimeSliceFinished())
+			if (TimeSliceManager->GetTimeSlicer().TestTimeSliceFinished())
 			{
 				return ETimeSliceWorkResult::CallAgainNextTimeSlice;
 			}
@@ -3257,7 +3242,7 @@ ETimeSliceWorkResult FRecastTileGenerator::GenerateCompressedLayersTimeSliced(FN
 
 		GenCompressedLayersTimeSlicedState = EGenerateCompressedLayersTimeSliced::CompactHeightField;
 
-		if (TimeSlicer.TestTimeSliceFinished())
+		if (TimeSliceManager->GetTimeSlicer().TestTimeSliceFinished())
 		{
 			return ETimeSliceWorkResult::CallAgainNextTimeSlice;
 		}
@@ -3274,7 +3259,7 @@ ETimeSliceWorkResult FRecastTileGenerator::GenerateCompressedLayersTimeSliced(FN
 
 		GenCompressedLayersTimeSlicedState = EGenerateCompressedLayersTimeSliced::ErodeWalkable;
 
-		if (TimeSlicer.TestTimeSliceFinished())
+		if (TimeSliceManager->GetTimeSlicer().TestTimeSliceFinished())
 		{
 			return ETimeSliceWorkResult::CallAgainNextTimeSlice;
 		}
@@ -3291,7 +3276,7 @@ ETimeSliceWorkResult FRecastTileGenerator::GenerateCompressedLayersTimeSliced(FN
 
 		GenCompressedLayersTimeSlicedState = EGenerateCompressedLayersTimeSliced::BuildLayers;
 
-		if (TimeSlicer.TestTimeSliceFinished())
+		if (TimeSliceManager->GetTimeSlicer().TestTimeSliceFinished())
 		{
 			return ETimeSliceWorkResult::CallAgainNextTimeSlice;
 		}
@@ -3301,7 +3286,7 @@ ETimeSliceWorkResult FRecastTileGenerator::GenerateCompressedLayersTimeSliced(FN
 		const bool bRecastBuildLayers = RecastBuildLayers(BuildContext, *RasterContext);
 
 		//this could have done a fair amount of work either way so check time slice
-		TimeSlicer.TestTimeSliceFinished();
+		TimeSliceManager->GetTimeSlicer().TestTimeSliceFinished();
 
 		if (!bRecastBuildLayers)
 		{
@@ -3312,7 +3297,7 @@ ETimeSliceWorkResult FRecastTileGenerator::GenerateCompressedLayersTimeSliced(FN
 
 		GenCompressedLayersTimeSlicedState = EGenerateCompressedLayersTimeSliced::BuildTileCache;
 
-		if (TimeSlicer.IsTimeSliceFinishedCached())
+		if (TimeSliceManager->GetTimeSlicer().IsTimeSliceFinishedCached())
 		{
 			return ETimeSliceWorkResult::CallAgainNextTimeSlice;
 		}
@@ -3324,7 +3309,7 @@ ETimeSliceWorkResult FRecastTileGenerator::GenerateCompressedLayersTimeSliced(FN
 		const bool bRecastBuildTileCache = RecastBuildTileCache(BuildContext, *RasterContext);
 	
 		//this could have done a fair amount of work either way so check time slice
-		TimeSlicer.TestTimeSliceFinished();
+		TimeSliceManager->GetTimeSlicer().TestTimeSliceFinished();
 
 		if (!bRecastBuildTileCache)
 		{
@@ -3718,6 +3703,8 @@ ETimeSliceWorkResult FRecastTileGenerator::GenerateNavigationDataTimeSliced(FNav
 {
 	SCOPE_CYCLE_COUNTER(STAT_Navigation_RecastBuildNavigation);
 
+	check(TimeSliceManager);
+
 	FTileCacheCompressor TileCompressor;
 	ETimeSliceWorkResult WorkResult = ETimeSliceWorkResult::Succeeded;
 	dtStatus status = DT_SUCCESS;
@@ -3748,7 +3735,7 @@ ETimeSliceWorkResult FRecastTileGenerator::GenerateNavigationDataTimeSliced(FNav
 				continue;
 			}
 
-			if (TimeSlicer.IsTimeSliceFinishedCached())
+			if (TimeSliceManager->GetTimeSlicer().IsTimeSliceFinishedCached())
 			{
 				WorkResult = ETimeSliceWorkResult::CallAgainNextTimeSlice;
 				break;
@@ -3757,7 +3744,7 @@ ETimeSliceWorkResult FRecastTileGenerator::GenerateNavigationDataTimeSliced(FNav
 			const bool bGenDataLayer = GenerateNavigationDataLayer(BuildContext, TileCompressor, *GenNavDataTimeSlicedAllocator, *GenNavDataTimeSlicedGenerationContext, GenNavDataLayerTimeSlicedIdx);
 
 			//carry on iterating but don't do any more work if the time slice is finished (as we may not need to in which case we can avoid calling this function again)
-			TimeSlicer.TestTimeSliceFinished();
+			TimeSliceManager->GetTimeSlicer().TestTimeSliceFinished();
 
 			if (!bGenDataLayer)
 			{
@@ -4219,6 +4206,18 @@ void FRecastNavMeshGenerator::Init()
 {
 	check(DestNavMesh);
 
+	UNavigationSystemV1* NavSys = FNavigationSystem::GetCurrent<UNavigationSystemV1>(GetWorld());
+
+	if (NavSys)
+	{
+		SyncTimeSlicedData.TimeSliceManager = &(NavSys->GetNavRegenTimeSliceManager());
+	}
+	else
+	{
+		//no time slice manager no time sliced regen
+		SyncTimeSlicedData.bTimeSliceRegenActive = false;
+	}
+
 	ConfigureBuildProperties(Config);
 
 	BBoxGrowth = FVector(2.0f * Config.borderSize * Config.cs);
@@ -4308,7 +4307,6 @@ void FRecastNavMeshGenerator::Init()
 		ConstructTiledNavMesh();
 
 		// mark all the areas we need to update, which is the whole (known) navigable space if not restricted to active tiles
-		const UNavigationSystemV1* NavSys = FNavigationSystem::GetCurrent<UNavigationSystemV1>(GetWorld());
 		if (NavSys && NavSys->IsActiveTilesGenerationEnabled() == false)
 		{
 			MarkNavBoundsDirty();
@@ -4854,10 +4852,6 @@ TArray<uint32> FRecastNavMeshGenerator::RemoveTileLayers(const int32 TileX, cons
 
 FRecastNavMeshGenerator::FSyncTimeSlicedData::FSyncTimeSlicedData()
 	: CurrentTileRegenDuration(0.)
-	, MinTimeSliceDuration(0.00075)
-	, MaxTimeSliceDuration(0.004)
-	, RealTimeSecsLastCall(-1.f)
-	, MaxDesiredTileRegenDuration(0.7f)
 #if TIME_SLICE_NAV_REGEN
 	, bTimeSliceRegenActive(true)
 	, bNextTimeSliceRegenActive(true)
@@ -4868,7 +4862,7 @@ FRecastNavMeshGenerator::FSyncTimeSlicedData::FSyncTimeSlicedData()
 	, ProcessTileTasksSyncState(EProcessTileTasksSyncTimeSlicedState::Init)
 	, AddGeneratedTilesState(EAddGeneratedTilesTimeSlicedState::Init)
 	, AddGenTilesLayerIndex(0)
-	, TimeSlicer(0.0025)
+	, TimeSliceManager(nullptr)
 {
 }
 
@@ -4967,6 +4961,8 @@ ETimeSliceWorkResult FRecastNavMeshGenerator::AddGeneratedTilesTimeSliced(FRecas
 {
 	SCOPE_CYCLE_COUNTER(STAT_Navigation_RecastAddGeneratedTiles);
 
+	check(SyncTimeSlicedData.TimeSliceManager);
+
 	const int32 TileX = TileGenerator.GetTileX();
 	const int32 TileY = TileGenerator.GetTileY();
 	dtNavMesh* DetourMesh = DestNavMesh->GetRecastNavMeshImpl()->GetRecastMesh();
@@ -5004,7 +5000,7 @@ ETimeSliceWorkResult FRecastNavMeshGenerator::AddGeneratedTilesTimeSliced(FRecas
 			{
 				if (TileGenerator.IsLayerChanged(SyncTimeSlicedData.AddGenTilesLayerIndex))
 				{
-					if (SyncTimeSlicedData.TimeSlicer.IsTimeSliceFinishedCached())
+					if (SyncTimeSlicedData.TimeSliceManager->GetTimeSlicer().IsTimeSliceFinishedCached())
 					{
 						WorkResult = ETimeSliceWorkResult::CallAgainNextTimeSlice;
 						break;
@@ -5012,7 +5008,7 @@ ETimeSliceWorkResult FRecastNavMeshGenerator::AddGeneratedTilesTimeSliced(FRecas
 
 					AddGeneratedTileLayer(SyncTimeSlicedData.AddGenTilesLayerIndex, TileGenerator, SyncTimeSlicedData.OldLayerTileIdMapCached, SyncTimeSlicedData.ResultTileIndicesCached);
 
-					SyncTimeSlicedData.TimeSlicer.TestTimeSliceFinished();
+					SyncTimeSlicedData.TimeSliceManager->GetTimeSlicer().TestTimeSliceFinished();
 				}
 			}
 		}
@@ -5557,8 +5553,9 @@ TArray<uint32> FRecastNavMeshGenerator::ProcessTileTasksSyncTimeSliced()
 	QUICK_SCOPE_CYCLE_COUNTER(STAT_RecastNavMeshGenerator_ProcessTileTasksSyncTimeSliced);
 	CSV_SCOPED_TIMING_STAT(NAVREGEN, ProcessTileTasksSyncTimeSliced);
 
+	check(SyncTimeSlicedData.TimeSliceManager);
+
 	TArray<uint32> UpdatedTiles;
-	const UWorld* World = GetWorld();
 	double TimeStartProcessingTileThisFrame = 0.;
 
 	auto HasWorkToDo = [this]()
@@ -5566,19 +5563,14 @@ TArray<uint32> FRecastNavMeshGenerator::ProcessTileTasksSyncTimeSliced()
 		return (PendingDirtyTiles.Num() > 0) || SyncTimeSlicedData.TileGeneratorSync.IsValid();
 	};
 
-
-	auto EndFunction = [&, this](bool bCalcTileRegenDuration) {
+	auto EndFunction = [&, this](bool bCalcTileRegenDuration, bool bStartedTimeSlice)
+	{
 		// Release memory, list could be quite big after map load
 		if (PendingDirtyTiles.Num() == 0)
 		{
 			PendingDirtyTiles.Empty(64);
 		}
 
-		if (World)
-		{
-			SyncTimeSlicedData.RealTimeSecsLastCall = World->GetRealTimeSeconds();
-		}
-
 		//this will only be true when we haven't finished generating this tile but are ending
 		//the function and need to record the TileRegenDuration so far for the tile
 		//being currently processed
@@ -5587,55 +5579,23 @@ TArray<uint32> FRecastNavMeshGenerator::ProcessTileTasksSyncTimeSliced()
 			SyncTimeSlicedData.CurrentTileRegenDuration += (FPlatformTime::Seconds() - TimeStartProcessingTileThisFrame);
 		}
 
+		if (bStartedTimeSlice)
+		{
+			SyncTimeSlicedData.TimeSliceManager->GetTimeSlicer().EndTimeSliceAndAdjustDuration();
+		}
+
 		return UpdatedTiles;
 	};
 
-	//Calculate the time slice duration
-	//Calc the MovingWindowDeltaTimeAverage this accounts for all scenarios we could be tile regening including unbounded frame rates or dropping frames as well as keeping
-	//calculation to an average which is fairly local temporally
-	if (World && SyncTimeSlicedData.RealTimeSecsLastCall >= 0.f)
-	{
-		const float DeltaTime = World->GetRealTimeSeconds() - SyncTimeSlicedData.RealTimeSecsLastCall;
-		SyncTimeSlicedData.MovingWindowDeltaTime.PushValue(DeltaTime);
-	}	
-	
+	const bool bHadWorktoDo = HasWorkToDo();
+
 	//only calculate the time slice and process tiles if we have work to do
-	if (HasWorkToDo())
+	if (bHadWorktoDo)
 	{
-		CSV_SCOPED_TIMING_STAT(NAVREGEN, ProcessTileTasksSyncTimeSlicedDoWork);
+		SyncTimeSlicedData.TimeSliceManager->GetTimeSlicer().StartTimeSlice();
 
 		const bool bGameStaticNavMesh = IsGameStaticNavMesh(DestNavMesh);
 
-		SyncTimeSlicedData.TimeSlicer.StartTimeSlice();
-
-		const float DeltaTimesAverage = (SyncTimeSlicedData.MovingWindowDeltaTime.GetAverage() > 0.f) ? SyncTimeSlicedData.MovingWindowDeltaTime.GetAverage() : (1.f / 30.f); //use default 33 ms
-
-		const double TileRegenTimesAverage = (SyncTimeSlicedData.MovingWindowTileRegenTime.GetAverage() > 0.) ? SyncTimeSlicedData.MovingWindowTileRegenTime.GetAverage() : 0.0025; //use default of 2.5 milli secs to regen a full tile
-
-		//calculate the max desired frames to regen all the tiles in PendingDirtyTiles
-		const float MaxDesiredFramesToRegen = FMath::FloorToFloat(SyncTimeSlicedData.MaxDesiredTileRegenDuration / DeltaTimesAverage);
-
-		//tiles to add to PendingDirtyTiles if the current tile is taking longer than average to regen
-		//we add 1 tile for however many times longer the current tile is taking compared with the moving window average
-		const int32 TilesToAddForLongCurrentTileRegen = (SyncTimeSlicedData.CurrentTileRegenDuration > 0.) ? (static_cast<int32>(SyncTimeSlicedData.CurrentTileRegenDuration/ TileRegenTimesAverage)) : 0;
-
-		const int32 TotalTilesToRegen = PendingDirtyTiles.Num() + (SyncTimeSlicedData.TileGeneratorSync.IsValid() ? 1 : 0);
-
-		//calculate the total processing time to regen all the tiles based on the moving window average
-		const double TotalRegenTime = TileRegenTimesAverage * static_cast<double>(TotalTilesToRegen + TilesToAddForLongCurrentTileRegen);
-
-		//calculate the time slice per frame required to regen all the tiles clamped between MinTimeSliceDuration and MaxTimeSliceDuration
-		const double NextRegenTimeSliceTime = FMath::Clamp(TotalRegenTime / static_cast<double>(MaxDesiredFramesToRegen), SyncTimeSlicedData.MinTimeSliceDuration, SyncTimeSlicedData.MaxTimeSliceDuration);
-		SyncTimeSlicedData.TimeSlicer.SetTimeSliceDuration(NextRegenTimeSliceTime);
-
-#if !UE_BUILD_SHIPPING
-		CSV_CUSTOM_STAT(NAVREGEN, NavTileRegenTimeSliceTime, static_cast<float>(NextRegenTimeSliceTime), ECsvCustomStatOp::Set);
-		CSV_CUSTOM_STAT(NAVREGEN, NavTileRegenQueueLength, TotalTilesToRegen, ECsvCustomStatOp::Set);
-		CSV_CUSTOM_STAT(NAVREGEN, TilesToAddForLongCurrentTileRegen, TilesToAddForLongCurrentTileRegen, ECsvCustomStatOp::Set);
-		CSV_CUSTOM_STAT(NAVREGEN, NavTileAvRegenTime, static_cast<float>(SyncTimeSlicedData.MovingWindowTileRegenTime.GetAverage()), ECsvCustomStatOp::Set);
-		CSV_CUSTOM_STAT(NAVREGEN, NavTileAvRegenDeltaTime, static_cast<float>(SyncTimeSlicedData.MovingWindowDeltaTime.GetAverage()), ECsvCustomStatOp::Set);
-#endif
-
 		// Submit pending tile elements
 		do
 		{
@@ -5650,7 +5610,7 @@ TArray<uint32> FRecastNavMeshGenerator::ProcessTileTasksSyncTimeSliced()
 				//next frame (as we've finished time slice processing the last tile)
 				if (!SyncTimeSlicedData.bNextTimeSliceRegenActive)
 				{
-					return EndFunction(false);
+					return EndFunction(false, true);
 				}
 
 				SyncTimeSlicedData.TileGeneratorSync = CreateTileGeneratorFromPendingElement(TileLocation);
@@ -5671,9 +5631,9 @@ TArray<uint32> FRecastNavMeshGenerator::ProcessTileTasksSyncTimeSliced()
 					}
 				}
 
-				if (SyncTimeSlicedData.TimeSlicer.TestTimeSliceFinished())
+				if (SyncTimeSlicedData.TimeSliceManager->GetTimeSlicer().TestTimeSliceFinished())
 				{
-					return EndFunction(true);
+					return EndFunction(true, true);
 				}
 			}
 			else
@@ -5702,9 +5662,9 @@ TArray<uint32> FRecastNavMeshGenerator::ProcessTileTasksSyncTimeSliced()
 					SyncTimeSlicedData.ProcessTileTasksSyncState = EProcessTileTasksSyncTimeSlicedState::AddGeneratedTiles;
 				}
 
-				if (SyncTimeSlicedData.TimeSlicer.IsTimeSliceFinishedCached())
+				if (SyncTimeSlicedData.TimeSliceManager->GetTimeSlicer().IsTimeSliceFinishedCached())
 				{
-					return EndFunction(true);
+					return EndFunction(true, true);
 				}
 			}//fall through to next state
 			case EProcessTileTasksSyncTimeSlicedState::AddGeneratedTiles:
@@ -5715,10 +5675,10 @@ TArray<uint32> FRecastNavMeshGenerator::ProcessTileTasksSyncTimeSliced()
 				{
 					SyncTimeSlicedData.ProcessTileTasksSyncState = EProcessTileTasksSyncTimeSlicedState::StoreCompessedTileCacheLayers;
 				}
-
-				if (SyncTimeSlicedData.TimeSlicer.IsTimeSliceFinishedCached())
+				 
+				if (SyncTimeSlicedData.TimeSliceManager->GetTimeSlicer().IsTimeSliceFinishedCached())
 				{
-					return EndFunction(true);
+					return EndFunction(true, true);
 				}
 			}//fall through to next state
 			case EProcessTileTasksSyncTimeSlicedState::StoreCompessedTileCacheLayers:
@@ -5738,16 +5698,21 @@ TArray<uint32> FRecastNavMeshGenerator::ProcessTileTasksSyncTimeSliced()
 			}//fall through to next state
 			case EProcessTileTasksSyncTimeSlicedState::Finish:
 			{
-				//no need to check time slicing as not much work done
 				//reset state to Init for next tile to be processed
 				SyncTimeSlicedData.ProcessTileTasksSyncState = EProcessTileTasksSyncTimeSlicedState::Init;
 				SyncTimeSlicedData.TileGeneratorSync.Reset();
 
 				SyncTimeSlicedData.CurrentTileRegenDuration += (FPlatformTime::Seconds() - TimeStartProcessingTileThisFrame);
 
-				SyncTimeSlicedData.MovingWindowTileRegenTime.PushValue(SyncTimeSlicedData.CurrentTileRegenDuration);
+				SyncTimeSlicedData.TimeSliceManager->PushTileRegenTime(SyncTimeSlicedData.CurrentTileRegenDuration);
 
 				SyncTimeSlicedData.CurrentTileRegenDuration = 0.;
+				
+				//test time slice 
+				if (SyncTimeSlicedData.TimeSliceManager->GetTimeSlicer().TestTimeSliceFinished())
+				{
+					return EndFunction(false, true);
+				}
 			}
 			break;
 			default:
@@ -5761,8 +5726,8 @@ TArray<uint32> FRecastNavMeshGenerator::ProcessTileTasksSyncTimeSliced()
 
 	// we only hit this if we have processed too many tiles in a frame and we will already
 	// have calculated the tile regen duration, or if we have processed no tiles and we also
-	// don't need to calcualte the tile regen duration
-	return EndFunction(false);
+	// don't want to calcualte the tile regen duration
+	return EndFunction(false, bHadWorktoDo);
 }
 
 //this code path is approx 10% faster than ProcessTileTasksSyncTimeSliced, however it spikes far worse for most use cases.
@@ -5821,17 +5786,25 @@ TArray<uint32> FRecastNavMeshGenerator::ProcessTileTasks(const int32 NumTasksToP
 #if RECAST_ASYNC_REBUILDING
 	UpdatedTiles = ProcessTileTasksAsync(NumTasksToProcess);
 #else
-	//only switch bTimeSliceRegen state if we are not time slicing or if we are but aren't part way through time slicing a tile
-	if (SyncTimeSlicedData.bTimeSliceRegenActive != SyncTimeSlicedData.bNextTimeSliceRegenActive)
+	if (SyncTimeSlicedData.TimeSliceManager)
 	{
-		if (!SyncTimeSlicedData.bTimeSliceRegenActive)
-		{
-			SyncTimeSlicedData.bTimeSliceRegenActive = SyncTimeSlicedData.bNextTimeSliceRegenActive;
-		}
-		else if (!SyncTimeSlicedData.TileGeneratorSync.IsValid())//test if we have finished processing a tile
+		//only switch bTimeSliceRegen state if we are not time slicing or if we are but aren't part way through time slicing a tile
+		if (SyncTimeSlicedData.bTimeSliceRegenActive != SyncTimeSlicedData.bNextTimeSliceRegenActive)
 		{
-			SyncTimeSlicedData.bTimeSliceRegenActive = SyncTimeSlicedData.bNextTimeSliceRegenActive;
-		}
+			if (!SyncTimeSlicedData.bTimeSliceRegenActive)
+			{
+				SyncTimeSlicedData.bTimeSliceRegenActive = SyncTimeSlicedData.bNextTimeSliceRegenActive;
+			}
+			else if (!SyncTimeSlicedData.TileGeneratorSync.IsValid())//test if we have finished processing a tile
+			{
+				SyncTimeSlicedData.bTimeSliceRegenActive = SyncTimeSlicedData.bNextTimeSliceRegenActive;
+			}
+		}	
+	}
+	else
+	{
+		//No time slice manager no timesliced regen
+		SyncTimeSlicedData.bTimeSliceRegenActive = false;
 	}
 
 	if (SyncTimeSlicedData.bTimeSliceRegenActive)
@@ -5989,11 +5962,19 @@ bool FRecastNavMeshGenerator::IsBuildInProgress(bool bCheckDirtyToo) const
 		|| SyncTimeSlicedData.TileGeneratorSync.Get();
 }
 
+#if !RECAST_ASYNC_REBUILDING
+bool FRecastNavMeshGenerator::GetTimeSliceData(int32& OutNumRemainingBuildTasks, double& OutCurrentBuildTaskDuration) const
+{
+	//it's probably just faster to calculate these rather than branch and only calculate them if bTimeSliceRegenActive is true;
+	OutNumRemainingBuildTasks = GetNumRemaningBuildTasksHelper();
+	OutCurrentBuildTaskDuration = SyncTimeSlicedData.CurrentTileRegenDuration;
+	return SyncTimeSlicedData.bTimeSliceRegenActive;
+}
+#endif
+
 int32 FRecastNavMeshGenerator::GetNumRemaningBuildTasks() const
 {
-	return RunningDirtyTiles.Num() 
-		+ PendingDirtyTiles.Num()
-		+ (SyncTimeSlicedData.TileGeneratorSync.Get() ? 1 : 0);
+	return GetNumRemaningBuildTasksHelper();
 }
 
 int32 FRecastNavMeshGenerator::GetNumRunningBuildTasks() const
diff --git a/Engine/Source/Runtime/NavigationSystem/Private/NavigationSystem.cpp b/Engine/Source/Runtime/NavigationSystem/Private/NavigationSystem.cpp
--- a/Engine/Source/Runtime/NavigationSystem/Private/NavigationSystem.cpp
+++ b/Engine/Source/Runtime/NavigationSystem/Private/NavigationSystem.cpp
@@ -233,6 +233,120 @@ namespace
 	}
 #endif // ENABLE_VISUAL_LOG
 }
+
+void FNavRegenTimeSlicer::SetupTimeSlice(double SliceDuration)
+{
+	RemainingDuration = OriginalDuration = SliceDuration; 
+	StartTime = TimeLastTested = 0.;
+	bTimeSliceFinishedCached = false; 
+}
+
+void FNavRegenTimeSlicer::StartTimeSlice()
+{
+	ensureMsgf(!bTimeSliceFinishedCached, TEXT("Starting a time slice that has already been tested as finished! Call SetupTimeSlice() before calling StartTimeSlice() again!"));
+	ensureMsgf(RemainingDuration > 0., TEXT("Attempting to start a time slice that has zero duration!"));
+
+	TimeLastTested = StartTime = FPlatformTime::Seconds();
+}
+ 
+void FNavRegenTimeSlicer::EndTimeSliceAndAdjustDuration()
+{
+	RemainingDuration = FMath::Max(RemainingDuration - (TimeLastTested - StartTime), 0.);
+}
+
+bool FNavRegenTimeSlicer::TestTimeSliceFinished() const
+{
+	ensureMsgf(!bTimeSliceFinishedCached, TEXT("Testing time slice is finished when we have already confirmed that!"));
+
+	TimeLastTested = FPlatformTime::Seconds();
+
+	bTimeSliceFinishedCached = (TimeLastTested - StartTime) >= RemainingDuration;
+	return bTimeSliceFinishedCached;
+}
+
+FNavRegenTimeSliceManager::FNavRegenTimeSliceManager()
+	: MinTimeSliceDuration(0.00075)
+	, MaxTimeSliceDuration(0.004)
+	, FrameNumOld(TNumericLimits<int64>::Max() - 1)
+	, MaxDesiredTileRegenDuration(0.7)
+	, TimeLastCall(-1.f)
+	, NavDataIdx(0)
+#if TIME_SLICE_NAV_REGEN
+	, bDoTimeSlicedUpdate(true)
+#else
+	, bDoTimeSlicedUpdate(false)
+#endif
+{
+}
+
+void FNavRegenTimeSliceManager::CalcAverageDeltaTime(uint64 FrameNum)
+{
+	const double CurTime = FPlatformTime::Seconds();
+
+	if ((FrameNumOld + 1) == FrameNum)
+	{
+		const double DeltaTime = CurTime - TimeLastCall;
+		MovingWindowDeltaTime.PushValue(DeltaTime);
+	}
+	TimeLastCall = CurTime;
+	FrameNumOld = FrameNum;
+}
+
+void FNavRegenTimeSliceManager::CalcTimeSliceDuration(int32 NumTilesToRegen, const TArray<double>& CurrentTileRegenDurations)
+{
+	const float DeltaTimesAverage = (MovingWindowDeltaTime.GetAverage() > 0.f) ? MovingWindowDeltaTime.GetAverage() : (1.f / 30.f); //use default 33 ms
+
+	const double TileRegenTimesAverage = (MovingWindowTileRegenTime.GetAverage() > 0.) ? MovingWindowTileRegenTime.GetAverage() : 0.0025; //use default of 2.5 milli secs to regen a full tile
+
+	//calculate the max desired frames to regen all the tiles in PendingDirtyTiles
+	const float MaxDesiredFramesToRegen = FMath::FloorToFloat(MaxDesiredTileRegenDuration / DeltaTimesAverage);
+
+	//tiles to add to PendingDirtyTiles if the current tiles are taking longer than average to regen
+	//we add 1 tile for however many times longer the current tile is taking compared with the moving window average
+	int32 TilesToAddForLongCurrentTileRegen = 0;
+	for (const double RegenDuration : CurrentTileRegenDurations)
+	{
+		TilesToAddForLongCurrentTileRegen += (RegenDuration > 0.) ? (static_cast<int32>(RegenDuration / TileRegenTimesAverage)) : 0;
+	}
+
+	//calculate the total processing time to regen all the tiles based on the moving window average
+	const double TotalRegenTime = TileRegenTimesAverage * static_cast<double>(NumTilesToRegen + TilesToAddForLongCurrentTileRegen);
+
+	//calculate the time slice per frame required to regen all the tiles clamped between MinTimeSliceDuration and MaxTimeSliceDuration
+	const double NextRegenTimeSliceTime = FMath::Clamp(TotalRegenTime / static_cast<double>(MaxDesiredFramesToRegen), MinTimeSliceDuration, MaxTimeSliceDuration);
+
+	TimeSlicer.SetupTimeSlice(NextRegenTimeSliceTime);
+
+#if !UE_BUILD_SHIPPING
+	CSV_CUSTOM_STAT(NavigationSystem, NavTileRegenTimeSliceTimeMs, static_cast<float>(NextRegenTimeSliceTime * 1000.), ECsvCustomStatOp::Set);
+	CSV_CUSTOM_STAT(NavigationSystem, NavTileNumTilesToRegen, NumTilesToRegen, ECsvCustomStatOp::Set);
+	CSV_CUSTOM_STAT(NavigationSystem, NavTilesToAddForLongCurrentTileRegen, TilesToAddForLongCurrentTileRegen, ECsvCustomStatOp::Set);
+	CSV_CUSTOM_STAT(NavigationSystem, NavTileAvRegenTimeMs, static_cast<float>(MovingWindowTileRegenTime.GetAverage() * 1000.), ECsvCustomStatOp::Set);
+	CSV_CUSTOM_STAT(NavigationSystem, NavTileAvRegenDeltaTimeMs, static_cast<float>(MovingWindowDeltaTime.GetAverage() * 1000.), ECsvCustomStatOp::Set);
+#endif
+}
+
+void FNavRegenTimeSliceManager::SetMinTimeSliceDuration(double NewMinTimeSliceDuration)
+{
+	MinTimeSliceDuration = NewMinTimeSliceDuration;
+
+	UE_LOG(LogNavigationDataBuild, Log, TEXT("Navigation System: MinTimeSliceDuration = %f"), MinTimeSliceDuration);
+}
+
+void FNavRegenTimeSliceManager::SetMaxTimeSliceDuration(double NewMaxTimeSliceDuration)
+{
+	MaxTimeSliceDuration = NewMaxTimeSliceDuration;
+
+	UE_LOG(LogNavigationDataBuild, Log, TEXT("Navigation System: MaxTimeSliceDuration = %f"), MaxTimeSliceDuration);
+}
+
+void FNavRegenTimeSliceManager::SetMaxDesiredTileRegenDuration(float NewMaxDesiredTileRegenDuration)
+{
+	MaxDesiredTileRegenDuration = NewMaxDesiredTileRegenDuration;
+
+	UE_LOG(LogNavigationDataBuild, Log, TEXT("Navigation System: MaxDesiredTileRegenDuration = %f"), MaxDesiredTileRegenDuration);
+}
+
 //----------------------------------------------------------------------//
 // UNavigationSystemV1                                                                
 //----------------------------------------------------------------------//
@@ -925,6 +1039,40 @@ void UNavigationSystemV1::SetCrowdManager(UCrowdManagerBase* NewCrowdManager)
 	}
 }
 
+void UNavigationSystemV1::CalcTimeSlicedUpdateData(TArray<double>& OutCurrentTimeSlicedBuildTaskDurations, TArray<bool>& OutIsTimeSlicingArray, bool& bOutAnyNonTimeSlicedGenerators, int32& OutNumTimeSlicedRemainingBuildTasks)
+{
+	OutNumTimeSlicedRemainingBuildTasks = 0;
+	OutIsTimeSlicingArray.SetNumZeroed(NavDataSet.Num());
+	bOutAnyNonTimeSlicedGenerators = false;
+	OutCurrentTimeSlicedBuildTaskDurations.Empty();
+	OutCurrentTimeSlicedBuildTaskDurations.Reserve(NavDataSet.Num());
+
+	for (int32 NavDataIdx = 0; NavDataIdx < NavDataSet.Num(); ++NavDataIdx)
+	{
+		const ANavigationData* NavData = NavDataSet[NavDataIdx];
+		const FNavDataGenerator* Generator = NavData ? NavData->GetGenerator() : nullptr;
+		if (Generator)
+		{
+			double TimeSlicedBuildTaskDuration = 0.;
+			int32 NumRemainingBuildTasksTemp = 0;
+
+			if (Generator->GetTimeSliceData(NumRemainingBuildTasksTemp, TimeSlicedBuildTaskDuration))
+			{
+				OutIsTimeSlicingArray[NavDataIdx] = true;
+				OutNumTimeSlicedRemainingBuildTasks += NumRemainingBuildTasksTemp;
+				if (TimeSlicedBuildTaskDuration > 0.)
+				{
+					OutCurrentTimeSlicedBuildTaskDurations.Push(TimeSlicedBuildTaskDuration);
+				}
+			}
+			else
+			{
+				bOutAnyNonTimeSlicedGenerators = true;
+			}
+		}
+	}
+}
+
 void UNavigationSystemV1::Tick(float DeltaSeconds)
 {
 	SET_DWORD_STAT(STAT_Navigation_ObservedPathsCount, 0);
@@ -980,11 +1128,90 @@ void UNavigationSystemV1::Tick(float DeltaSeconds)
 		{
 			SCOPE_CYCLE_COUNTER(STAT_Navigation_TickAsyncBuild);
 
-			for (ANavigationData* NavData : NavDataSet)
+			bool bDoStandardTickAsync = true;
+
+			if (NavRegenTimeSliceManager.DoTimeSlicedUpdate())
 			{
-				if (NavData)
+				int32 NumTimeSlicedRemainingBuildTasks = 0;
+				TArray<double> CurrentTimeSlicedBuildTaskDurations;
+				TArray<bool> IsTimeSlicingArray;
+				bool bAnyNonTimeSlicedGenerators = false;
+
+				NavRegenTimeSliceManager.CalcAverageDeltaTime(GFrameCounter);
+
+				CalcTimeSlicedUpdateData(CurrentTimeSlicedBuildTaskDurations, IsTimeSlicingArray, bAnyNonTimeSlicedGenerators, NumTimeSlicedRemainingBuildTasks);
+
+				if (NumTimeSlicedRemainingBuildTasks > 0)
 				{
-					NavData->TickAsyncBuild(DeltaSeconds);
+					NavRegenTimeSliceManager.CalcTimeSliceDuration(NumTimeSlicedRemainingBuildTasks, CurrentTimeSlicedBuildTaskDurations);
+
+					//The general idea here is to tick any non time sliced generators once per frame. Time sliced generators we aim to tick one per frame and move to the next, next frame. In the
+					//case where one time sliced generator doesn't use the whole time slice we move to the next time sliced generator. That generator will only be considered to have a full frames
+					//processing if either it runs out of work or uses a large % of the time slice. Depending we either tick it again next frame or go to the next time sliced generator (next frame).
+					bool bNavDataIdxSet = false;
+					int32 NavDataIdxTemp = NavRegenTimeSliceManager.GetNavDataIdx();
+					const double RemainingFractionConsideredWholeTick = 0.8;
+					const int32 FirstNavDataIdx = NavDataIdxTemp = NavDataIdxTemp % NavDataSet.Num();
+
+					for (int32 NavDataIter = 0; NavDataIter < NavDataSet.Num(); ++NavDataIter)
+					{
+						if (ANavigationData* NavData = NavDataSet[NavDataIdxTemp])
+						{
+							if (IsTimeSlicingArray[NavDataIdxTemp])
+							{
+								if (NavRegenTimeSliceManager.GetTimeSlicer().IsTimeSliceFinishedCached())
+								{
+									//if we haven't set the NavDataIdx then this is the TimeSliced Generator to process next frame
+									if (!bNavDataIdxSet)
+									{
+										NavRegenTimeSliceManager.SetNavDataIdx(NavDataIdxTemp);
+										bNavDataIdxSet = true;
+									}
+
+									//if the time slice is finished and we have no non time sliced generators then stop TickAsyncBuild, otherwise continue
+									if (!bAnyNonTimeSlicedGenerators)
+									{
+										break;
+									}
+									continue;
+								}
+								else if (NavRegenTimeSliceManager.GetTimeSlicer().GetRemainingDurationFraction() < RemainingFractionConsideredWholeTick)
+								{
+									//don't check bNavDataIdxSet here, either this time sliced generator won't get enough time this frame to be considered
+									//a whole tick or it will complete and there is some time sliced left - in the later case next frame we'll process the 
+									//next time sliced generator we process this frame or the first Idx we processed this frame
+									NavRegenTimeSliceManager.SetNavDataIdx(NavDataIdxTemp);
+									bNavDataIdxSet = true;
+								}
+							}
+							NavData->TickAsyncBuild(DeltaSeconds);
+						}
+						//Increment and mod NavDataIdxTemp
+						++NavDataIdxTemp;
+						NavDataIdxTemp %= NavDataSet.Num();
+					}
+
+					//if we processed all the time sliced generators and there is still some time slice left
+					//OR if we haven't SetNavDataIdx() we should start next frame where we started this frame
+					if (!NavRegenTimeSliceManager.GetTimeSlicer().IsTimeSliceFinishedCached() || !bNavDataIdxSet)
+					{
+						NavRegenTimeSliceManager.SetNavDataIdx(FirstNavDataIdx);
+						bNavDataIdxSet = true;
+					}
+					//don't do the standard TickASyncBuild as we have already processed the regen appropriately 
+					bDoStandardTickAsync = false;
+				}
+			}
+
+			//if we aren't time sliced rebuilding and / or if there aren't any time sliced nav data's with work to do just tick all nav data
+			if (bDoStandardTickAsync)
+			{
+				for (ANavigationData* NavData : NavDataSet)
+				{
+					if (NavData)
+					{
+						NavData->TickAsyncBuild(DeltaSeconds);
+					}
 				}
 			}
 		}
@@ -4599,4 +4826,4 @@ void UNavigationSystemModuleConfig::PostEditChangeProperty(FPropertyChangedEvent
 }
 #endif // WITH_EDITOR
 
-#undef LOCTEXT_NAMESPACE
+#undef LOCTEXT_NAMESPACE
\ No newline at end of file
diff --git a/Engine/Source/Runtime/NavigationSystem/Public/NavMesh/RecastNavMeshGenerator.h b/Engine/Source/Runtime/NavigationSystem/Public/NavMesh/RecastNavMeshGenerator.h
--- a/Engine/Source/Runtime/NavigationSystem/Public/NavMesh/RecastNavMeshGenerator.h
+++ b/Engine/Source/Runtime/NavigationSystem/Public/NavMesh/RecastNavMeshGenerator.h
@@ -13,7 +13,6 @@
 #include "AI/NavDataGenerator.h"
 #include "NavMesh/RecastHelpers.h"
 #include "NavDebugTypes.h"
-#include "Math/MovingWindowAverageFast.h"
 
 #if WITH_RECAST
 
@@ -38,6 +37,7 @@ struct FTileCacheCompressor;
 struct FTileCacheAllocator;
 struct FTileGenerationContext;
 class dtNavMesh;
+class FNavRegenTimeSliceManager;
 
 #define MAX_VERTS_PER_POLY	6
 
@@ -191,39 +191,12 @@ struct FRcTileBox
 };
 
 /**
+ * TIME SLICING 
  * The general idea is that any function that handles time slicing internally will be named XXXXTimeSliced
  * and returns a ETimeSliceWorkResult. These functions also call TestTimeSliceFinished() internally when required,
  * IsTimeSliceFinishedCached() can be called externally after they have finished. Non time sliced functions are 
  * managed externally and the calling function should call TestTimeSliceFinished() when necessary.
  */
-struct NAVIGATIONSYSTEM_API FTimeSlicer
-{
-public:
-	FTimeSlicer()
-		: TimeSliceDuration(0.)
-		, TimeSliceStartTime(0.)
-		, bTimeSliceFinishedCached(false)
-	{
-	}
-
-	FTimeSlicer(double SliceDuration)
-		: TimeSliceDuration(SliceDuration)
-		, TimeSliceStartTime(0.)
-		, bTimeSliceFinishedCached(false)
-	{
-	}
-
-	void SetTimeSliceDuration(double SliceDuration);
-	void StartTimeSlice();
-	double GetStartTime() const;
-	bool TestTimeSliceFinished() const;
-	bool IsTimeSliceFinishedCached() const;
-
-private:
-	double TimeSliceDuration;
-	double TimeSliceStartTime;
-	mutable bool bTimeSliceFinishedCached;
-};
 
 /** Return state of calling time sliced functions */
 enum class ETimeSliceWorkResult : uint8
@@ -428,7 +401,7 @@ protected:
 	int32 RasterizeTrianglesTimeSlicedRawGeomIdx;
 	int32 RasterizeTrianglesTimeSlicedInstTransformIdx;
 	TNavStatArray<uint8> RasterizeGeomRecastTriAreas;
-	const FTimeSlicer& TimeSlicer;
+	const FNavRegenTimeSliceManager* TimeSliceManager;
 
 	TUniquePtr<struct FTileCacheAllocator> GenNavDataTimeSlicedAllocator;
 	TUniquePtr<struct FTileGenerationContext> GenNavDataTimeSlicedGenerationContext;
@@ -624,6 +597,13 @@ public:
 	/** Asks generator to update navigation affected by DirtyAreas */
 	virtual void RebuildDirtyAreas(const TArray<FNavigationDirtyArea>& DirtyAreas) override;
 	virtual bool IsBuildInProgress(bool bCheckDirtyToo = false) const override;
+
+#if !RECAST_ASYNC_REBUILDING
+	/** returns true if we are time slicing and the data is valid to use false otherwise*/
+	virtual bool GetTimeSliceData(int32& OutNumRemainingBuildTasks, double& OutCurrentBuildTaskDuration) const override;
+#endif
+
+	int32 GetNumRemaningBuildTasksHelper() const { return RunningDirtyTiles.Num() + PendingDirtyTiles.Num() + static_cast<int32>(SyncTimeSlicedData.TileGeneratorSync.IsValid()); }
 	virtual int32 GetNumRemaningBuildTasks() const override;
 	virtual int32 GetNumRunningBuildTasks() const override;
 
@@ -683,7 +663,7 @@ public:
 	static void GetDebugGeometry(const FNavigationRelevantData& EncodedData, FNavDebugMeshData& DebugMeshData);
 #endif  // !UE_BUILD_SHIPPING
 
-	FTimeSlicer& GetTimeSlicer() { return SyncTimeSlicedData.TimeSlicer; }
+	const FNavRegenTimeSliceManager* GetTimeSliceManager() const { return SyncTimeSlicedData.TimeSliceManager; }
 	
 	void SetNextTimeSliceRegenActive(bool bRegenState) { SyncTimeSlicedData.bNextTimeSliceRegenActive = bRegenState; }
 
@@ -841,15 +821,6 @@ protected:
 		FSyncTimeSlicedData();
 
 		double CurrentTileRegenDuration;
-		double MinTimeSliceDuration;
-		double MaxTimeSliceDuration;
-		float RealTimeSecsLastCall; //no need to reset this
-		/** The max real world desired time to Regen all the tiles in PendingDirtyTiles,
-		 *  Note it could take longer than this, as the time slice is clamped per frame between
-		 *	MinTimeSliceDuration and MaxTimeSliceDuration.
-		 */
-		float MaxDesiredTileRegenDuration;
-
 		/** if we are currently using time sliced regen or not - currently an experimental feature.
 		 *  do not manipulate this value directly instead call SetNextTimeSliceRegenState()
 		 */
@@ -874,13 +845,9 @@ protected:
 		/** Used by AddGeneratedTilesTimeSliced */
 		int32 AddGenTilesLayerIndex;
 
-		/** Used to calculate the moving window average of the actual time spent inside functions used to regenerate a a tile, this is processing time not actual time over multiple frames */
-		FMovingWindowAverageFast<double, 128> MovingWindowTileRegenTime;
-		/** Used to calculate the actual moving window delta time */
-		FMovingWindowAverageFast<float, 128> MovingWindowDeltaTime;
 		/** Do not null or Reset this directly instead call ResetTimeSlicedTileGeneratorSync(). The only exception currently is in ProcessTileTasksSyncTimeSliced */
 		TSharedPtr<FRecastTileGenerator> TileGeneratorSync;
-		FTimeSlicer TimeSlicer;
+		FNavRegenTimeSliceManager* TimeSliceManager;
 	};
 
 	FSyncTimeSlicedData SyncTimeSlicedData;
diff --git a/Engine/Source/Runtime/NavigationSystem/Public/NavigationData.h b/Engine/Source/Runtime/NavigationSystem/Public/NavigationData.h
--- a/Engine/Source/Runtime/NavigationSystem/Public/NavigationData.h
+++ b/Engine/Source/Runtime/NavigationSystem/Public/NavigationData.h
@@ -621,7 +621,10 @@ public:
 	/** Cancels current build  */
 	virtual void CancelBuild();
 
-	/** Ticks navigation build  */
+	/** Ticks navigation build
+	 *  If the generator is set to time sliced rebuild then this function will only get called when 
+	 *  there is sufficient time (effectively roughly once in n frames where n is the number of time sliced nav data generators currently building)
+	 */
 	virtual void TickAsyncBuild(float DeltaSeconds);
 	
 	/** Retrieves navmesh's generator */
diff --git a/Engine/Source/Runtime/NavigationSystem/Public/NavigationSystem.h b/Engine/Source/Runtime/NavigationSystem/Public/NavigationSystem.h
--- a/Engine/Source/Runtime/NavigationSystem/Public/NavigationSystem.h
+++ b/Engine/Source/Runtime/NavigationSystem/Public/NavigationSystem.h
@@ -18,6 +18,7 @@
 #include "AI/NavigationSystemConfig.h"
 #include "NavigationOctreeController.h"
 #include "NavigationDirtyAreasController.h"
+#include "Math/MovingWindowAverageFast.h"
 #if WITH_EDITOR
 #include "UnrealEdMisc.h"
 #endif // WITH_EDITOR
@@ -94,6 +95,101 @@ namespace ENavigationBuildLock
 	};
 }
 
+
+class NAVIGATIONSYSTEM_API FNavRegenTimeSlicer
+{
+public:
+	/** Setup the initial values for a time slice. This can be called on an instance after TestTimeSliceFinished() has returned true and EndTimeSliceAndAdjustDuration() has been called */
+	void SetupTimeSlice(double SliceDuration);
+
+	/** 
+	 *  Starts the time slice, this can be called multiple times as long as EndTimeSliceAndAdjustDuration() is called between each call.
+	 *  StartTimeSlice should not be called after TestTimeSliceFinished() has returned true
+	 */
+	void StartTimeSlice();
+
+	/** 
+	 *  Useful when multiple sections of code need to be timesliced per frame using the same time slice duration that do not necessarily occur concurrently.
+	 *  This ends the time sliced code section and adjusts the RemainingDuration based on the time used between calls to StartTimeSlice and the last call to TestTimeSliceFinished.
+	 *  Note the actual time slice is not tested in this function. Thats done in TestTimeSliceFinished!
+	 *  This can be called multiple times as long as StartTimeSlice() is called before EndTimeSliceAndAdjustDuration().
+	 *  EndTimeSliceAndAdjustDuration can be called after TestTimeSliceFinished() has returned true in this case the RemainingDuration will just be zero
+	 */
+	void EndTimeSliceAndAdjustDuration();
+	double GetStartTime() const { return StartTime; }
+	bool TestTimeSliceFinished() const;
+
+	//* Returns the cached result of calling TestTimeSliceFinished, false by default */
+	bool IsTimeSliceFinishedCached() const { return bTimeSliceFinishedCached; }
+	double GetRemainingDurationFraction() const { return OriginalDuration > 0. ? RemainingDuration / OriginalDuration : 0.; }
+
+protected:
+	double OriginalDuration = 0.;
+	double RemainingDuration = 0.;
+	double StartTime = 0.;
+	mutable double TimeLastTested = 0.;
+	mutable bool bTimeSliceFinishedCached = false;
+};
+
+class NAVIGATIONSYSTEM_API FNavRegenTimeSliceManager
+{
+public:
+	FNavRegenTimeSliceManager();
+
+	void PushTileRegenTime(double NewTime) { MovingWindowTileRegenTime.PushValue(NewTime);  }
+
+	double GetAverageTileRegenTime() const { return MovingWindowTileRegenTime.GetAverage();  }
+
+	double GetAverageDeltaTime() const { return MovingWindowDeltaTime.GetAverage(); }
+
+	bool DoTimeSlicedUpdate() const { return bDoTimeSlicedUpdate; }
+
+	void CalcAverageDeltaTime(uint64 FrameNum);
+
+	void CalcTimeSliceDuration(int32 NumTilesToRegen, const TArray<double>& CurrentTileRegenDurations);
+
+	void SetMinTimeSliceDuration(double NewMinTimeSliceDuration);
+
+	void SetMaxTimeSliceDuration(double NewMaxTimeSliceDuration);
+
+	void SetMaxDesiredTileRegenDuration(float NewMaxDesiredTileRegenDuration);
+
+	int32 GetNavDataIdx() const { return NavDataIdx;  }
+	void SetNavDataIdx(int32 InNavDataIdx) { NavDataIdx = InNavDataIdx; }
+
+	FNavRegenTimeSlicer& GetTimeSlicer() { return TimeSlicer; }
+	const FNavRegenTimeSlicer& GetTimeSlicer() const { return TimeSlicer; }
+
+protected:
+	FNavRegenTimeSlicer TimeSlicer;
+
+	/** Used to calculate the moving window average of the actual time spent inside functions used to regenerate a tile, this is processing time not actual time over multiple frames */
+	FMovingWindowAverageFast<double, 256> MovingWindowTileRegenTime;
+
+	/** Used to calculate the actual moving window delta time */
+	FMovingWindowAverageFast<double, 256> MovingWindowDeltaTime;
+
+	/** If there are enough tiles to process this in the Min Time Slice Duration */
+	double MinTimeSliceDuration;
+
+	/** The max Desired Time Slice Duration */
+	double MaxTimeSliceDuration;
+
+	uint64 FrameNumOld;
+
+	/** The max real world desired time to Regen all the tiles in PendingDirtyTiles,
+	 *  Note it could take longer than this, as the time slice is clamped per frame between
+	 *	MinTimeSliceDuration and MaxTimeSliceDuration.
+	 */
+	float MaxDesiredTileRegenDuration;
+
+	double TimeLastCall;
+
+	int32 NavDataIdx;
+
+	bool bDoTimeSlicedUpdate;
+};
+
 UCLASS(Within=World, config=Engine, defaultconfig)
 class NAVIGATIONSYSTEM_API UNavigationSystemV1 : public UNavigationSystemBase
 {
@@ -360,6 +456,8 @@ public:
 	UCrowdManagerBase* GetCrowdManager() const { return CrowdManager.Get(); }
 
 protected:
+	void CalcTimeSlicedUpdateData(TArray<double>& OutCurrentTimeSlicedBuildTaskDurations, TArray<bool>& OutIsTimeSlicingArray, bool& bOutAnyNonTimeSlicedGenerators, int32& OutNumTimeSlicedRemainingBuildTasks);
+
 	/** spawn new crowd manager */
 	virtual void CreateCrowdManager();
 
@@ -863,6 +961,8 @@ public:
 	//----------------------------------------------------------------------//
 	void CycleNavigationDataDrawn();
 
+	FNavRegenTimeSliceManager& GetNavRegenTimeSliceManager() { return NavRegenTimeSliceManager; }
+
 protected:
 
 	UPROPERTY()
@@ -921,6 +1021,8 @@ protected:
 	static TMap<INavLinkCustomInterface*, FWeakObjectPtr> PendingCustomLinkRegistration;
 	TSet<const UClass*> NavAreaClasses;
 
+	FNavRegenTimeSliceManager NavRegenTimeSliceManager;
+
 	/** delegate handler for PostLoadMap event */
 	void OnPostLoadMap(UWorld* LoadedWorld);
 #if WITH_EDITOR
