Addition bMaskFillCollisionUnderneathForNavmesh option on NavModifierVolume to go along bFillCollisionUnderneathForNavmesh.
This allows navvolumes to mark a 2D areas to be ignored by the fill under feature.

diff --git a/Engine/Source/Runtime/Engine/Classes/AI/Navigation/NavRelevantInterface.h b/Engine/Source/Runtime/Engine/Classes/AI/Navigation/NavRelevantInterface.h
--- a/Engine/Source/Runtime/Engine/Classes/AI/Navigation/NavRelevantInterface.h
+++ b/Engine/Source/Runtime/Engine/Classes/AI/Navigation/NavRelevantInterface.h
@@ -89,7 +89,7 @@ struct ENGINE_API FNavigationRelevantData : public TSharedFromThis<FNavigationRe
 	FORCEINLINE uint32 GetGeometryAllocatedSize() const { return CollisionData.GetAllocatedSize() + VoxelData.GetAllocatedSize(); }
 	FORCEINLINE int32 GetDirtyFlag() const
 	{
-		return ((HasGeometry() || IsPendingLazyGeometryGathering()) ? ENavigationDirtyFlag::Geometry : 0) |
+		return ((HasGeometry() || IsPendingLazyGeometryGathering() || Modifiers.GetMaskFillCollisionUnderneathForNavmesh()) ? ENavigationDirtyFlag::Geometry : 0) |
 			((HasModifiers() || IsPendingLazyModifiersGathering()) ? ENavigationDirtyFlag::DynamicModifier : 0) |
 			(Modifiers.HasAgentHeightAdjust() ? ENavigationDirtyFlag::UseAgentHeight : 0);
 	}
diff --git a/Engine/Source/Runtime/Engine/Private/AI/Navigation/NavigationModifier.cpp b/Engine/Source/Runtime/Engine/Private/AI/Navigation/NavigationModifier.cpp
--- a/Engine/Source/Runtime/Engine/Private/AI/Navigation/NavigationModifier.cpp
+++ b/Engine/Source/Runtime/Engine/Private/AI/Navigation/NavigationModifier.cpp
@@ -731,7 +731,8 @@ void FCompositeNavModifier::Reset()
 	bHasPotentialLinks = false;
 	bAdjustHeight = false;
 	bIsPerInstanceModifier = false;
-	bModifierFillCollisionUnderneathForNavmesh = false;
+	bFillCollisionUnderneathForNavmesh = false;
+	bMaskFillCollisionUnderneathForNavmesh = false;
 }
 
 void FCompositeNavModifier::Empty()
@@ -741,7 +742,8 @@ void FCompositeNavModifier::Empty()
 	CustomLinks.Empty();
 	bHasPotentialLinks = false;
 	bAdjustHeight = false;
-	bModifierFillCollisionUnderneathForNavmesh = false;
+	bFillCollisionUnderneathForNavmesh = false;
+	bMaskFillCollisionUnderneathForNavmesh = false;
 }
 
 FCompositeNavModifier FCompositeNavModifier::GetInstantiatedMetaModifier(const FNavAgentProperties* NavAgent, TWeakObjectPtr<UObject> WeakOwnerPtr) const
diff --git a/Engine/Source/Runtime/Engine/Public/AI/NavigationModifier.h b/Engine/Source/Runtime/Engine/Public/AI/NavigationModifier.h
--- a/Engine/Source/Runtime/Engine/Public/AI/NavigationModifier.h
+++ b/Engine/Source/Runtime/Engine/Public/AI/NavigationModifier.h
@@ -250,7 +250,14 @@ protected:
 
 struct ENGINE_API FCompositeNavModifier : public FNavigationModifier
 {
-	FCompositeNavModifier() : bHasPotentialLinks(false), bAdjustHeight(false), bHasLowAreaModifiers(false), bIsPerInstanceModifier(false), bModifierFillCollisionUnderneathForNavmesh(false) {}
+	FCompositeNavModifier() 
+		: bHasPotentialLinks(false)
+		, bAdjustHeight(false)
+		, bHasLowAreaModifiers(false)
+		, bIsPerInstanceModifier(false)
+		, bFillCollisionUnderneathForNavmesh(false)
+		, bMaskFillCollisionUnderneathForNavmesh(false)
+	{}
 
 	void Shrink();
 	void Reset();
@@ -258,7 +265,9 @@ struct ENGINE_API FCompositeNavModifier : public FNavigationModifier
 
 	FORCEINLINE bool IsEmpty() const 
 	{ 
-		return (Areas.Num() == 0) && (SimpleLinks.Num() == 0) && (CustomLinks.Num() == 0);
+		return (Areas.Num() == 0) && (SimpleLinks.Num() == 0) && (CustomLinks.Num() == 0) && 
+			!bFillCollisionUnderneathForNavmesh && 
+			!bMaskFillCollisionUnderneathForNavmesh;
 	}
 
 	void Add(const FAreaNavModifier& Area)
@@ -289,7 +298,8 @@ struct ENGINE_API FCompositeNavModifier : public FNavigationModifier
 		bHasMetaAreas |= Modifiers.bHasMetaAreas; 
 		bAdjustHeight |= Modifiers.HasAgentHeightAdjust();
 		bHasLowAreaModifiers |= Modifiers.HasLowAreaModifiers();
-		bModifierFillCollisionUnderneathForNavmesh |= Modifiers.GetFillCollisionUnderneathForNavmesh();
+		bFillCollisionUnderneathForNavmesh |= Modifiers.GetFillCollisionUnderneathForNavmesh();
+		bMaskFillCollisionUnderneathForNavmesh |= Modifiers.GetMaskFillCollisionUnderneathForNavmesh();
 	}
 
 	void CreateAreaModifiers(const UPrimitiveComponent* PrimComp, const TSubclassOf<UNavAreaBase> AreaClass);
@@ -305,9 +315,10 @@ struct ENGINE_API FCompositeNavModifier : public FNavigationModifier
 	FORCEINLINE bool HasAreas() const { return Areas.Num() > 0; }
 	FORCEINLINE bool HasLowAreaModifiers() const { return bHasLowAreaModifiers; }
 	FORCEINLINE bool IsPerInstanceModifier() const { return bIsPerInstanceModifier; }
-	FORCEINLINE bool GetFillCollisionUnderneathForNavmesh() const { return bModifierFillCollisionUnderneathForNavmesh; }
-	FORCEINLINE void SetFillCollisionUnderneathForNavmesh(bool bInRejectNavmeshUnderneath) { bModifierFillCollisionUnderneathForNavmesh = bInRejectNavmeshUnderneath; }
-
+	FORCEINLINE bool GetFillCollisionUnderneathForNavmesh() const { return bFillCollisionUnderneathForNavmesh; }
+	FORCEINLINE void SetFillCollisionUnderneathForNavmesh(bool bValue) { bFillCollisionUnderneathForNavmesh = bValue; }
+	FORCEINLINE bool GetMaskFillCollisionUnderneathForNavmesh() const { return bMaskFillCollisionUnderneathForNavmesh; }
+	FORCEINLINE void SetMaskFillCollisionUnderneathForNavmesh(bool bValue) { bMaskFillCollisionUnderneathForNavmesh = bValue; }
 	FORCEINLINE void ReserveForAdditionalAreas(int32 AdditionalElementsCount) { Areas.Reserve(Areas.Num() + AdditionalElementsCount); }
 
 	void MarkPotentialLinks() { bHasPotentialLinks = true; }
@@ -338,5 +349,6 @@ private:
 	uint32 bAdjustHeight : 1;
 	uint32 bHasLowAreaModifiers : 1;
     uint32 bIsPerInstanceModifier : 1;
-	uint32 bModifierFillCollisionUnderneathForNavmesh : 1;
+	uint32 bFillCollisionUnderneathForNavmesh : 1;
+	uint32 bMaskFillCollisionUnderneathForNavmesh : 1;
 };
diff --git a/Engine/Source/Runtime/Engine/Public/EngineStats.h b/Engine/Source/Runtime/Engine/Public/EngineStats.h
--- a/Engine/Source/Runtime/Engine/Public/EngineStats.h
+++ b/Engine/Source/Runtime/Engine/Public/EngineStats.h
@@ -127,6 +127,7 @@ DECLARE_CYCLE_STAT_EXTERN(TEXT("Recast: test path"),STAT_Navigation_RecastTestPa
 DECLARE_CYCLE_STAT_EXTERN(TEXT("Recast: storing compressed layers"), STAT_RecastNavMeshGenerator_StoringCompressedLayers, STATGROUP_Navigation, );
 DECLARE_CYCLE_STAT_EXTERN(TEXT("Recast: build compressed layers"),STAT_Navigation_RecastBuildCompressedLayers,STATGROUP_Navigation, );
 DECLARE_CYCLE_STAT_EXTERN(TEXT("Recast: create height field"),STAT_Navigation_RecastCreateHeightField,STATGROUP_Navigation, );
+DECLARE_CYCLE_STAT_EXTERN(TEXT("Recast: compute rasterizaion masks"), STAT_Navigation_RecastComputeRasterizationMasks, STATGROUP_Navigation, );
 DECLARE_CYCLE_STAT_EXTERN(TEXT("Recast: rasterize triangles"),STAT_Navigation_RecastRasterizeTriangles,STATGROUP_Navigation, );
 DECLARE_CYCLE_STAT_EXTERN(TEXT("Recast: create voxel filter"),STAT_Navigation_RecastVoxelFilter,STATGROUP_Navigation, );
 DECLARE_CYCLE_STAT_EXTERN(TEXT("Recast: create filters"),STAT_Navigation_RecastFilter,STATGROUP_Navigation, );
diff --git a/Engine/Source/Runtime/NavigationSystem/Private/NavMesh/RecastNavMeshGenerator.cpp b/Engine/Source/Runtime/NavigationSystem/Private/NavMesh/RecastNavMeshGenerator.cpp
--- a/Engine/Source/Runtime/NavigationSystem/Private/NavMesh/RecastNavMeshGenerator.cpp
+++ b/Engine/Source/Runtime/NavigationSystem/Private/NavMesh/RecastNavMeshGenerator.cpp
@@ -2371,6 +2371,14 @@ void FRecastTileGenerator::ApplyVoxelFilter(rcHeightfield* HF, float WalkableRad
 	}
 }
 
+void FRecastTileGenerator::InitRasterizationMaskArray(const rcHeightfield* SolidHF, TInlineMaskArray& OutRasterizationMasks)
+{
+	const int CellCount = SolidHF->width * SolidHF->height;
+	OutRasterizationMasks.SetNumUninitialized(CellCount);
+	const uint8 AllowAllFlags = 0xFF;
+	FMemory::Memset(OutRasterizationMasks.GetData(), AllowAllFlags, CellCount*sizeof(TInlineMaskArray::ElementType));
+}
+
 void FRecastTileGenerator::PrepareVoxelCache(const TNavStatArray<uint8>& RawCollisionCache, const FCompositeNavModifier& InModifier, TNavStatArray<rcSpanCache>& SpanData)
 {
 	// tile's geometry: voxel cache (only for synchronous rebuilds)
@@ -2394,13 +2402,24 @@ void FRecastTileGenerator::PrepareVoxelCache(const TNavStatArray<uint8>& RawColl
 		CachedCollisions.Indices, CachedCollisions.Header.NumFaces,
 		TriAreas.GetData());
 
+	TInlineMaskArray RasterizationMasks;
+	if (InModifier.GetMaskFillCollisionUnderneathForNavmesh())
+	{
+		const int32 Mask = ~RC_PROJECT_TO_BOTTOM;
+		for (const FAreaNavModifier& ModifierArea : InModifier.GetAreas())
+		{
+			MarkRasterizationMask(0 /*ctx*/, VoxelCacheContext.RasterizeHF, ModifierArea, FTransform::Identity, Mask, RasterizationMasks);
+		}
+	}
+
 	// To prevent navmesh generation under the triangles, set the RC_PROJECT_TO_BOTTOM flag to true.
 	// This rasterize triangles as filled columns down to the HF lower bound.
 	const rcRasterizationFlags Flags = InModifier.GetFillCollisionUnderneathForNavmesh() ? RC_PROJECT_TO_BOTTOM : rcRasterizationFlags(0);
 
+	TInlineMaskArray::ElementType* MaskArray = RasterizationMasks.Num() > 0 ? RasterizationMasks.GetData() : nullptr;
 	rcRasterizeTriangles(0, CachedCollisions.Verts, CachedCollisions.Header.NumVerts,
 		CachedCollisions.Indices, TriAreas.GetData(), CachedCollisions.Header.NumFaces,
-		*VoxelCacheContext.RasterizeHF, WalkableClimbVX, Flags);
+		*VoxelCacheContext.RasterizeHF, WalkableClimbVX, Flags, MaskArray);
 
 	const int32 NumSpans = rcCountSpans(0, *VoxelCacheContext.RasterizeHF);
 	if (NumSpans > 0)
@@ -2484,6 +2503,8 @@ void FRecastTileGenerator::AppendModifier(const FCompositeNavModifier& Modifier,
 	}
 		
 	ModifierElement.Areas = Modifier.GetAreas();
+	ModifierElement.bMaskFillCollisionUnderneathForNavmesh = Modifier.GetMaskFillCollisionUnderneathForNavmesh();
+
 	Modifiers.Add(MoveTemp(ModifierElement));
 }
 
@@ -2647,6 +2668,7 @@ struct FTileRasterizationContext
 	struct rcHeightfieldLayerSet* LayerSet;
 	struct rcCompactHeightfield* CompactHF;
 	TArray<FNavMeshTileData> Layers;
+	FRecastTileGenerator::TInlineMaskArray RasterizationMasks;
 
 private:
 	rcRasterizationFlags RasterizationFlags;
@@ -2718,12 +2740,15 @@ ETimeSliceWorkResult FRecastTileGenerator::RasterizeGeometryRecastTimeSliced(FNa
 	}// fall through to next state
 	case ERasterizeGeomRecastTimeSlicedState::RasterizeTriangles:
 	{
+		ComputeRasterizationMasks(BuildContext, RasterContext);
+
 		QUICK_SCOPE_CYCLE_COUNTER(STAT_Navigation_RasterizeGeomRecastRasterizeTriangles);
 
+		TInlineMaskArray::ElementType* MaskArray = RasterContext.RasterizationMasks.Num() > 0 ? RasterContext.RasterizationMasks.GetData() : nullptr;
 		rcRasterizeTriangles(&BuildContext,
 			Coords.GetData(), NumVerts,
 			Indices.GetData(), RasterizeGeomRecastTriAreas.GetData(), NumFaces,
-			*RasterContext.SolidHF, TileConfig.walkableClimb, RasterizationFlags);
+			*RasterContext.SolidHF, TileConfig.walkableClimb, RasterizationFlags, MaskArray);
 
 		RasterizeGeomRecastTriAreas.Reset();
 
@@ -2763,10 +2788,11 @@ void FRecastTileGenerator::RasterizeGeometryRecast(FNavMeshBuildContext& BuildCo
 	{
 		QUICK_SCOPE_CYCLE_COUNTER(STAT_Navigation_RasterizeGeomRecastRasterizeTriangles);
 
+		TInlineMaskArray::ElementType* MaskArray = RasterContext.RasterizationMasks.Num() > 0 ? RasterContext.RasterizationMasks.GetData() : nullptr;
 		rcRasterizeTriangles(&BuildContext,
 			Coords.GetData(), NumVerts,
 			Indices.GetData(), RasterizeGeomRecastTriAreas.GetData(), NumFaces,
-			*RasterContext.SolidHF, TileConfig.walkableClimb, RasterizationFlags);
+			*RasterContext.SolidHF, TileConfig.walkableClimb, RasterizationFlags, MaskArray);
 	}
 
 	RasterizeGeomRecastTriAreas.Reset();
@@ -3340,6 +3366,8 @@ bool FRecastTileGenerator::GenerateCompressedLayers(FNavMeshBuildContext& BuildC
 	{
 		return false;
 	}
+		
+	ComputeRasterizationMasks(BuildContext, RasterContext);
 
 	RasterizeTriangles(BuildContext, RasterContext);
 	if (!RasterContext.SolidHF || RasterContext.SolidHF->pools == 0)
@@ -3814,6 +3842,33 @@ bool FRecastTileGenerator::GenerateNavigationData(FNavMeshBuildContext& BuildCon
 	return bGenDataLayer;
 }
 
+void FRecastTileGenerator::ComputeRasterizationMasks(FNavMeshBuildContext& BuildContext, FTileRasterizationContext& RasterContext)
+{
+	SCOPE_CYCLE_COUNTER(STAT_Navigation_RecastComputeRasterizationMasks);
+
+	UE_LOG(LogNavigationDataBuild, VeryVerbose, TEXT("      %s"), ANSI_TO_TCHAR(__FUNCTION__));
+
+	for (const FRecastAreaNavModifierElement& ModifierElement : Modifiers)
+	{
+		if (ModifierElement.bMaskFillCollisionUnderneathForNavmesh)
+		{
+			const int32 Mask = ~RC_PROJECT_TO_BOTTOM;
+			for (const FAreaNavModifier& ModifierArea : ModifierElement.Areas)
+			{
+				for (const FTransform& LocalToWorld : ModifierElement.PerInstanceTransform)
+				{
+					MarkRasterizationMask(&BuildContext, RasterContext.SolidHF, ModifierArea, LocalToWorld, Mask, RasterContext.RasterizationMasks);
+				}
+
+				if (ModifierElement.PerInstanceTransform.Num() == 0)
+				{
+					MarkRasterizationMask(&BuildContext, RasterContext.SolidHF, ModifierArea, FTransform::Identity, Mask, RasterContext.RasterizationMasks);
+				}
+			}
+		}
+	}
+}
+
 void FRecastTileGenerator::MarkDynamicAreas(dtTileCacheLayer& Layer)
 {
 	SCOPE_CYCLE_COUNTER(STAT_Navigation_RecastMarkAreas);
@@ -3910,6 +3965,164 @@ void FRecastTileGenerator::MarkDynamicArea(const FAreaNavModifier& Modifier, con
 	}
 }
 
+void MarkBoxMask(const float* pos, const float* extent, const int mask, rcHeightfield& hf, int* rasterizationMasks)
+{
+	float* orig = hf.bmin;
+	float bmin[3], bmax[3];
+	rcVsub(bmin, pos, extent);
+	rcVadd(bmax, pos, extent);
+
+	const int w = hf.width;
+	const int h = hf.height;
+	const float ics = 1.0f/hf.cs;
+
+	int minx = (int)floorf((bmin[0]-orig[0])*ics);
+	int minz = (int)floorf((bmin[2]-orig[2])*ics);
+	int maxx = (int)floorf((bmax[0]-orig[0])*ics);
+	int maxz = (int)floorf((bmax[2]-orig[2])*ics);
+
+	if (maxx < 0) return;
+	if (minx >= w) return;
+	if (maxz < 0) return;
+	if (minz >= h) return;
+
+	if (minx < 0) minx = 0;
+	if (maxx >= w) maxx = w-1;
+	if (minz < 0) minz = 0;
+	if (maxz >= h) maxz = h-1;
+
+	for (int z = minz; z <= maxz; ++z)
+	{
+		for (int x = minx; x <= maxx; ++x)
+		{
+			rasterizationMasks[x+z*w] &= mask;
+		}
+	}
+}
+
+int PointInPoly(const float* verts, int nv, const float* p)
+{
+	int i, j, c = 0;
+	for (i = 0, j = nv-1; i < nv; j = i++)
+	{
+		const float* vi = &verts[i*3];
+		const float* vj = &verts[j*3];
+		if (((vi[2] > p[2]) != (vj[2] > p[2])) &&
+			(p[0] < (vj[0]-vi[0]) * (p[2]-vi[2]) / (vj[2]-vi[2]) + vi[0]))
+			c = !c;
+	}
+	return c;
+}
+
+// Similar to rcMarkConvexPolyArea
+void MarkConvexMask(const int mask, const float* verts, const int nv, rcHeightfield& hf, int* rasterizationMasks)
+{
+	float* orig = hf.bmin;
+	float bmin[3], bmax[3];
+	rcVcopy(bmin, verts);
+	rcVcopy(bmax, verts);
+	for (int i = 1; i < nv; ++i)
+	{
+		rcVmin(bmin, &verts[i*3]);
+		rcVmax(bmax, &verts[i*3]);
+	}
+
+	const int w = hf.width;
+	const int h = hf.height;
+	const float ics = 1.0f/hf.cs;
+
+	int minx = (int)((bmin[0]-orig[0])*ics);
+	int minz = (int)((bmin[2]-orig[2])*ics);
+	int maxx = (int)((bmax[0]-orig[0])*ics);
+	int maxz = (int)((bmax[2]-orig[2])*ics);
+
+	if (maxx < 0) return;
+	if (minx >= w) return;
+	if (maxz < 0) return;
+	if (minz >= h) return;
+
+	if (minx < 0) minx = 0;
+	if (maxx >= w) maxx = w-1;
+	if (minz < 0) minz = 0;
+	if (maxz >= h) maxz = h-1;
+
+	for (int z = minz; z <= maxz; ++z)
+	{
+		for (int x = minx; x <= maxx; ++x)
+		{
+			float p[3];
+			p[0] = orig[0] + (float)(x+0.5f)*hf.cs;
+			p[1] = 0.0f;
+			p[2] = orig[2] + (float)(z+0.5f)*hf.cs;
+			if (PointInPoly(verts, nv, p))
+			{
+				rasterizationMasks[x+z*w] &= mask;
+			}
+		}
+	}
+}
+
+void FRecastTileGenerator::MarkRasterizationMask(rcContext* /*BuildContext*/, rcHeightfield* SolidHF,
+	const FAreaNavModifier& Modifier, const FTransform& LocalToWorld, const int32 Mask, TInlineMaskArray& OutMaskArray)
+{
+	FBox ModifierBounds = Modifier.GetBounds().TransformBy(LocalToWorld);
+	if (!ModifierBounds.Intersect(TileBB))
+	{
+		return;
+	}
+
+	// Init on first use
+	if (OutMaskArray.Num() == 0)
+	{
+		InitRasterizationMaskArray(SolidHF, OutMaskArray);
+	}
+
+	switch (Modifier.GetShapeType())
+	{
+	case ENavigationShapeType::Box:
+	{
+		FBoxNavAreaData BoxData;
+		Modifier.GetBox(BoxData);
+		FBox WorldBox = FBox::BuildAABB(BoxData.Origin, BoxData.Extent).TransformBy(LocalToWorld);
+		FBox RecastBox = Unreal2RecastBox(WorldBox);
+		FVector RecastPos;
+		FVector RecastExtent;
+		RecastBox.GetCenterAndExtents(RecastPos, RecastExtent);
+		check(OutMaskArray.Num() == SolidHF->width*SolidHF->height);
+		MarkBoxMask(&(RecastPos.X), &(RecastExtent.X), Mask, *SolidHF, OutMaskArray.GetData());
+	}
+	break;
+
+	case ENavigationShapeType::Convex:
+	{
+		FConvexNavAreaData ConvexData;
+		Modifier.GetConvex(ConvexData);
+		TArray<FVector> ConvexVerts;
+		const float Expand = 0.f;
+		GrowConvexHull(Expand, ConvexData.Points, ConvexVerts);
+		if (ConvexVerts.Num())
+		{
+			TArray<float> ConvexCoords;
+			ConvexCoords.AddZeroed(ConvexVerts.Num() * 3);
+			float* ItCoord = ConvexCoords.GetData();
+			for (int32 i = 0; i < ConvexVerts.Num(); i++)
+			{
+				const FVector RecastV = Unreal2RecastPoint(ConvexVerts[i]);
+				*ItCoord = RecastV.X; ItCoord++;
+				*ItCoord = RecastV.Y; ItCoord++;
+				*ItCoord = RecastV.Z; ItCoord++;
+			}
+
+			MarkConvexMask(Mask, ConvexCoords.GetData(), ConvexVerts.Num(), *SolidHF, OutMaskArray.GetData());
+		}
+	}
+	break;
+
+	default: 
+	break;
+	}
+}
+
 void FRecastTileGenerator::MarkDynamicArea(const FAreaNavModifier& Modifier, const FTransform& LocalToWorld, dtTileCacheLayer& Layer, const int32 AreaID, const int32* ReplaceIDPtr)
 {
 	const float ExpandBy = TileConfig.AgentRadius;
diff --git a/Engine/Source/Runtime/NavigationSystem/Private/NavModifierVolume.cpp b/Engine/Source/Runtime/NavigationSystem/Private/NavModifierVolume.cpp
--- a/Engine/Source/Runtime/NavigationSystem/Private/NavModifierVolume.cpp
+++ b/Engine/Source/Runtime/NavigationSystem/Private/NavModifierVolume.cpp
@@ -31,6 +31,17 @@ void ANavModifierVolume::GetNavigationData(FNavigationRelevantData& Data) const
 		FAreaNavModifier AreaMod(GetBrushComponent(), AreaClass);
 		Data.Modifiers.Add(AreaMod);
 	}
+
+	if (bMaskFillCollisionUnderneathForNavmesh)
+	{
+		if (GetBrushComponent()->Brush != nullptr)
+		{
+			const FBox& Box = GetBrushComponent()->Brush->Bounds.GetBox();
+			FAreaNavModifier AreaMod(Box, GetBrushComponent()->GetComponentTransform(), AreaClass);
+			Data.Modifiers.SetMaskFillCollisionUnderneathForNavmesh(true);
+			Data.Modifiers.Add(AreaMod);
+		}
+	}
 }
 
 FBox ANavModifierVolume::GetNavigationBounds() const
diff --git a/Engine/Source/Runtime/NavigationSystem/Private/NavigationSystem.cpp b/Engine/Source/Runtime/NavigationSystem/Private/NavigationSystem.cpp
--- a/Engine/Source/Runtime/NavigationSystem/Private/NavigationSystem.cpp
+++ b/Engine/Source/Runtime/NavigationSystem/Private/NavigationSystem.cpp
@@ -89,6 +89,7 @@ DEFINE_STAT(STAT_Navigation_RecastTestPath);
 DEFINE_STAT(STAT_RecastNavMeshGenerator_StoringCompressedLayers);
 DEFINE_STAT(STAT_Navigation_RecastBuildCompressedLayers);
 DEFINE_STAT(STAT_Navigation_RecastCreateHeightField);
+DEFINE_STAT(STAT_Navigation_RecastComputeRasterizationMasks);
 DEFINE_STAT(STAT_Navigation_RecastRasterizeTriangles);
 DEFINE_STAT(STAT_Navigation_RecastVoxelFilter);
 DEFINE_STAT(STAT_Navigation_RecastFilter);
diff --git a/Engine/Source/Runtime/NavigationSystem/Public/NavMesh/RecastNavMeshGenerator.h b/Engine/Source/Runtime/NavigationSystem/Public/NavMesh/RecastNavMeshGenerator.h
--- a/Engine/Source/Runtime/NavigationSystem/Public/NavMesh/RecastNavMeshGenerator.h
+++ b/Engine/Source/Runtime/NavigationSystem/Public/NavMesh/RecastNavMeshGenerator.h
@@ -166,6 +166,8 @@ struct FRecastAreaNavModifierElement
 	// Per instance transformations in unreal coords
 	// When empty areas are in world space
 	TArray<FTransform>	PerInstanceTransform;
+
+	bool bMaskFillCollisionUnderneathForNavmesh = false;
 };
 
 struct FRcTileBox
@@ -305,6 +307,8 @@ public:
 	FRecastInternalDebugData& GetMutableDebugData() { return DebugData; }
 #endif
 		
+	typedef TArray<int32, TInlineAllocator<4096>> TInlineMaskArray;
+
 protected:
 	/** Does the actual TimeSliced tile generation. 
 	 *	@note always trigger tile generation only via DoWorkTimeSliced(). This is a worker function
@@ -361,6 +365,12 @@ protected:
 
 	void ApplyVoxelFilter(struct rcHeightfield* SolidHF, float WalkableRadius);
 
+	/** Compute rasterization mask */
+	void InitRasterizationMaskArray(const rcHeightfield* SolidHF, TInlineMaskArray& OutRasterizationMasks);
+	void ComputeRasterizationMasks(FNavMeshBuildContext& BuildContext, FTileRasterizationContext& RasterContext);
+	void MarkRasterizationMask(rcContext* /*BuildContext*/, rcHeightfield* SolidHF,
+		const FAreaNavModifier& Modifier, const FTransform& LocalToWorld, const int32 Mask, TInlineMaskArray& OutMaskArray);
+
 	/** apply areas from DynamicAreas to layer */
 	void MarkDynamicAreas(dtTileCacheLayer& Layer);
 	void MarkDynamicArea(const FAreaNavModifier& Modifier, const FTransform& LocalToWorld, dtTileCacheLayer& Layer);
diff --git a/Engine/Source/Runtime/NavigationSystem/Public/NavModifierVolume.h b/Engine/Source/Runtime/NavigationSystem/Public/NavModifierVolume.h
--- a/Engine/Source/Runtime/NavigationSystem/Public/NavModifierVolume.h
+++ b/Engine/Source/Runtime/NavigationSystem/Public/NavModifierVolume.h
@@ -24,6 +24,10 @@ protected:
 	UPROPERTY(EditAnywhere, BlueprintReadOnly, Category=Default)
 	TSubclassOf<UNavArea> AreaClass;
 
+	/** Experimental: if set, the 2D space occupied by the volume box will ignore FillCollisionUnderneathForNavmesh */
+	UPROPERTY(EditAnywhere, Category = Default, AdvancedDisplay)
+	bool bMaskFillCollisionUnderneathForNavmesh;
+
 public:
 	ANavModifierVolume(const FObjectInitializer& ObjectInitializer);
 
diff --git a/Engine/Source/Runtime/Navmesh/Private/Recast/RecastRasterization.cpp b/Engine/Source/Runtime/Navmesh/Private/Recast/RecastRasterization.cpp
--- a/Engine/Source/Runtime/Navmesh/Private/Recast/RecastRasterization.cpp
+++ b/Engine/Source/Runtime/Navmesh/Private/Recast/RecastRasterization.cpp
@@ -440,14 +440,15 @@ static void rasterizeTri(const float* v0, const float* v1, const float* v2,
 						 const float* bmin, const float* bmax,
 						 const float cs, const float ics, const float ich, 
 						 const int flagMergeThr,
-						 const int rasterizationFlags /*UE4*/)
+						 const int rasterizationFlags, /*UE4*/
+	                     const int* rasterizationMasks /*UE4*/)
 {
 	rcEdgeHit* const hfEdgeHits = hf.EdgeHits; //this prevents a static analysis warning
 
 	const int w = hf.width;
 	const int h = hf.height;
 	const float by = bmax[1] - bmin[1];
-	const bool projectToBottom = rasterizationFlags & RC_PROJECT_TO_BOTTOM; //UE4
+	const int projectTriToBottom = rasterizationFlags & RC_PROJECT_TO_BOTTOM; //UE4
 
 	int intverts[3][2];
 
@@ -483,7 +484,8 @@ static void rasterizeTri(const float* v0, const float* v1, const float* v2,
 		if (triangle_smax > by) triangle_smax = by;
 
 		// Snap the span to the heightfield height grid.
-		unsigned short triangle_ismin = projectToBottom ? 0 : (unsigned short)rcClamp((int)floorf(triangle_smin * ich), 0, RC_SPAN_MAX_HEIGHT); //UE4
+		const int projectSpanToBottom = rasterizationMasks != nullptr ? (projectTriToBottom & rasterizationMasks[x0+y0*w]) : projectTriToBottom;	//UE4
+		unsigned short triangle_ismin = projectSpanToBottom ? 0 : (unsigned short)rcClamp((int)floorf(triangle_smin * ich), 0, RC_SPAN_MAX_HEIGHT); //UE4
 		unsigned short triangle_ismax = (unsigned short)rcClamp((int)ceilf(triangle_smax * ich), (int)triangle_ismin+1, RC_SPAN_MAX_HEIGHT);
 
 		addSpan(hf, x0, y0, triangle_ismin, triangle_ismax, area, flagMergeThr);
@@ -642,24 +644,48 @@ static void rasterizeTri(const float* v0, const float* v1, const float* v2,
 			}
 		}
 
-		// Snap the span to the heightfield height grid.
-		const unsigned short triangle_ismin_clamp = projectToBottom ? 0 : (unsigned short)rcClamp((int)triangle_ismin, 0, RC_SPAN_MAX_HEIGHT); //UE4
-		const unsigned short triangle_ismax_clamp = (unsigned short)rcClamp((int)triangle_ismin, (int)triangle_ismin_clamp+1, RC_SPAN_MAX_HEIGHT);
-
-		for (int y = y0; y <= y1; y++)
+		if (rasterizationMasks == nullptr) //UE4
 		{
-			int xloop0 = intMax(hf.RowExt[y + 1].MinCol, x0);
-			int xloop1 = intMin(hf.RowExt[y + 1].MaxCol, x1);
-			for (int x = xloop0; x <= xloop1; x++)
+			// Snap the span to the heightfield height grid.
+			const unsigned short triangle_ismin_clamp = projectTriToBottom ? 0 : (unsigned short)rcClamp((int)triangle_ismin, 0, RC_SPAN_MAX_HEIGHT); //UE4
+			const unsigned short triangle_ismax_clamp = (unsigned short)rcClamp((int)triangle_ismin, (int)triangle_ismin_clamp+1, RC_SPAN_MAX_HEIGHT);
+
+			for (int y = y0; y <= y1; y++)
 			{
-				addSpan(hf, x, y, triangle_ismin_clamp, triangle_ismax_clamp, area, flagMergeThr);
-			}
+				int xloop0 = intMax(hf.RowExt[y + 1].MinCol, x0);
+				int xloop1 = intMin(hf.RowExt[y + 1].MaxCol, x1);
+				for (int x = xloop0; x <= xloop1; x++)
+				{
+					addSpan(hf, x, y, triangle_ismin_clamp, triangle_ismax_clamp, area, flagMergeThr);
+				}
 
-			// reset for next triangle
-			hf.RowExt[y + 1].MinCol = hf.width + 2;
-			hf.RowExt[y + 1].MaxCol = -2;
+				// reset for next triangle
+				hf.RowExt[y + 1].MinCol = hf.width + 2;
+				hf.RowExt[y + 1].MaxCol = -2;
+			}
 		}
+		else
+		{
+// @UE4 BEGIN
+			for (int y = y0; y <= y1; y++)
+			{
+				int xloop0 = intMax(hf.RowExt[y + 1].MinCol, x0);
+				int xloop1 = intMin(hf.RowExt[y + 1].MaxCol, x1);
+				for (int x = xloop0; x <= xloop1; x++)
+				{
+					// Snap the span to the heightfield height grid.
+					const int projectSpanToBottom = projectTriToBottom & rasterizationMasks[x+y*w];		//UE4
+					const unsigned short triangle_ismin_clamp = projectSpanToBottom ? 0 : (unsigned short)rcClamp((int)triangle_ismin, 0, RC_SPAN_MAX_HEIGHT);
+					const unsigned short triangle_ismax_clamp = (unsigned short)rcClamp((int)triangle_ismin, (int)triangle_ismin_clamp+1, RC_SPAN_MAX_HEIGHT);
+					addSpan(hf, x, y, triangle_ismin_clamp, triangle_ismax_clamp, area, flagMergeThr);
+				}
 
+				// reset for next triangle
+				hf.RowExt[y + 1].MinCol = hf.width + 2;
+				hf.RowExt[y + 1].MaxCol = -2;
+			}
+// @UE4 END
+		}
 	}
 	else
 	{
@@ -825,7 +851,8 @@ static void rasterizeTri(const float* v0, const float* v1, const float* v2,
 				// Skip the span if it is outside the heightfield bbox
 				if (smin >= RC_SPAN_MAX_HEIGHT || smax < 0) continue;
 
-				smin = projectToBottom ? 0 : intMax(smin, 0); //UE4
+				const int projectSpanToBottom = rasterizationMasks != nullptr ? (projectTriToBottom & rasterizationMasks[x+y*w]) : projectTriToBottom; //UE4
+				smin = projectSpanToBottom ? 0 : intMax(smin, 0); //UE4
 				smax = intMin(intMax(smax,smin+1), RC_SPAN_MAX_HEIGHT);
 
 	#if TEST_NEW_RASTERIZER
@@ -892,7 +919,8 @@ static void rasterizeTri(const float* v0, const float* v1, const float* v2,
 						 const float* bmin, const float* bmax,
 						 const float cs, const float ics, const float ich,
 						 const int flagMergeThr,
-						 const int rasterizationFlags) //UE4
+						 const int rasterizationFlags,  //UE4
+						 const int* rasterizationMasks) //UE4
 {
 	const int w = hf.width;
 	const int h = hf.height;
@@ -980,8 +1008,7 @@ static void rasterizeTri(const float* v0, const float* v1, const float* v2,
 /// @see rcHeightfield
 void rcRasterizeTriangle(rcContext* ctx, const float* v0, const float* v1, const float* v2,
 						 const unsigned char area, rcHeightfield& solid,
-						 const int flagMergeThr, 
-						 const int rasterizationFlags) //UE4
+						 const int flagMergeThr, const int rasterizationFlags, const int* rasterizationMasks) //UE4
 {
 	rcAssert(ctx);
 
@@ -989,7 +1016,7 @@ void rcRasterizeTriangle(rcContext* ctx, const float* v0, const float* v1, const
 
 	const float ics = 1.0f/solid.cs;
 	const float ich = 1.0f/solid.ch;
-	rasterizeTri(v0, v1, v2, area, solid, solid.bmin, solid.bmax, solid.cs, ics, ich, flagMergeThr, rasterizationFlags); //UE4
+	rasterizeTri(v0, v1, v2, area, solid, solid.bmin, solid.bmax, solid.cs, ics, ich, flagMergeThr, rasterizationFlags, rasterizationMasks); //UE4
 
 	ctx->stopTimer(RC_TIMER_RASTERIZE_TRIANGLES);
 }
@@ -1001,7 +1028,7 @@ void rcRasterizeTriangle(rcContext* ctx, const float* v0, const float* v1, const
 /// @see rcHeightfield
 void rcRasterizeTriangles(rcContext* ctx, const float* verts, const int /*nv*/,
 						  const int* tris, const unsigned char* areas, const int nt,
-						  rcHeightfield& solid, const int flagMergeThr, const int rasterizationFlags) //UE4
+						  rcHeightfield& solid, const int flagMergeThr, const int rasterizationFlags, const int* rasterizationMasks) //UE4
 {
 	if (ctx)
 		ctx->startTimer(RC_TIMER_RASTERIZE_TRIANGLES);
@@ -1015,7 +1042,7 @@ void rcRasterizeTriangles(rcContext* ctx, const float* verts, const int /*nv*/,
 		const float* v1 = &verts[tris[i*3+1]*3];
 		const float* v2 = &verts[tris[i*3+2]*3];
 		// Rasterize.
-		rasterizeTri(v0, v1, v2, areas[i], solid, solid.bmin, solid.bmax, solid.cs, ics, ich, flagMergeThr, rasterizationFlags); //UE4
+		rasterizeTri(v0, v1, v2, areas[i], solid, solid.bmin, solid.bmax, solid.cs, ics, ich, flagMergeThr, rasterizationFlags, rasterizationMasks); //UE4
 	}
 	
 	if (ctx)
@@ -1029,7 +1056,7 @@ void rcRasterizeTriangles(rcContext* ctx, const float* verts, const int /*nv*/,
 /// @see rcHeightfield
 void rcRasterizeTriangles(rcContext* ctx, const float* verts, const int /*nv*/,
 						  const unsigned short* tris, const unsigned char* areas, const int nt,
-						  rcHeightfield& solid, const int flagMergeThr, const int rasterizationFlags) //UE4
+						  rcHeightfield& solid, const int flagMergeThr, const int rasterizationFlags, const int* rasterizationMasks) //UE4
 {
 	if (ctx)
 		ctx->startTimer(RC_TIMER_RASTERIZE_TRIANGLES);
@@ -1043,7 +1070,7 @@ void rcRasterizeTriangles(rcContext* ctx, const float* verts, const int /*nv*/,
 		const float* v1 = &verts[tris[i*3+1]*3];
 		const float* v2 = &verts[tris[i*3+2]*3];
 		// Rasterize.
-		rasterizeTri(v0, v1, v2, areas[i], solid, solid.bmin, solid.bmax, solid.cs, ics, ich, flagMergeThr, rasterizationFlags); //UE4
+		rasterizeTri(v0, v1, v2, areas[i], solid, solid.bmin, solid.bmax, solid.cs, ics, ich, flagMergeThr, rasterizationFlags, rasterizationMasks); //UE4
 	}
 	
 	if (ctx)
@@ -1056,7 +1083,7 @@ void rcRasterizeTriangles(rcContext* ctx, const float* verts, const int /*nv*/,
 ///
 /// @see rcHeightfield
 void rcRasterizeTriangles(rcContext* ctx, const float* verts, const unsigned char* areas, const int nt,
-						  rcHeightfield& solid, const int flagMergeThr, const int rasterizationFlags) //UE4
+						  rcHeightfield& solid, const int flagMergeThr, const int rasterizationFlags, const int* rasterizationMasks) //UE4
 {
 	if (ctx)
 		ctx->startTimer(RC_TIMER_RASTERIZE_TRIANGLES);
@@ -1070,7 +1097,7 @@ void rcRasterizeTriangles(rcContext* ctx, const float* verts, const unsigned cha
 		const float* v1 = &verts[(i*3+1)*3];
 		const float* v2 = &verts[(i*3+2)*3];
 		// Rasterize.
-		rasterizeTri(v0, v1, v2, areas[i], solid, solid.bmin, solid.bmax, solid.cs, ics, ich, flagMergeThr, rasterizationFlags); //UE4
+		rasterizeTri(v0, v1, v2, areas[i], solid, solid.bmin, solid.bmax, solid.cs, ics, ich, flagMergeThr, rasterizationFlags, rasterizationMasks); //UE4
 	}
 	
 	if (ctx)
diff --git a/Engine/Source/Runtime/Navmesh/Public/Recast/Recast.h b/Engine/Source/Runtime/Navmesh/Public/Recast/Recast.h
--- a/Engine/Source/Runtime/Navmesh/Public/Recast/Recast.h
+++ b/Engine/Source/Runtime/Navmesh/Public/Recast/Recast.h
@@ -925,9 +925,12 @@ NAVMESH_API void rcCacheSpans(rcContext* ctx, rcHeightfield& hf, rcSpanCache* ca
 ///  @param[in,out]	solid			An initialized heightfield.
 ///  @param[in]		flagMergeThr	The distance where the walkable flag is favored over the non-walkable flag.
 ///  								[Limit: >= 0] [Units: vx]
+///  @param[in]     rtzFlags		Flags to change the rasterization behavior			//UE4
+///  @param[in]     rtzMasks		Mask for the rasterization flags [Size: hf.w*hf.h]	//UE4
 NAVMESH_API void rcRasterizeTriangle(rcContext* ctx, const float* v0, const float* v1, const float* v2,
 						 const unsigned char area, rcHeightfield& solid,
-						 const int flagMergeThr = 1);
+						 const int flagMergeThr = 1, 
+						 const int rasterizationFlags = 0, const int* rasterizationMasks = nullptr); //UE4
 
 /// Rasterizes an indexed triangle mesh into the specified heightfield.
 ///  @ingroup recast
@@ -940,10 +943,12 @@ NAVMESH_API void rcRasterizeTriangle(rcContext* ctx, const float* v0, const floa
 ///  @param[in,out]	solid			An initialized heightfield.
 ///  @param[in]		flagMergeThr	The distance where the walkable flag is favored over the non-walkable flag. 
 ///  								[Limit: >= 0] [Units: vx]
-///  @param[in]     int				Flags to change the rasterization behavior		//UE4
+///  @param[in]     rtzFlags		Flags to change the rasterization behavior			//UE4
+///  @param[in]     rtzMasks		Mask for the rasterization flags [Size: hf.w*hf.h]	//UE4
 NAVMESH_API void rcRasterizeTriangles(rcContext* ctx, const float* verts, const int nv,
 						  const int* tris, const unsigned char* areas, const int nt,
-						  rcHeightfield& solid, const int flagMergeThr = 1, const int rasterizationFlags = 0 /*UE4*/);
+						  rcHeightfield& solid, const int flagMergeThr = 1, 
+						  const int rasterizationFlags = 0, const int* rasterizationMasks = nullptr); //UE4
 
 /// Rasterizes an indexed triangle mesh into the specified heightfield.
 ///  @ingroup recast
@@ -956,10 +961,12 @@ NAVMESH_API void rcRasterizeTriangles(rcContext* ctx, const float* verts, const
 ///  @param[in,out]	solid			An initialized heightfield.
 ///  @param[in]		flagMergeThr	The distance where the walkable flag is favored over the non-walkable flag. 
 ///  								[Limit: >= 0] [Units: vx]
-///  @param[in]     int				Flags to change the rasterization behavior		//UE4
+///  @param[in]     rtzFlags		Flags to change the rasterization behavior			//UE4
+///  @param[in]     rtzMasks		Mask for the rasterization flags [Size: hf.w*hf.h]	//UE4
 NAVMESH_API void rcRasterizeTriangles(rcContext* ctx, const float* verts, const int nv,
 						  const unsigned short* tris, const unsigned char* areas, const int nt,
-						  rcHeightfield& solid, const int flagMergeThr = 1, const int rasterizationFlags = 0 /*UE4*/);
+						  rcHeightfield& solid, const int flagMergeThr = 1, 
+						  const int rasterizationFlags = 0, const int* rasterizationMasks = nullptr); //UE4
 
 /// Rasterizes triangles into the specified heightfield.
 ///  @ingroup recast
@@ -970,8 +977,11 @@ NAVMESH_API void rcRasterizeTriangles(rcContext* ctx, const float* verts, const
 ///  @param[in,out]	solid			An initialized heightfield.
 ///  @param[in]		flagMergeThr	The distance where the walkable flag is favored over the non-walkable flag. 
 ///  								[Limit: >= 0] [Units: vx]
+///  @param[in]     rtzFlags		Flags to change the rasterization behavior			//UE4
+///  @param[in]     rtzMasks		Mask for the rasterization flags [Size: hf.w*hf.h]	//UE4
 NAVMESH_API void rcRasterizeTriangles(rcContext* ctx, const float* verts, const unsigned char* areas, const int nt,
-						  rcHeightfield& solid, const int flagMergeThr = 1);
+						  rcHeightfield& solid, const int flagMergeThr = 1,
+						  const int rasterizationFlags = 0, const int* rasterizationMasks = nullptr); //UE4
 
 /// Marks non-walkable spans as walkable if their maximum is within @p walkableClimp of a walkable neighbor. 
 ///  @ingroup recast
