Fix recast not working if WITH_PHYSX is false, but WITH_CHAOS is true

diff --git a/Engine/Source/Runtime/NavigationSystem/Private/NavMesh/RecastNavMeshGenerator.cpp b/Engine/Source/Runtime/NavigationSystem/Private/NavMesh/RecastNavMeshGenerator.cpp
--- a/Engine/Source/Runtime/NavigationSystem/Private/NavMesh/RecastNavMeshGenerator.cpp
+++ b/Engine/Source/Runtime/NavigationSystem/Private/NavMesh/RecastNavMeshGenerator.cpp
@@ -1224,7 +1224,6 @@ FORCEINLINE_DEBUGGABLE void ExportRigidBodySetup(UBodySetup& BodySetup, TNavStat
 
 FORCEINLINE_DEBUGGABLE void ExportComponent(UActorComponent* Component, FRecastGeometryExport& GeomExport, const FBox* ClipBounds=NULL)
 {
-#if WITH_PHYSX
 	bool bHasData = false;
 
 	UPrimitiveComponent* PrimComp = Cast<UPrimitiveComponent>(Component);
@@ -1247,7 +1246,6 @@ FORCEINLINE_DEBUGGABLE void ExportComponent(UActorComponent* Component, FRecastG
 			GeomExport.SlopeOverride = BodySetup->WalkableSlopeOverride;
 		}
 	}
-#endif // WITH_PHYSX
 }
 
 FORCEINLINE void TransformVertexSoupToRecast(const TArray<FVector>& VertexSoup, TNavStatArray<FVector>& Verts, TNavStatArray<int32>& Faces)
diff --git a/Engine/Source/Runtime/NavigationSystem/Private/NavigationSystemTypes.cpp b/Engine/Source/Runtime/NavigationSystem/Private/NavigationSystemTypes.cpp
--- a/Engine/Source/Runtime/NavigationSystem/Private/NavigationSystemTypes.cpp
+++ b/Engine/Source/Runtime/NavigationSystem/Private/NavigationSystemTypes.cpp
@@ -235,6 +235,8 @@ namespace NavigationHelper
 	{
 #if WITH_PHYSX
 		const bool bBodyHasGeometry = (BodySetup.AggGeom.GetElementCount() > 0 || BodySetup.TriMeshes.Num() > 0);
+#elif WITH_CHAOS
+		const bool bBodyHasGeometry = (BodySetup.AggGeom.GetElementCount() > 0 || BodySetup.ChaosTriMeshes.Num() > 0);
 #else
 		const bool bBodyHasGeometry = (BodySetup.AggGeom.GetElementCount() > 0);
 #endif
