Minor changes to previous CLfor moving tile regen time slice management logic to navigation system as suggested by Aris https://p4-swarm.epicgames.net/reviews/11589014/

diff --git a/Engine/Source/Runtime/NavigationSystem/Private/NavMesh/RecastNavMeshGenerator.cpp b/Engine/Source/Runtime/NavigationSystem/Private/NavMesh/RecastNavMeshGenerator.cpp
--- a/Engine/Source/Runtime/NavigationSystem/Private/NavMesh/RecastNavMeshGenerator.cpp
+++ b/Engine/Source/Runtime/NavigationSystem/Private/NavMesh/RecastNavMeshGenerator.cpp
@@ -4421,7 +4421,7 @@ void FRecastNavMeshGenerator::Init()
 
 	if (NavSys)
 	{
-		SyncTimeSlicedData.TimeSliceManager = &(NavSys->GetNavRegenTimeSliceManager());
+		SyncTimeSlicedData.TimeSliceManager = &(NavSys->GetMutableNavRegenTimeSliceManager());
 	}
 	else
 	{
@@ -5821,7 +5821,7 @@ TArray<uint32> FRecastNavMeshGenerator::ProcessTileTasksSyncTimeSliced()
 				//next frame (as we've finished time slice processing the last tile)
 				if (!SyncTimeSlicedData.bNextTimeSliceRegenActive)
 				{
-					return EndFunction(false, true);
+					return EndFunction(false /* bCalcTileRegenDuration */, bHadWorktoDo);
 				}
 
 				SyncTimeSlicedData.TileGeneratorSync = CreateTileGeneratorFromPendingElement(TileLocation);
@@ -5844,7 +5844,7 @@ TArray<uint32> FRecastNavMeshGenerator::ProcessTileTasksSyncTimeSliced()
 
 				if (SyncTimeSlicedData.TimeSliceManager->GetTimeSlicer().TestTimeSliceFinished())
 				{
-					return EndFunction(true, true);
+					return EndFunction(true /* bCalcTileRegenDuration */, bHadWorktoDo);
 				}
 			}
 			else
@@ -5875,7 +5875,7 @@ TArray<uint32> FRecastNavMeshGenerator::ProcessTileTasksSyncTimeSliced()
 
 				if (SyncTimeSlicedData.TimeSliceManager->GetTimeSlicer().IsTimeSliceFinishedCached())
 				{
-					return EndFunction(true, true);
+					return EndFunction(true /* bCalcTileRegenDuration */, bHadWorktoDo);
 				}
 			}//fall through to next state
 			case EProcessTileTasksSyncTimeSlicedState::AddGeneratedTiles:
@@ -5889,7 +5889,7 @@ TArray<uint32> FRecastNavMeshGenerator::ProcessTileTasksSyncTimeSliced()
 				 
 				if (SyncTimeSlicedData.TimeSliceManager->GetTimeSlicer().IsTimeSliceFinishedCached())
 				{
-					return EndFunction(true, true);
+					return EndFunction(true /* bCalcTileRegenDuration */, bHadWorktoDo);
 				}
 			}//fall through to next state
 			case EProcessTileTasksSyncTimeSlicedState::StoreCompessedTileCacheLayers:
@@ -5922,7 +5922,8 @@ TArray<uint32> FRecastNavMeshGenerator::ProcessTileTasksSyncTimeSliced()
 				//test time slice 
 				if (SyncTimeSlicedData.TimeSliceManager->GetTimeSlicer().TestTimeSliceFinished())
 				{
-					return EndFunction(false, true);
+					//we just calculated and set TileRegenDuration so no need to calculate it again
+					return EndFunction(false /* bCalcTileRegenDuration */, bHadWorktoDo);
 				}
 			}
 			break;
@@ -5937,8 +5938,8 @@ TArray<uint32> FRecastNavMeshGenerator::ProcessTileTasksSyncTimeSliced()
 
 	// we only hit this if we have processed too many tiles in a frame and we will already
 	// have calculated the tile regen duration, or if we have processed no tiles and we also
-	// don't want to calcualte the tile regen duration
-	return EndFunction(false, bHadWorktoDo);
+	// don't want to calculate the tile regen duration
+	return EndFunction(false /* bCalcTileRegenDuration */, bHadWorktoDo);
 }
 
 //this code path is approx 10% faster than ProcessTileTasksSyncTimeSliced, however it spikes far worse for most use cases.
diff --git a/Engine/Source/Runtime/NavigationSystem/Private/NavigationSystem.cpp b/Engine/Source/Runtime/NavigationSystem/Private/NavigationSystem.cpp
--- a/Engine/Source/Runtime/NavigationSystem/Private/NavigationSystem.cpp
+++ b/Engine/Source/Runtime/NavigationSystem/Private/NavigationSystem.cpp
@@ -1045,8 +1045,7 @@ void UNavigationSystemV1::CalcTimeSlicedUpdateData(TArray<double>& OutCurrentTim
 	OutNumTimeSlicedRemainingBuildTasks = 0;
 	OutIsTimeSlicingArray.SetNumZeroed(NavDataSet.Num());
 	bOutAnyNonTimeSlicedGenerators = false;
-	OutCurrentTimeSlicedBuildTaskDurations.Empty();
-	OutCurrentTimeSlicedBuildTaskDurations.Reserve(NavDataSet.Num());
+	OutCurrentTimeSlicedBuildTaskDurations.Reset(NavDataSet.Num());
 
 	for (int32 NavDataIdx = 0; NavDataIdx < NavDataSet.Num(); ++NavDataIdx)
 	{
diff --git a/Engine/Source/Runtime/NavigationSystem/Public/NavigationSystem.h b/Engine/Source/Runtime/NavigationSystem/Public/NavigationSystem.h
--- a/Engine/Source/Runtime/NavigationSystem/Public/NavigationSystem.h
+++ b/Engine/Source/Runtime/NavigationSystem/Public/NavigationSystem.h
@@ -961,7 +961,7 @@ public:
 	//----------------------------------------------------------------------//
 	void CycleNavigationDataDrawn();
 
-	FNavRegenTimeSliceManager& GetNavRegenTimeSliceManager() { return NavRegenTimeSliceManager; }
+	FNavRegenTimeSliceManager& GetMutableNavRegenTimeSliceManager() { return NavRegenTimeSliceManager; }
 
 protected:
 
