Work to Deprecate IsNavigationBuildInProgress(bool) and IsBuildInProgress(bool).  These functions behave inconsistently between different navmesh build types and were causing issues. Not only that but their default params were set opposite ways (true / false).

diff --git a/Engine/Source/Editor/UnrealEd/Private/NavigationBuildingNotification.cpp b/Engine/Source/Editor/UnrealEd/Private/NavigationBuildingNotification.cpp
--- a/Engine/Source/Editor/UnrealEd/Private/NavigationBuildingNotification.cpp
+++ b/Engine/Source/Editor/UnrealEd/Private/NavigationBuildingNotification.cpp
@@ -155,7 +155,7 @@ void FNavigationBuildingNotificationImpl::Tick(float DeltaTime)
 		
 		UNavigationSystemV1* NavSys = FNavigationSystem::GetCurrent<UNavigationSystemV1>(EditorContext.World());
 		const bool bBuildInProgress = NavSys != nullptr
-			&& NavSys->IsNavigationBuildInProgress( GetDefault<ULevelEditorMiscSettings>()->bNavigationAutoUpdate ? true : false) == true
+			&& (GetDefault<ULevelEditorMiscSettings>()->bNavigationAutoUpdate ? NavSys->IsNavigationBuildInProgress() : NavSys->GetNumRunningBuildTasks() > 0)
 			&& NavSys->GetNumRemainingBuildTasks() > 0;
 
 		if (!bPreviouslyDetectedBuild && bBuildInProgress)
diff --git a/Engine/Source/Runtime/Engine/Public/AI/NavDataGenerator.h b/Engine/Source/Runtime/Engine/Public/AI/NavDataGenerator.h
--- a/Engine/Source/Runtime/Engine/Public/AI/NavDataGenerator.h
+++ b/Engine/Source/Runtime/Engine/Public/AI/NavDataGenerator.h
@@ -32,7 +32,11 @@ public:
 	virtual void RebuildDirtyAreas(const TArray<FNavigationDirtyArea>& DirtyAreas) {}
 
 	/** determines whether this generator is performing navigation building actions at the moment*/
-	virtual bool IsBuildInProgress(bool bCheckDirtyToo = false) const { return false; }	
+	UE_DEPRECATED(4.26, "This function is deprecated. Please use IsBuildInProgressCheckDirty")
+	virtual bool IsBuildInProgress(bool bCheckDirtyToo = false) const { return false; }
+
+	/** determines whether this generator is performing navigation building actions at the moment, dirty areas are also checked */
+	virtual bool IsBuildInProgressCheckDirty() const { return false; }
 
 	virtual bool GetTimeSliceData(int32& OutNumRemainingBuildTasks, double& OutCurrentBuildTaskDuration) const { OutNumRemainingBuildTasks = 0; OutCurrentBuildTaskDuration = 0.; return false; }
 
diff --git a/Engine/Source/Runtime/NavigationSystem/Private/NavGraph/NavGraphGenerator.cpp b/Engine/Source/Runtime/NavigationSystem/Private/NavGraph/NavGraphGenerator.cpp
--- a/Engine/Source/Runtime/NavigationSystem/Private/NavGraph/NavGraphGenerator.cpp
+++ b/Engine/Source/Runtime/NavigationSystem/Private/NavGraph/NavGraphGenerator.cpp
@@ -15,17 +15,6 @@ FNavGraphGenerator::~FNavGraphGenerator()
 
 }
 
-//----------------------------------------------------------------------//
-// FNavDataGenerator overrides
-//----------------------------------------------------------------------//
-bool FNavGraphGenerator::IsBuildInProgress(bool bCheckDirtyToo) const
-{
-	return false;
-}
-
-//----------------------------------------------------------------------//
-// 
-//----------------------------------------------------------------------//
 void FNavGraphGenerator::Init()
 {
 
diff --git a/Engine/Source/Runtime/NavigationSystem/Private/NavGraph/NavGraphGenerator.h b/Engine/Source/Runtime/NavigationSystem/Private/NavGraph/NavGraphGenerator.h
--- a/Engine/Source/Runtime/NavigationSystem/Private/NavGraph/NavGraphGenerator.h
+++ b/Engine/Source/Runtime/NavigationSystem/Private/NavGraph/NavGraphGenerator.h
@@ -22,12 +22,6 @@ private:
 	FNavGraphGenerator(FNavGraphGenerator const& NoCopy) { check(0); };
 	FNavGraphGenerator& operator=(FNavGraphGenerator const& NoCopy) { check(0); return *this; }
 
-public:
-	//----------------------------------------------------------------------//
-	// FNavDataGenerator overrides
-	//----------------------------------------------------------------------//
-	virtual bool IsBuildInProgress(bool bCheckDirtyToo = false) const override;
-
 private:
 	// Performs initial setup of member variables so that generator is ready to do its thing from this point on
 	void Init();
diff --git a/Engine/Source/Runtime/NavigationSystem/Private/NavMesh/RecastNavMeshGenerator.cpp b/Engine/Source/Runtime/NavigationSystem/Private/NavMesh/RecastNavMeshGenerator.cpp
--- a/Engine/Source/Runtime/NavigationSystem/Private/NavMesh/RecastNavMeshGenerator.cpp
+++ b/Engine/Source/Runtime/NavigationSystem/Private/NavMesh/RecastNavMeshGenerator.cpp
@@ -4913,7 +4913,7 @@ void FRecastNavMeshGenerator::RemoveTiles(const TArray<FIntPoint>& Tiles)
 			PendingDirtyTiles.Remove(DirtyTile);
 		}
 
-		if (SyncTimeSlicedData.TileGeneratorSync.Get())
+		if (SyncTimeSlicedData.TileGeneratorSync.IsValid())
 		{
 			if (SyncTimeSlicedData.TileGeneratorSync->GetTileX() == TileXY.X && SyncTimeSlicedData.TileGeneratorSync->GetTileY() == TileXY.Y)
 			{
@@ -5312,7 +5312,7 @@ bool FRecastNavMeshGenerator::HasDirtyTiles() const
 {
 	return (PendingDirtyTiles.Num() > 0 
 		|| RunningDirtyTiles.Num() > 0
-		|| SyncTimeSlicedData.TileGeneratorSync.Get() != nullptr
+		|| SyncTimeSlicedData.TileGeneratorSync.IsValid()
 		);
 }
 
@@ -6167,11 +6167,17 @@ void FRecastNavMeshGenerator::ExportAggregatedGeometry(const FKAggregateGeom& Ag
 	}
 }
 
-bool FRecastNavMeshGenerator::IsBuildInProgress(bool bCheckDirtyToo) const
+bool FRecastNavMeshGenerator::IsBuildInProgressCheckDirty() const
 {
 	return RunningDirtyTiles.Num()
-		|| (bCheckDirtyToo && PendingDirtyTiles.Num())
-		|| SyncTimeSlicedData.TileGeneratorSync.Get();
+		|| PendingDirtyTiles.Num()
+		|| SyncTimeSlicedData.TileGeneratorSync.IsValid();
+}
+
+//deprecated
+bool FRecastNavMeshGenerator::IsBuildInProgress(bool bCheckDirtyToo) const
+{
+	return IsBuildInProgressCheckDirty();
 }
 
 #if !RECAST_ASYNC_REBUILDING
diff --git a/Engine/Source/Runtime/NavigationSystem/Private/NavigationSystem.cpp b/Engine/Source/Runtime/NavigationSystem/Private/NavigationSystem.cpp
--- a/Engine/Source/Runtime/NavigationSystem/Private/NavigationSystem.cpp
+++ b/Engine/Source/Runtime/NavigationSystem/Private/NavigationSystem.cpp
@@ -1849,7 +1849,7 @@ bool UNavigationSystemV1::IsNavigationBuilt(const AWorldSettings* Settings) cons
 #if WITH_EDITOR
 				|| GEditor != NULL
 #endif // WITH_EDITOR
-				) && (Generator == NULL || Generator->IsBuildInProgress(/*bCheckDirtyToo=*/true) == true))
+				) && (Generator == NULL || Generator->IsBuildInProgressCheckDirty() == true))
 			{
 				bIsBuilt = false;
 				break;
@@ -3576,7 +3576,7 @@ void UNavigationSystemV1::RebuildDirtyAreas(float DeltaSeconds)
 	DefaultDirtyAreasController.Tick(DeltaSeconds, NavDataSet, bForceRebuilding);
 }
 
-bool UNavigationSystemV1::IsNavigationBuildInProgress(bool bCheckDirtyToo)
+bool UNavigationSystemV1::IsNavigationBuildInProgress()
 {
 	bool bRet = false;
 
@@ -3590,8 +3590,8 @@ bool UNavigationSystemV1::IsNavigationBuildInProgress(bool bCheckDirtyToo)
 	for (int32 NavDataIndex = 0; NavDataIndex < NavDataSet.Num(); ++NavDataIndex)
 	{
 		ANavigationData* NavData = NavDataSet[NavDataIndex];
-		if (NavData != NULL && NavData->GetGenerator() != NULL 
-			&& NavData->GetGenerator()->IsBuildInProgress(bCheckDirtyToo) == true)
+		if (NavData != NULL && NavData->GetGenerator() != NULL
+			&& NavData->GetGenerator()->IsBuildInProgressCheckDirty() == true)
 		{
 			bRet = true;
 			break;
@@ -3601,6 +3601,12 @@ bool UNavigationSystemV1::IsNavigationBuildInProgress(bool bCheckDirtyToo)
 	return bRet;
 }
 
+//deprecated
+bool UNavigationSystemV1::IsNavigationBuildInProgress(bool bCheckDirtyToo)
+{
+	return IsNavigationBuildInProgress();
+}
+
 void UNavigationSystemV1::OnNavigationGenerationFinished(ANavigationData& NavData)
 {
 	OnNavigationGenerationFinishedDelegate.Broadcast(&NavData);
diff --git a/Engine/Source/Runtime/NavigationSystem/Public/NavMesh/RecastNavMeshGenerator.h b/Engine/Source/Runtime/NavigationSystem/Public/NavMesh/RecastNavMeshGenerator.h
--- a/Engine/Source/Runtime/NavigationSystem/Public/NavMesh/RecastNavMeshGenerator.h
+++ b/Engine/Source/Runtime/NavigationSystem/Public/NavMesh/RecastNavMeshGenerator.h
@@ -606,7 +606,14 @@ public:
 
 	/** Asks generator to update navigation affected by DirtyAreas */
 	virtual void RebuildDirtyAreas(const TArray<FNavigationDirtyArea>& DirtyAreas) override;
+
+	PRAGMA_DISABLE_DEPRECATION_WARNINGS
+	UE_DEPRECATED(4.26, "This function is deprecated. Please use IsBuildInProgressCheckDirty")
 	virtual bool IsBuildInProgress(bool bCheckDirtyToo = false) const override;
+	PRAGMA_ENABLE_DEPRECATION_WARNINGS
+
+	/** determines whether this generator is performing navigation building actions at the moment, dirty areas are also checked */
+	virtual bool IsBuildInProgressCheckDirty() const override;
 
 #if !RECAST_ASYNC_REBUILDING
 	/** returns true if we are time slicing and the data is valid to use false otherwise*/
diff --git a/Engine/Source/Runtime/NavigationSystem/Public/NavigationSystem.h b/Engine/Source/Runtime/NavigationSystem/Public/NavigationSystem.h
--- a/Engine/Source/Runtime/NavigationSystem/Public/NavigationSystem.h
+++ b/Engine/Source/Runtime/NavigationSystem/Public/NavigationSystem.h
@@ -838,9 +838,13 @@ public:
 	virtual void OnNavigationBoundsAdded(ANavMeshBoundsVolume* NavVolume);
 	virtual void OnNavigationBoundsRemoved(ANavMeshBoundsVolume* NavVolume);
 
-	/** Used to display "navigation building in progress" notify */
-	bool IsNavigationBuildInProgress(bool bCheckDirtyToo = true);
+	/** determines whether any generator is performing navigation building actions at the moment */
+	UE_DEPRECATED(4.26, "This function is deprecated.  Please use IsNavigationBuildInProgress")
+	bool IsNavigationBuildInProgress(bool bCheckDirtyToo);
 
+	/** determines whether any generator is performing navigation building actions at the moment, dirty areas are also checked */
+	bool IsNavigationBuildInProgress();
+	
 	virtual void OnNavigationGenerationFinished(ANavigationData& NavData);
 
 	/** Used to display "navigation building in progress" counter */
