Fix chaos export height field slice version, used when doing lazy octree gather geometry

#rnx

diff --git a/Engine/Source/Runtime/Engine/Public/AI/NavigationSystemHelpers.h b/Engine/Source/Runtime/Engine/Public/AI/NavigationSystemHelpers.h
--- a/Engine/Source/Runtime/Engine/Public/AI/NavigationSystemHelpers.h
+++ b/Engine/Source/Runtime/Engine/Public/AI/NavigationSystemHelpers.h
@@ -43,13 +43,14 @@ struct FNavigableGeometryExport
 	virtual void ExportPxTriMesh32Bit(physx::PxTriangleMesh const * const TriMesh, const FTransform& LocalToWorld) = 0;
 	virtual void ExportPxConvexMesh(physx::PxConvexMesh const * const ConvexMesh, const FTransform& LocalToWorld) = 0;
 	virtual void ExportPxHeightField(physx::PxHeightField const * const HeightField, const FTransform& LocalToWorld) = 0;
+	virtual void ExportPxHeightFieldSlice(const FNavHeightfieldSamples& PrefetchedHeightfieldSamples, const int32 NumRows, const int32 NumCols, const FTransform& LocalToWorld, const FBox& SliceBox) = 0;
 #endif // WITH_PHYSX
 #if WITH_CHAOS
 	virtual void ExportChaosTriMesh(const Chaos::FTriangleMeshImplicitObject* const TriMesh, const FTransform& LocalToWorld) = 0;
 	virtual void ExportChaosConvexMesh(const FKConvexElem* const Convex, const FTransform& LocalToWorld) = 0;
 	virtual void ExportChaosHeightField(const Chaos::THeightField<float>* const Heightfield, const FTransform& LocalToWorld) = 0;
+	virtual void ExportChaosHeightFieldSlice(const FNavHeightfieldSamples& PrefetchedHeightfieldSamples, const int32 NumRows, const int32 NumCols, const FTransform& LocalToWorld, const FBox& SliceBox) = 0;
 #endif
-	virtual void ExportHeightFieldSlice(const FNavHeightfieldSamples& PrefetchedHeightfieldSamples, const int32 NumRows, const int32 NumCols, const FTransform& LocalToWorld, const FBox& SliceBox) = 0;
 	virtual void ExportRigidBodySetup(UBodySetup& BodySetup, const FTransform& LocalToWorld) = 0;
 	virtual void ExportCustomMesh(const FVector* VertexBuffer, int32 NumVerts, const int32* IndexBuffer, int32 NumIndices, const FTransform& LocalToWorld) = 0;
 	
diff --git a/Engine/Source/Runtime/Landscape/Private/LandscapeCollision.cpp b/Engine/Source/Runtime/Landscape/Private/LandscapeCollision.cpp
--- a/Engine/Source/Runtime/Landscape/Private/LandscapeCollision.cpp
+++ b/Engine/Source/Runtime/Landscape/Private/LandscapeCollision.cpp
@@ -2007,7 +2007,11 @@ void ULandscapeHeightfieldCollisionComponent::GatherGeometrySlice(FNavigableGeom
 		FTransform HFToW = GetComponentTransform();
 		HFToW.MultiplyScale3D(FVector(CollisionScale, CollisionScale, LANDSCAPE_ZSCALE));
 
-		GeomExport.ExportHeightFieldSlice(CachedHeightFieldSamples, HeightfieldRowsCount, HeightfieldColumnsCount, HFToW, SliceBox);
+#if WITH_PHYSX && PHYSICS_INTERFACE_PHYSX
+		GeomExport.ExportPxHeightFieldSlice(CachedHeightFieldSamples, HeightfieldRowsCount, HeightfieldColumnsCount, HFToW, SliceBox);
+#elif WITH_CHAOS
+		GeomExport.ExportChaosHeightFieldSlice(CachedHeightFieldSamples, HeightfieldRowsCount, HeightfieldColumnsCount, HFToW, SliceBox);
+#endif
 	}
 }
 
@@ -2061,15 +2065,14 @@ void ULandscapeHeightfieldCollisionComponent::PrepareGeometryExportSync()
 
 			if(CachedHeightFieldSamples.Heights.Num() != HeightfieldRowsCount * HeightfieldRowsCount)
 			{
-				QUICK_SCOPE_CYCLE_COUNTER(STAT_NavMesh_ExportPxHeightField_saveCells);
+				QUICK_SCOPE_CYCLE_COUNTER(STAT_NavMesh_ExportChaosHeightField_saveCells);
 
 				CachedHeightFieldSamples.Heights.SetNumUninitialized(HeightfieldRowsCount * HeightfieldRowsCount);
 				
 				for(int32 SampleIndex = 0; SampleIndex < CachedHeightFieldSamples.Heights.Num(); ++SampleIndex)
 				{
 					CachedHeightFieldSamples.Heights[SampleIndex] = HeightfieldRef->Heightfield->GetHeight(SampleIndex);
-					// #PHYSTODO Hole support
-					CachedHeightFieldSamples.Holes.Add(false);
+					CachedHeightFieldSamples.Holes.Add(HeightfieldRef->Heightfield->IsHole(SampleIndex));
 				}
 			}
 		}
diff --git a/Engine/Source/Runtime/NavigationSystem/Private/NavMesh/RecastNavMeshGenerator.cpp b/Engine/Source/Runtime/NavigationSystem/Private/NavMesh/RecastNavMeshGenerator.cpp
--- a/Engine/Source/Runtime/NavigationSystem/Private/NavMesh/RecastNavMeshGenerator.cpp
+++ b/Engine/Source/Runtime/NavigationSystem/Private/NavMesh/RecastNavMeshGenerator.cpp
@@ -226,13 +226,14 @@ struct FRecastGeometryExport : public FNavigableGeometryExport
 	virtual void ExportPxTriMesh32Bit(physx::PxTriangleMesh const * const TriMesh, const FTransform& LocalToWorld) override;
 	virtual void ExportPxConvexMesh(physx::PxConvexMesh const * const ConvexMesh, const FTransform& LocalToWorld) override;
 	virtual void ExportPxHeightField(physx::PxHeightField const * const HeightField, const FTransform& LocalToWorld) override;
-#endif // WITH_PHYSX
+	virtual void ExportPxHeightFieldSlice(const FNavHeightfieldSamples& PrefetchedHeightfieldSamples, const int32 NumRows, const int32 NumCols, const FTransform& LocalToWorld, const FBox& SliceBox) override;
+#endif // PHYSICS_INTERFACE_PHYSX
 #if WITH_CHAOS
 	virtual void ExportChaosTriMesh(const Chaos::FTriangleMeshImplicitObject* const TriMesh, const FTransform& LocalToWorld) override;
 	virtual void ExportChaosConvexMesh(const FKConvexElem* const Convex, const FTransform& LocalToWorld) override;
 	virtual void ExportChaosHeightField(const Chaos::THeightField<float>* const Heightfield, const FTransform& LocalToWorld) override;
+	virtual void ExportChaosHeightFieldSlice(const FNavHeightfieldSamples& PrefetchedHeightfieldSamples, const int32 NumRows, const int32 NumCols, const FTransform& LocalToWorld, const FBox& SliceBox) override;
 #endif
-	virtual void ExportHeightFieldSlice(const FNavHeightfieldSamples& PrefetchedHeightfieldSamples, const int32 NumRows, const int32 NumCols, const FTransform& LocalToWorld, const FBox& SliceBox) override;
 	virtual void ExportCustomMesh(const FVector* InVertices, int32 NumVerts, const int32* InIndices, int32 NumIndices, const FTransform& LocalToWorld) override;
 	virtual void ExportRigidBodySetup(UBodySetup& BodySetup, const FTransform& LocalToWorld) override;
 	virtual void AddNavModifiers(const FCompositeNavModifier& Modifiers) override;
@@ -553,6 +554,94 @@ void ExportPxHeightField(PxHeightField const * const HeightField, const FTransfo
 		}
 	}
 }
+
+void ExportPxHeightFieldSlice(const FNavHeightfieldSamples& PrefetchedHeightfieldSamples, const int32 NumRows, const int32 NumCols, const FTransform& LocalToWorld
+	, TNavStatArray<float>& VertexBuffer, TNavStatArray<int32>& IndexBuffer, const FBox& SliceBox
+	, FBox& UnrealBounds)
+{
+	QUICK_SCOPE_CYCLE_COUNTER(STAT_NavMesh_ExportPxHeightFieldSlice);
+
+#if PHYSICS_INTERFACE_PHYSX
+	static const uint32 SizeOfPx = sizeof(physx::PxI16);
+	static const uint32 SizeOfHeight = PrefetchedHeightfieldSamples.Heights.GetTypeSize();
+	ensure(SizeOfPx == SizeOfHeight);
+#endif // PHYSICS_INTERFACE_PHYSX
+
+	// calculate the actual start and number of columns we want
+	const FBox LocalBox = SliceBox.TransformBy(LocalToWorld.Inverse());
+	const bool bMirrored = (LocalToWorld.GetDeterminant() < 0.f);
+
+	const int32 MinX = FMath::Clamp(FMath::FloorToInt(LocalBox.Min.X) - 1, 0, NumCols);
+	const int32 MinY = FMath::Clamp(FMath::FloorToInt(LocalBox.Min.Y) - 1, 0, NumRows);
+	const int32 MaxX = FMath::Clamp(FMath::CeilToInt(LocalBox.Max.X) + 1, 0, NumCols);
+	const int32 MaxY = FMath::Clamp(FMath::CeilToInt(LocalBox.Max.Y) + 1, 0, NumRows);
+	const int32 SizeX = MaxX - MinX;
+	const int32 SizeY = MaxY - MinY;
+
+	if (SizeX <= 0 || SizeY <= 0)
+	{
+		// slice is outside bounds, skip
+		return;
+	}
+
+	const int32 VertOffset = VertexBuffer.Num() / 3;
+	const int32 NumVerts = SizeX * SizeY;
+	const int32 NumQuads = (SizeX - 1) * (SizeY - 1);
+	{
+		QUICK_SCOPE_CYCLE_COUNTER(STAT_RecastGeometryExport_AllocatingMemory);
+		VertexBuffer.Reserve(VertexBuffer.Num() + NumVerts * 3);
+		IndexBuffer.Reserve(IndexBuffer.Num() + NumQuads * 3 * 2);
+	}
+
+	for (int32 IdxY = 0; IdxY < SizeY; IdxY++)
+	{
+		for (int32 IdxX = 0; IdxX < SizeX; IdxX++)
+		{
+			const int32 CoordX = IdxX + MinX;
+			const int32 CoordY = IdxY + MinY;
+			const int32 SampleIdx = ((bMirrored ? CoordX : (NumCols - CoordX - 1)) * NumCols) + CoordY;
+
+			const FVector UnrealCoords = LocalToWorld.TransformPosition(FVector(CoordX, CoordY, PrefetchedHeightfieldSamples.Heights[SampleIdx]));
+			VertexBuffer.Add(UnrealCoords.X);
+			VertexBuffer.Add(UnrealCoords.Y);
+			VertexBuffer.Add(UnrealCoords.Z);
+		}
+	}
+
+	for (int32 IdxY = 0; IdxY < SizeY - 1; IdxY++)
+	{
+		for (int32 IdxX = 0; IdxX < SizeX - 1; IdxX++)
+		{
+			const int32 CoordX = IdxX + MinX;
+			const int32 CoordY = IdxY + MinY;
+			const int32 SampleIdx = ((bMirrored ? CoordX : (NumCols - CoordX - 1)) * NumCols) + CoordY;
+
+			const bool bIsHole = PrefetchedHeightfieldSamples.Holes[SampleIdx];
+			if (bIsHole)
+			{
+				continue;
+			}
+
+			const int32 I00 = (IdxX + 0) + (IdxY + 0) * SizeX;
+			int32 I01 = (IdxX + 0) + (IdxY + 1) * SizeX;
+			int32 I10 = (IdxX + 1) + (IdxY + 0) * SizeX;
+			const int32 I11 = (IdxX + 1) + (IdxY + 1) * SizeX;
+			if (bMirrored)
+			{
+				Swap(I01, I10);
+			}
+
+			IndexBuffer.Add(VertOffset + I00);
+			IndexBuffer.Add(VertOffset + I11);
+			IndexBuffer.Add(VertOffset + I10);
+
+			IndexBuffer.Add(VertOffset + I00);
+			IndexBuffer.Add(VertOffset + I01);
+			IndexBuffer.Add(VertOffset + I11);
+		}
+	}
+}
+
 #endif // WITH_PHYSX
 
 #if WITH_CHAOS
@@ -717,7 +806,7 @@ void ExportChaosHeightField(const Chaos::THeightField<float>* const HeightField,
 		return;
 	}
 
-	QUICK_SCOPE_CYCLE_COUNTER(STAT_NavMesh_ExportPxHeightField);
+	QUICK_SCOPE_CYCLE_COUNTER(STAT_NavMesh_ExportChaosHeightField);
 
 	const int32 NumRows = HeightField->GetNumRows();
 	const int32 NumCols = HeightField->GetNumCols();
@@ -735,7 +824,7 @@ void ExportChaosHeightField(const Chaos::THeightField<float>* const HeightField,
 	{
 		for(int32 X = 0; X < NumCols; X++)
 		{
-			const int32 SampleIdx = Y * NumCols + X;
+			const int32 SampleIdx = Y * NumCols + X;  // #PHYSTODO bMirrored support
 
 			const FVector UnrealCoords = LocalToWorld.TransformPosition(FVector(X, Y, HeightField->GetHeight(SampleIdx)));
 			UnrealBounds += UnrealCoords;
@@ -755,7 +844,7 @@ void ExportChaosHeightField(const Chaos::THeightField<float>* const HeightField,
 				continue;
 			}
 
-			const int32 I0 = Y * NumCols + X;
+			const int32 I0 = Y * NumCols + X;  // #PHYSTODO bMirrored support
 			int32 I1 = I0 + 1;
 			int32 I2 = I0 + NumCols;
 			const int32 I3 = I2 + 1;
@@ -776,13 +865,12 @@ void ExportChaosHeightField(const Chaos::THeightField<float>* const HeightField,
 		}
 	}
 }
-#endif
 
-void ExportHeightFieldSlice(const FNavHeightfieldSamples& PrefetchedHeightfieldSamples, const int32 NumRows, const int32 NumCols, const FTransform& LocalToWorld
+void ExportChaosHeightFieldSlice(const FNavHeightfieldSamples& PrefetchedHeightfieldSamples, const int32 NumRows, const int32 NumCols, const FTransform& LocalToWorld
 	, TNavStatArray<float>& VertexBuffer, TNavStatArray<int32>& IndexBuffer, const FBox& SliceBox
 	, FBox& UnrealBounds)
 {
-	QUICK_SCOPE_CYCLE_COUNTER(STAT_NavMesh_ExportHeightFieldSlice);
+	QUICK_SCOPE_CYCLE_COUNTER(STAT_NavMesh_ExportChaosHeightFieldSlice);
 
 #if PHYSICS_INTERFACE_PHYSX
 	static const uint32 SizeOfPx = sizeof(physx::PxI16);
@@ -822,7 +910,8 @@ void ExportHeightFieldSlice(const FNavHeightfieldSamples& PrefetchedHeightfieldS
 		{
 			const int32 CoordX = IdxX + MinX;
 			const int32 CoordY = IdxY + MinY;
-			const int32 SampleIdx = ((bMirrored ? CoordX : (NumCols - CoordX - 1)) * NumCols) + CoordY;
+			const int32 SampleIdx = CoordY * NumCols + CoordX; // #PHYSTODO bMirrored support
+
 
 			const FVector UnrealCoords = LocalToWorld.TransformPosition(FVector(CoordX, CoordY, PrefetchedHeightfieldSamples.Heights[SampleIdx]));
 			VertexBuffer.Add(UnrealCoords.X);
@@ -837,7 +926,7 @@ void ExportHeightFieldSlice(const FNavHeightfieldSamples& PrefetchedHeightfieldS
 		{
 			const int32 CoordX = IdxX + MinX;
 			const int32 CoordY = IdxY + MinY;
-			const int32 SampleIdx = ((bMirrored ? CoordX : (NumCols - CoordX - 1)) * NumCols) + CoordY;
+			const int32 SampleIdx = CoordY * NumCols + CoordX;  // #PHYSTODO bMirrored support
 
 			const bool bIsHole = PrefetchedHeightfieldSamples.Holes[SampleIdx];
 			if (bIsHole)
@@ -845,25 +934,26 @@ void ExportHeightFieldSlice(const FNavHeightfieldSamples& PrefetchedHeightfieldS
 				continue;
 			}
 
-			const int32 I00 = (IdxX + 0) + (IdxY + 0)*SizeX;
-			int32 I01 = (IdxX + 0) + (IdxY + 1)*SizeX;
-			int32 I10 = (IdxX + 1) + (IdxY + 0)*SizeX;
-			const int32 I11 = (IdxX + 1) + (IdxY + 1)*SizeX;
+			const int32 I0 = IdxY * SizeX + IdxX;
+			int32 I1 = I0 + 1;
+			int32 I2 = I0 + SizeX;
+			const int32 I3 = I2 + 1;
 			if (bMirrored)
 			{
-				Swap(I01, I10);
+				Swap(I1, I2);
 			}
 
-			IndexBuffer.Add(VertOffset + I00);
-			IndexBuffer.Add(VertOffset + I11);
-			IndexBuffer.Add(VertOffset + I10);
+			IndexBuffer.Add(VertOffset + I0);
+			IndexBuffer.Add(VertOffset + I3);
+			IndexBuffer.Add(VertOffset + I1);
 
-			IndexBuffer.Add(VertOffset + I00);
-			IndexBuffer.Add(VertOffset + I01);
-			IndexBuffer.Add(VertOffset + I11);
+			IndexBuffer.Add(VertOffset + I0);
+			IndexBuffer.Add(VertOffset + I2);
+			IndexBuffer.Add(VertOffset + I3);
 		}
 	}
 }
+#endif
 
 void ExportCustomMesh(const FVector* InVertices, int32 NumVerts, const int32* InIndices, int32 NumIndices, const FTransform& LocalToWorld,
 					  TNavStatArray<float>& VertexBuffer, TNavStatArray<int32>& IndexBuffer, FBox& UnrealBounds)
@@ -1345,6 +1435,11 @@ void FRecastGeometryExport::ExportPxHeightField(physx::PxHeightField const * con
 {
 	RecastGeometryExport::ExportPxHeightField(HeightField, LocalToWorld, VertexBuffer, IndexBuffer, Data->Bounds);
 }
+
+void FRecastGeometryExport::ExportPxHeightFieldSlice(const FNavHeightfieldSamples& PrefetchedHeightfieldSamples, const int32 NumRows, const int32 NumCols, const FTransform& LocalToWorld, const FBox& SliceBox)
+{
+	RecastGeometryExport::ExportPxHeightFieldSlice(PrefetchedHeightfieldSamples, NumRows, NumCols, LocalToWorld, VertexBuffer, IndexBuffer, SliceBox, Data->Bounds);
+}
 #endif // WITH_PHYSX
 
 #if WITH_CHAOS
@@ -1362,12 +1457,12 @@ void FRecastGeometryExport::ExportChaosHeightField(const Chaos::THeightField<flo
 {
 	RecastGeometryExport::ExportChaosHeightField(Heightfield, LocalToWorld, VertexBuffer, IndexBuffer, Data->Bounds);
 }
-#endif
 
-void FRecastGeometryExport::ExportHeightFieldSlice(const FNavHeightfieldSamples& PrefetchedHeightfieldSamples, const int32 NumRows, const int32 NumCols, const FTransform& LocalToWorld, const FBox& SliceBox)
+void FRecastGeometryExport::ExportChaosHeightFieldSlice(const FNavHeightfieldSamples& PrefetchedHeightfieldSamples, const int32 NumRows, const int32 NumCols, const FTransform& LocalToWorld, const FBox& SliceBox)
 {
-	RecastGeometryExport::ExportHeightFieldSlice(PrefetchedHeightfieldSamples, NumRows, NumCols, LocalToWorld, VertexBuffer, IndexBuffer, SliceBox, Data->Bounds);
+	RecastGeometryExport::ExportChaosHeightFieldSlice(PrefetchedHeightfieldSamples, NumRows, NumCols, LocalToWorld, VertexBuffer, IndexBuffer, SliceBox, Data->Bounds);
 }
+#endif
 
 void FRecastGeometryExport::ExportCustomMesh(const FVector* InVertices, int32 NumVerts, const int32* InIndices, int32 NumIndices, const FTransform& LocalToWorld)
 {
