[UNavRelevantComponent] Fixed some calls to GetNavigationParent made before the cache initialization since this might result in improper registration in the NavOctree. Also added error reporting to detect those cases in the future.

diff --git a/Engine/Source/Runtime/NavigationSystem/Private/NavRelevantComponent.cpp b/Engine/Source/Runtime/NavigationSystem/Private/NavRelevantComponent.cpp
--- a/Engine/Source/Runtime/NavigationSystem/Private/NavRelevantComponent.cpp
+++ b/Engine/Source/Runtime/NavigationSystem/Private/NavRelevantComponent.cpp
@@ -4,11 +4,12 @@
 #include "NavigationSystem.h"
 
 UNavRelevantComponent::UNavRelevantComponent(const FObjectInitializer& ObjectInitializer) : Super(ObjectInitializer)
+, bAttachToOwnersRoot(true)
+, bBoundsInitialized(false)
+, bNavParentCacheInitialized(false)
 {
 	bCanEverAffectNavigation = true;
 	bNavigationRelevant = true;
-	bAttachToOwnersRoot = true;
-	bBoundsInitialized = false;
 }
 
 void UNavRelevantComponent::OnRegister()
@@ -50,6 +51,9 @@ void UNavRelevantComponent::OnRegister()
 		}
 	}
 
+	// Mark cache as initialized (even if null) from this point so calls to GetNavigationParent can be validated.
+	bNavParentCacheInitialized = true;
+
 	FNavigationSystem::OnComponentRegistered(*this);
 }
 
@@ -84,6 +88,11 @@ void UNavRelevantComponent::UpdateNavigationBounds()
 
 UObject* UNavRelevantComponent::GetNavigationParent() const
 {
+	UE_CLOG(!bNavParentCacheInitialized, LogNavigation, Error, 
+		TEXT("%s called before initialization of the navigation parent cache for [%s]. This might cause improper registration in the NavOctree and must be fixed."), 
+		ANSI_TO_TCHAR(__FUNCTION__),
+		*GetFullName());
+
 	return CachedNavParent;
 }
 
@@ -117,5 +126,9 @@ void UNavRelevantComponent::SetNavigationRelevancy(bool bRelevant)
 
 void UNavRelevantComponent::RefreshNavigationModifiers()
 {
-	FNavigationSystem::UpdateComponentData(*this);
+	// Only update after component registration since some required informations are initialized at that time (i.e. Cached Navigation Parent)
+	if (bRegistered)
+	{
+		FNavigationSystem::UpdateComponentData(*this);
+	}
 }
diff --git a/Engine/Source/Runtime/NavigationSystem/Public/NavRelevantComponent.h b/Engine/Source/Runtime/NavigationSystem/Public/NavRelevantComponent.h
--- a/Engine/Source/Runtime/NavigationSystem/Public/NavRelevantComponent.h
+++ b/Engine/Source/Runtime/NavigationSystem/Public/NavRelevantComponent.h
@@ -46,6 +46,7 @@ protected:
 	uint32 bAttachToOwnersRoot : 1;
 
 	mutable uint32 bBoundsInitialized : 1;
+	uint32 bNavParentCacheInitialized : 1;
 
 	UPROPERTY(transient)
 	UObject* CachedNavParent;
