[Asynchronous Pathfind Queries]
- Added synchronization between the async pathfind queries task and the game thread. The sync point is implemented in OnWorldPostActorTick. This prevents pathfind requests from running at the same time as Navmesh generation tasks or beyond the engine frame boundary.
- Replaced task to broadcast async pathfinding results by a queue that will be flushed by the NavigationSystem udpate. This prevents tasks from being started during loading or other undesired locations.
- Replaced task to register new active path instance by a lock to make sure all paths are registered for potential invalidation before adding a generated tile on the main thread.

diff --git a/Engine/Source/Runtime/NavigationSystem/Private/NavigationSystem.cpp b/Engine/Source/Runtime/NavigationSystem/Private/NavigationSystem.cpp
--- a/Engine/Source/Runtime/NavigationSystem/Private/NavigationSystem.cpp
+++ b/Engine/Source/Runtime/NavigationSystem/Private/NavigationSystem.cpp
@@ -57,6 +57,7 @@ static const uint32 REGISTRATION_QUEUE_SIZE = 16;	// and we'll not reallocate
 
 DECLARE_CYCLE_STAT(TEXT("Nav Tick: mark dirty"), STAT_Navigation_TickMarkDirty, STATGROUP_Navigation);
 DECLARE_CYCLE_STAT(TEXT("Nav Tick: async build"), STAT_Navigation_TickAsyncBuild, STATGROUP_Navigation);
+DECLARE_CYCLE_STAT(TEXT("Nav Tick: dispatch async pathfinding results"), STAT_Navigation_DispatchAsyncPathfindingResults, STATGROUP_Navigation);
 DECLARE_CYCLE_STAT(TEXT("Nav Tick: async pathfinding"), STAT_Navigation_TickAsyncPathfinding, STATGROUP_Navigation);
 DECLARE_CYCLE_STAT(TEXT("Debug NavOctree Time"), STAT_DebugNavOctree, STATGROUP_Navigation);
 
@@ -375,6 +376,7 @@ UNavigationSystemV1::UNavigationSystemV1(const FObjectInitializer& ObjectInitial
 	, bSkipAgentHeightCheckWhenPickingNavData(false)
 	, DirtyAreaWarningSizeThreshold(-1.0f)
 	, OperationMode(FNavigationSystemRunMode::InvalidMode)
+	, bAbortAsyncQueriesRequested(false)
 	, NavBuildingLockFlags(0)
 	, InitialNavBuildingLockFlags(0)
 	, bInitialSetupHasBeenPerformed(false)
@@ -527,6 +529,7 @@ UNavigationSystemV1::UNavigationSystemV1(const FObjectInitializer& ObjectInitial
 		AsyncPathFindingQueries.Reserve( INITIAL_ASYNC_QUERIES_SIZE );
 		NavDataRegistrationQueue.Reserve( REGISTRATION_QUEUE_SIZE );
 	
+		FWorldDelegates::OnWorldPostActorTick.AddUObject(this, &UNavigationSystemV1::OnWorldPostActorTick);
 		FWorldDelegates::LevelAddedToWorld.AddUObject(this, &UNavigationSystemV1::OnLevelAddedToWorld);
 		FWorldDelegates::LevelRemovedFromWorld.AddUObject(this, &UNavigationSystemV1::OnLevelRemovedFromWorld);
 #if !UE_BUILD_SHIPPING
@@ -1220,6 +1223,14 @@ void UNavigationSystemV1::Tick(float DeltaSeconds)
 	CSV_CUSTOM_STAT(NavTasks, NumRemainingTasks, GetNumRemainingBuildTasks(), ECsvCustomStatOp::Set);
 	CSV_CUSTOM_STAT(NavTasks, NumRunningTasks, GetNumRunningBuildTasks(), ECsvCustomStatOp::Set);
 
+	// In multithreaded configuration we can process async pathfinding queries
+	// in dedicated task while dispatching completed queries results on the main thread.
+	// The created task can start and append new result right away so we transfer
+	// completed queries before to keep the list safe.
+	TArray<FAsyncPathFindingQuery> AsyncPathFindingCompletedQueriesToDispatch;
+	Swap(AsyncPathFindingCompletedQueriesToDispatch, AsyncPathFindingCompletedQueries);
+
+	// Trigger the async pathfinding queries (new ones and those that may have been postponed from last frame)
 	if (AsyncPathFindingQueries.Num() > 0)
 	{
 		SCOPE_CYCLE_COUNTER(STAT_Navigation_TickAsyncPathfinding);
@@ -1227,6 +1238,9 @@ void UNavigationSystemV1::Tick(float DeltaSeconds)
 		AsyncPathFindingQueries.Reset();
 	}
 
+	// Dispatch async pathfinding queries results from last frame
+	DispatchAsyncQueriesResults(AsyncPathFindingCompletedQueriesToDispatch);
+
 	if (CrowdManager.IsValid())
 	{
 		CrowdManager->Tick(DeltaSeconds);
@@ -1416,16 +1430,33 @@ void UNavigationSystemV1::TriggerAsyncQueries(TArray<FAsyncPathFindingQuery>& Pa
 		STAT_FSimpleDelegateGraphTask_NavigationSystemBatchedAsyncQueries,
 		STATGROUP_TaskGraphTasks);
 
-	FSimpleDelegateGraphTask::CreateAndDispatchWhenReady(
+	AsyncPathFindingTask = FSimpleDelegateGraphTask::CreateAndDispatchWhenReady(
 		FSimpleDelegateGraphTask::FDelegate::CreateUObject(this, &UNavigationSystemV1::PerformAsyncQueries, PathFindingQueries),
 		GET_STATID(STAT_FSimpleDelegateGraphTask_NavigationSystemBatchedAsyncQueries), nullptr, CPrio_TriggerAsyncQueries.Get());
 }
 
-static void AsyncQueryDone(FAsyncPathFindingQuery Query)
+void UNavigationSystemV1::PostponeAsyncQueries()
 {
-	CSV_SCOPED_TIMING_STAT(NavigationSystem, AsyncNavQueryFinished);
+	if (AsyncPathFindingTask.GetReference() && !AsyncPathFindingTask->IsComplete())
+	{
+		bAbortAsyncQueriesRequested = true;
+		FTaskGraphInterface::Get().WaitUntilTaskCompletes(AsyncPathFindingTask, ENamedThreads::GameThread);
+		bAbortAsyncQueriesRequested = false;
+	}
+}
 
-	Query.OnDoneDelegate.ExecuteIfBound(Query.QueryID, Query.Result.Result, Query.Result.Path);
+void UNavigationSystemV1::DispatchAsyncQueriesResults(const TArray<FAsyncPathFindingQuery>& PathFindingQueries) const
+{
+	if (PathFindingQueries.Num() > 0)
+	{
+		SCOPE_CYCLE_COUNTER(STAT_Navigation_DispatchAsyncPathfindingResults);
+		CSV_SCOPED_TIMING_STAT(NavigationSystem, AsyncNavQueryFinished);
+
+		for (const FAsyncPathFindingQuery& Query : PathFindingQueries)
+		{
+			Query.OnDoneDelegate.ExecuteIfBound(Query.QueryID, Query.Result.Result, Query.Result.Path);
+		}
+	}
 }
 
 void UNavigationSystemV1::PerformAsyncQueries(TArray<FAsyncPathFindingQuery> PathFindingQueries)
@@ -1437,7 +1468,8 @@ void UNavigationSystemV1::PerformAsyncQueries(TArray<FAsyncPathFindingQuery> Pat
 	{
 		return;
 	}
-	
+
+	int32 NumProcessed = 0;
 	for (FAsyncPathFindingQuery& Query : PathFindingQueries)
 	{
 		// @todo this is not necessarily the safest way to use UObjects outside of main thread. 
@@ -1460,17 +1492,28 @@ void UNavigationSystemV1::PerformAsyncQueries(TArray<FAsyncPathFindingQuery> Pat
 		{
 			Query.Result = ENavigationQueryResult::Error;
 		}
+		++NumProcessed;
+
+		// Check for abort request from the main tread
+		if (bAbortAsyncQueriesRequested)
+		{
+			break;
+		}
+	}
 
-		// @todo make it return more informative results (bResult == false)
-		// trigger calling delegate on main thread - otherwise it may depend too much on stuff being thread safe
-		DECLARE_CYCLE_STAT(TEXT("FSimpleDelegateGraphTask.Async nav query finished"),
-			STAT_FSimpleDelegateGraphTask_AsyncNavQueryFinished,
-			STATGROUP_TaskGraphTasks);
+	const int32 NumQueries = PathFindingQueries.Num();
+	const int32 NumPostponed = NumQueries - NumProcessed;
 
-		FSimpleDelegateGraphTask::CreateAndDispatchWhenReady(
-			FSimpleDelegateGraphTask::FDelegate::CreateStatic(AsyncQueryDone, Query),
-			GET_STATID(STAT_FSimpleDelegateGraphTask_AsyncNavQueryFinished), NULL, ENamedThreads::GameThread);
+	// Queue remaining queries for next frame
+	if (bAbortAsyncQueriesRequested)
+	{
+		AsyncPathFindingQueries.Append(PathFindingQueries.GetData() + NumProcessed, NumPostponed);
 	}
+	
+	// Append to list of completed queries to dispatch results in main thread
+	AsyncPathFindingCompletedQueries.Append(PathFindingQueries.GetData(), NumProcessed);
+
+	UE_LOG(LogNavigation, Log, TEXT("Async pathfinding queries: %d completed, %d postponed to next frame"), NumProcessed, NumPostponed);
 }
 
 bool UNavigationSystemV1::GetRandomPoint(FNavLocation& ResultLocation, ANavigationData* NavData, FSharedConstNavQueryFilter QueryFilter)
diff --git a/Engine/Source/Runtime/NavigationSystem/Public/NavigationData.h b/Engine/Source/Runtime/NavigationSystem/Public/NavigationData.h
--- a/Engine/Source/Runtime/NavigationSystem/Public/NavigationData.h
+++ b/Engine/Source/Runtime/NavigationSystem/Public/NavigationData.h
@@ -627,7 +627,7 @@ public:
 	 */
 	virtual void TickAsyncBuild(float DeltaSeconds);
 	
-	/** Retrieves navmesh's generator */
+	/** Retrieves navigation data generator */
 	FNavDataGenerator* GetGenerator() { return NavDataGenerator.Get(); }
 	const FNavDataGenerator* GetGenerator() const { return NavDataGenerator.Get(); }
 
@@ -665,15 +665,7 @@ public:
 		SharedPath->SetQueryData(QueryData);
 		SharedPath->SetTimeStamp( GetWorldTimeStamp() );
 
-		DECLARE_CYCLE_STAT(TEXT("FSimpleDelegateGraphTask.Adding a path to ActivePaths"),
-			STAT_FSimpleDelegateGraphTask_AddingPathToActivePaths,
-			STATGROUP_TaskGraphTasks);
-
-		FSimpleDelegateGraphTask::CreateAndDispatchWhenReady(
-			FSimpleDelegateGraphTask::FDelegate::CreateUObject(const_cast<ANavigationData*>(this), &ANavigationData::RegisterActivePath, SharedPath),
-			GET_STATID(STAT_FSimpleDelegateGraphTask_AddingPathToActivePaths), NULL, ENamedThreads::GameThread
-		);
-
+		const_cast<ANavigationData*>(this)->RegisterActivePath(SharedPath);
 		return SharedPath;
 	}
 	
@@ -687,7 +679,11 @@ public:
 		ObservedPaths.Add(SharedPath);
 	}
 
-	void RequestRePath(FNavPathSharedPtr Path, ENavPathUpdateType::Type Reason) { RepathRequests.AddUnique(FNavPathRecalculationRequest(Path, Reason)); }
+    void RequestRePath(FNavPathSharedPtr Path, ENavPathUpdateType::Type Reason)
+    {
+	    check(IsInGameThread());
+	    RepathRequests.AddUnique(FNavPathRecalculationRequest(Path, Reason)); 
+    }
 
 protected:
 	/** removes from ActivePaths all paths that no longer have shared references (and are invalid in fact) */
@@ -695,7 +691,8 @@ protected:
 
 	void RegisterActivePath(FNavPathSharedPtr SharedPath)
 	{
-		check(IsInGameThread());
+		// Paths can be registered from main thread and async pathfinding thread
+		FScopeLock PathLock(&ActivePathsLock);
 		ActivePaths.Add(SharedPath);
 	}
 
@@ -952,6 +949,9 @@ protected:
 	 */
 	TArray<FNavPathWeakPtr> ActivePaths;
 
+	/** Synchronization object for paths registration from main thread and async pathfinding thread */
+	FCriticalSection ActivePathsLock;
+
 	/**
 	 *	Contains paths that requested observing its goal's location. These paths will be 
 	 *	processed on a regular basis (@see ObservedPathsTickInterval) */
diff --git a/Engine/Source/Runtime/NavigationSystem/Public/NavigationSystem.h b/Engine/Source/Runtime/NavigationSystem/Public/NavigationSystem.h
--- a/Engine/Source/Runtime/NavigationSystem/Public/NavigationSystem.h
+++ b/Engine/Source/Runtime/NavigationSystem/Public/NavigationSystem.h
@@ -857,7 +857,7 @@ protected:
 	/** Sets up SuportedAgents and NavigationDataCreators. Override it to add additional setup, but make sure to call Super implementation */
 	virtual void DoInitialSetup();
 
-	/** spawn new crowd manager */
+	/** Find or create abstract nav data */
 	virtual void UpdateAbstractNavData();
 	
 	/** Called during ConditionalPopulateNavOctree and gives subclassess a chance
@@ -971,9 +971,19 @@ protected:
 
 	UPROPERTY()
 	FNavigationSystemRunMode OperationMode;
-	
+
+	/** Queued async pathfinding queries to process in the next update. */
 	TArray<FAsyncPathFindingQuery> AsyncPathFindingQueries;
 
+	/** Queued async pathfinding results computed by the dedicated task in the last frame and ready to dispatch in the next update. */
+	TArray<FAsyncPathFindingQuery> AsyncPathFindingCompletedQueries;
+
+	/** Graph event that the main thread will wait for to synchronize with the async pathfinding task, if any. */
+	FGraphEventRef AsyncPathFindingTask;
+
+	/** Flag used by main thread to ask the async pathfinding task to stop and postpone remaining queries, if any. */
+	TAtomic<bool> bAbortAsyncQueriesRequested;
+
 	FCriticalSection NavDataRegistration;
 
 	TMap<FNavAgentProperties, TWeakObjectPtr<ANavigationData> > AgentToNavDataMap;
@@ -1094,6 +1104,9 @@ protected:
 	/** Handler for FWorldDelegates::LevelRemovedFromWorld event */
 	void OnLevelRemovedFromWorld(ULevel* InLevel, UWorld* InWorld);
 
+	/** Handler for FWorldDelegates::OnWorldPostActorTick event */
+	void OnWorldPostActorTick(UWorld* World, ELevelTick TickType, float DeltaTime) { PostponeAsyncQueries(); }
+
 	/** Adds given request to requests queue. Note it's to be called only on game thread only */
 	void AddAsyncQuery(const FAsyncPathFindingQuery& Query);
 		 
@@ -1104,6 +1117,15 @@ protected:
 	/** Processes pathfinding requests given in PathFindingQueries.*/
 	void PerformAsyncQueries(TArray<FAsyncPathFindingQuery> PathFindingQueries);
 
+	/** Broadcasts completion delegate for all completed async pathfinding requests. */
+	void DispatchAsyncQueriesResults(const TArray<FAsyncPathFindingQuery>& PathFindingQueries) const;
+
+	/**
+	 * Requests the async pathfinding task to abort and waits for it to complete
+	 * before resuming the main thread. Pathfind task will postpone remaining queries to next frame.
+	 */
+	void PostponeAsyncQueries();
+
 	/** */
 	virtual void DestroyNavOctree();
 
