RecastNavMeshGenerator: optim in MarkDirtyTiles by reducing search size in dirty tiles set and num of copies between containers. For cases with large amount of pending elements (>100000) average timing goes from >5 ms to ~0.6ms.
With this change we now preserve the order of the pending elements (if tile sorting isn't used).

diff --git a/Engine/Source/Runtime/NavigationSystem/Private/NavMesh/RecastNavMeshGenerator.cpp b/Engine/Source/Runtime/NavigationSystem/Private/NavMesh/RecastNavMeshGenerator.cpp
--- a/Engine/Source/Runtime/NavigationSystem/Private/NavMesh/RecastNavMeshGenerator.cpp
+++ b/Engine/Source/Runtime/NavigationSystem/Private/NavMesh/RecastNavMeshGenerator.cpp
@@ -5525,7 +5525,7 @@ void FRecastNavMeshGenerator::MarkDirtyTiles(const TArray<FNavigationDirtyArea>&
 		bool bDoTileInclusionTest = false;
 		FBox AdjustedAreaBounds = DirtyArea.Bounds;
 		
-		// if it's not expanding the navigatble area
+		// if it's not expanding the navigable area
 		if (DirtyArea.HasFlag(ENavigationDirtyFlag::NavigationBounds) == false)
 		{
 			// and is outside of current bounds
@@ -5605,31 +5605,30 @@ void FRecastNavMeshGenerator::MarkDirtyTiles(const TArray<FNavigationDirtyArea>&
 	
 	int32 NumTilesMarked = DirtyTiles.Num();
 
-	// Merge all pending tiles into one container
-	for (const FPendingTileElement& Element : PendingDirtyTiles)
+	// Merge new dirty tiles info with existing pending elements
+	for (FPendingTileElement& ExistingElement : PendingDirtyTiles)
 	{
-		FPendingTileElement* ExistingElement = DirtyTiles.Find(Element);
-		if (ExistingElement)
+		FSetElementId Id = DirtyTiles.FindId(ExistingElement);
+		if (Id.IsValidId())
 		{
-			ExistingElement->bRebuildGeometry|= Element.bRebuildGeometry;
+			const FPendingTileElement& DirtyElement = DirtyTiles[Id];
+
+			ExistingElement.bRebuildGeometry |= DirtyElement.bRebuildGeometry;
 			// Append area bounds to existing list 
-			if (ExistingElement->bRebuildGeometry == false)
+			if (ExistingElement.bRebuildGeometry == false)
 			{
-				ExistingElement->DirtyAreas.Append(Element.DirtyAreas);
+				ExistingElement.DirtyAreas.Append(DirtyElement.DirtyAreas);
 			}
 			else
 			{
-				ExistingElement->DirtyAreas.Empty();
+				ExistingElement.DirtyAreas.Empty();
 			}
-		}
-		else
-		{
-			DirtyTiles.Add(Element);
+			DirtyTiles.Remove(Id);
 		}
 	}
-	
-	// Dump results into array
-	PendingDirtyTiles.Empty(DirtyTiles.Num());
+
+	// Append remaining new dirty tile elements
+	PendingDirtyTiles.Reserve(PendingDirtyTiles.Num() + DirtyTiles.Num());
 	for(const FPendingTileElement& Element : DirtyTiles)
 	{
 		PendingDirtyTiles.Add(Element);
