Fixed holes in NavMesh caused by single voxel areas near area boundary.
- implemented median filter which is applied after areas are marked for tile cache
- implemented debug draw functionality for tile cache build steps
- fixed DetourDebugDraw.h API export

diff --git a/Engine/Source/Runtime/NavigationSystem/Private/NavMesh/NavMeshRenderingComponent.cpp b/Engine/Source/Runtime/NavigationSystem/Private/NavMesh/NavMeshRenderingComponent.cpp
--- a/Engine/Source/Runtime/NavigationSystem/Private/NavMesh/NavMeshRenderingComponent.cpp
+++ b/Engine/Source/Runtime/NavigationSystem/Private/NavMesh/NavMeshRenderingComponent.cpp
@@ -65,6 +65,25 @@ namespace FNavMeshRenderingHelpers
 		return true;
 	}
 
+	bool PointInView(const FVector& Position, const FSceneView* View)
+	{
+		if (FVector::DistSquaredXY(Position, View->ViewMatrices.GetViewOrigin()) > ARecastNavMesh::GetDrawDistanceSq())
+		{
+			return false;
+		}
+
+		for (int32 PlaneIdx = 0; PlaneIdx < View->ViewFrustum.Planes.Num(); ++PlaneIdx)
+		{
+			const FPlane& CurPlane = View->ViewFrustum.Planes[PlaneIdx];
+			if (CurPlane.PlaneDot(Position) > 0.f)
+			{
+				return false;
+			}
+		}
+
+		return true;
+	}
+
 	bool LineInCorrectDistance(const FVector& Start, const FVector& End, const FSceneView* View, float CorrectDistance = -1)
 	{
 		const float MaxDistanceSq = (CorrectDistance > 0) ? FMath::Square(CorrectDistance) : ARecastNavMesh::GetDrawDistanceSq();
@@ -248,6 +267,8 @@ void FNavMeshSceneProxyData::Reset()
 	NavMeshEdgeLines.Reset();
 	NavLinkLines.Reset();
 	ClusterLinkLines.Reset();
+	AuxLines.Reset();
+	AuxPoints.Reset();
 	AuxBoxes.Reset();
 	DebugLabels.Reset();
 	OctreeBounds.Reset();
@@ -293,7 +314,7 @@ void FNavMeshSceneProxyData::Serialize(FArchive& Ar)
 		Ar << MeshBuilder.ClusterColor;
 	}
 
-	TArray<FDebugRenderSceneProxy::FDebugLine>* LineArraysToSerialize[] = { &ThickLineItems, &TileEdgeLines, &NavMeshEdgeLines, &NavLinkLines, &ClusterLinkLines };
+	TArray<FDebugRenderSceneProxy::FDebugLine>* LineArraysToSerialize[] = { &ThickLineItems, &TileEdgeLines, &NavMeshEdgeLines, &NavLinkLines, &ClusterLinkLines, &AuxLines };
 	for (int32 ArrIdx = 0; ArrIdx < UE_ARRAY_COUNT(LineArraysToSerialize); ArrIdx++)
 	{
 		int32 NumItems = LineArraysToSerialize[ArrIdx]->Num();
@@ -313,6 +334,31 @@ void FNavMeshSceneProxyData::Serialize(FArchive& Ar)
 		}
 	}
 
+	int32 NumPoints = AuxPoints.Num();
+	Ar << NumPoints;
+	if (Ar.IsLoading())
+	{
+		FDebugPoint TmpPoint(FVector::ZeroVector, FColor::Black, 0.0f);
+		AuxPoints.Reserve(NumPoints);
+		for (int32 Idx = 0; Idx < NumPoints; Idx++)
+		{
+			Ar << TmpPoint.Position;
+			Ar << TmpPoint.Color;
+			Ar << TmpPoint.Size;
+
+			AuxPoints.Add(TmpPoint);
+		}
+	}
+	else
+	{
+		for (int32 Idx = 0; Idx < NumPoints; Idx++)
+		{
+			Ar << AuxPoints[Idx].Position;
+			Ar << AuxPoints[Idx].Color;
+			Ar << AuxPoints[Idx].Size;
+		}
+	}
+
 	int32 NumBoxes = AuxBoxes.Num();
 	Ar << NumBoxes;
 	if (Ar.IsLoading())
@@ -382,6 +428,8 @@ uint32 FNavMeshSceneProxyData::GetAllocatedSize() const
 		NavMeshEdgeLines.GetAllocatedSize() +
 		NavLinkLines.GetAllocatedSize() +
 		ClusterLinkLines.GetAllocatedSize() +
+		AuxLines.GetAllocatedSize() +
+		AuxPoints.GetAllocatedSize() +
 		AuxBoxes.GetAllocatedSize() +
 		DebugLabels.GetAllocatedSize() +
 		OctreeBounds.GetAllocatedSize();
@@ -860,34 +908,108 @@ void FNavMeshSceneProxyData::GatherData(const ARecastNavMesh* NavMesh, int32 InN
 	}
 }
 
+
 #if RECAST_INTERNAL_DEBUG_DATA
+
+namespace FNavMeshRenderingHelpers
+{
+	struct FUniqueColor
+	{
+		FUniqueColor() : Color(0), Count(0) {}
+		FUniqueColor(const uint32 InColor, const uint32 InCount) : Color(InColor), Count(InCount) {}
+		uint32 Color;
+		uint32 Count;
+		bool operator==(const FUniqueColor& Other) const
+		{
+			return Color == Other.Color;
+
+		}
+	};
+	FORCEINLINE uint32 GetTypeHash(const FUniqueColor& UniqueColor)
+	{
+		return UniqueColor.Color;
+	}
+}
+
 void FNavMeshSceneProxyData::AddMeshForInternalData(const FRecastInternalDebugData& InInternalData)
 {
-	// Only supports triangles for now
-	for (const FRecastInternalDebugData::Command& cmd : InInternalData.Commands)
+	if (InInternalData.TriangleIndices.Num() > 0)
 	{
-		if (cmd.Prim != DU_DRAW_TRIS)
+		const TArray<uint32>& Indices = InInternalData.TriangleIndices;
+		const TArray<FVector>& Vertices = InInternalData.TriangleVertices;
+		const TArray<uint32>& Colors = InInternalData.TriangleColors;
+
+		if (ensure(Vertices.Num() == Colors.Num()))
 		{
-			UE_LOG(LogNavigation, Warning, TEXT("Debug data display command not supported yet."));
-			return;
+			// Split the mesh into different colored pieces (because we cannot use vertex colors).
+			TSet<FNavMeshRenderingHelpers::FUniqueColor> UniqueColors;
+			uint32 PrevColor = Colors[Indices[0]];
+			FSetElementId ColorIdx = UniqueColors.Add(FNavMeshRenderingHelpers::FUniqueColor(PrevColor, 1));
+			for (int32 i = 3; i < Indices.Num(); i += 3)
+			{
+				const uint32 Color = Colors[Indices[i]];	// Use the first color as representative of the triangle color.
+				if (Color != PrevColor)
+				{
+					ColorIdx = UniqueColors.Add(FNavMeshRenderingHelpers::FUniqueColor(PrevColor, 1));
+					PrevColor = Color;
+				}
+				else
+				{
+					UniqueColors[ColorIdx].Count++;
+				}
+			}
+
+			// Add triangles
+			for (FNavMeshRenderingHelpers::FUniqueColor CurrentColor : UniqueColors)
+			{
+				const uint32 VertexCount = CurrentColor.Count * 3;
+				FDebugMeshData& MeshData = MeshBuilders.AddDefaulted_GetRef();
+
+				MeshData.Indices.Reserve(VertexCount);
+				MeshData.Vertices.Reserve(VertexCount);
+				for (int32 i = 0; i < Indices.Num(); i += 3)
+				{
+					const uint32 Color = Colors[Indices[i]];
+					if (Color == CurrentColor.Color)
+					{
+						const int32 VertexBase = MeshData.Vertices.Num();
+						FNavMeshRenderingHelpers::AddVertex(MeshData, Vertices[Indices[i]] + NavMeshDrawOffset, FColor(Color));
+						FNavMeshRenderingHelpers::AddVertex(MeshData, Vertices[Indices[i + 1]] + NavMeshDrawOffset, FColor(Color));
+						FNavMeshRenderingHelpers::AddVertex(MeshData, Vertices[Indices[i + 2]] + NavMeshDrawOffset, FColor(Color));
+						MeshData.Indices.Add(VertexBase);
+						MeshData.Indices.Add(VertexBase + 1);
+						MeshData.Indices.Add(VertexBase + 2);
+					}
+				}
+				MeshData.ClusterColor = FColor(CurrentColor.Color);
+			}
+		}
+	}
+
+	if (InInternalData.LineVertices.Num() > 0)
+	{
+		const TArray<FVector>& Vertices = InInternalData.LineVertices;
+		const TArray<uint32>& Colors = InInternalData.LineColors;
+		if (ensure(Vertices.Num() == Colors.Num()))
+		{
+			for (int32 i = 0; i < Vertices.Num(); i += 2)
+			{
+				AuxLines.Emplace(Vertices[i] + NavMeshDrawOffset, Vertices[i + 1] + NavMeshDrawOffset, FColor(Colors[i]), 0.0f);
+			}
 		}
 	}
 
-	// Add triangles
-	const TArray<FVector>& Verts = InInternalData.Vertices;
-	const TArray<uint32>&  VertsColors = InInternalData.VertexColors;
-	const TArray<uint32>&  Indices = InInternalData.Indices;
-	if (ensure(Verts.Num() == VertsColors.Num()) && ensure(Indices.Num() % 3 == 0))
+	if (InInternalData.PointVertices.Num() > 0)
 	{
-		MeshBuilders.AddDefaulted();
-		FDebugMeshData& MeshData = MeshBuilders.Last();
-		MeshData.Vertices.Reserve(Verts.Num());
-		for (int32 VertIdx = 0; VertIdx < Verts.Num(); ++VertIdx)
+		const TArray<FVector>& Vertices = InInternalData.PointVertices;
+		const TArray<uint32>& Colors = InInternalData.PointColors;
+		if (ensure(Vertices.Num() == Colors.Num()))
 		{
-			FNavMeshRenderingHelpers::AddVertex(MeshData, Verts[VertIdx], FColor(VertsColors[VertIdx]));
+			for (int32 i = 0; i < Vertices.Num(); i++)
+			{
+				AuxPoints.Emplace(Vertices[i] + NavMeshDrawOffset, FColor(Colors[i]), 5.0f);
+			}
 		}
-		MeshData.Indices = Indices;
-		MeshData.ClusterColor = FColor(60, 200, 200, 100); //find the proper material to use the vert colors instead
 	}
 }
 #endif //RECAST_INTERNAL_DEBUG_DATA
@@ -1158,6 +1280,28 @@ void FNavMeshSceneProxy::GetDynamicMeshElements(const TArray<const FSceneView*>&
 				}
 			}
 
+			Num = ProxyData.AuxLines.Num();
+			PDI->AddReserveLines(SDPG_World, Num, false, false);
+			PDI->AddReserveLines(SDPG_Foreground, Num, false, true);
+			for (int32 Index = 0; Index < Num; ++Index)
+			{
+				const auto& Line = ProxyData.AuxLines[Index];
+				if (FNavMeshRenderingHelpers::LineInView(Line.Start, Line.End, View))
+				{
+					PDI->DrawLine(Line.Start, Line.End, Line.Color, SDPG_World, Line.Thickness, 0, true);
+				}
+			}
+
+			Num = ProxyData.AuxPoints.Num();
+			for (int32 Index = 0; Index < Num; ++Index)
+			{
+				const auto& Point = ProxyData.AuxPoints[Index];
+				if (FNavMeshRenderingHelpers::PointInView(Point.Position, View))
+				{
+					PDI->DrawPoint(Point.Position, Point.Color, Point.Size, SDPG_World);
+				}
+			}
+
 			Num = ProxyData.ThickLineItems.Num();
 			PDI->AddReserveLines(SDPG_Foreground, Num, false, true);
 			for (int32 Index = 0; Index < Num; ++Index)
diff --git a/Engine/Source/Runtime/NavigationSystem/Private/NavMesh/RecastNavMeshGenerator.cpp b/Engine/Source/Runtime/NavigationSystem/Private/NavMesh/RecastNavMeshGenerator.cpp
--- a/Engine/Source/Runtime/NavigationSystem/Private/NavMesh/RecastNavMeshGenerator.cpp
+++ b/Engine/Source/Runtime/NavigationSystem/Private/NavMesh/RecastNavMeshGenerator.cpp
@@ -35,6 +35,7 @@
 #if RECAST_INTERNAL_DEBUG_DATA
 #include "DebugUtils/DebugDraw.h"
 #include "DebugUtils/RecastDebugDraw.h"
+#include "DebugUtils/DetourDebugDraw.h"
 #endif //RECAST_INTERNAL_DEBUG_DATA
 
 #ifndef OUTPUT_NAV_TILE_LAYER_COMPRESSION_DATA
@@ -3636,6 +3637,13 @@ bool FRecastTileGenerator::GenerateNavigationDataLayer(FNavMeshBuildContext& Bui
 		}
 	}
 
+#if RECAST_INTERNAL_DEBUG_DATA
+	if (GNavmeshDisplayStep == 100 && IsTileToDebug())
+	{
+		duDebugDrawTileCacheLayerRegions(&BuildContext.InternalDebugData, *GenerationContext.Layer, TileConfig.cs, TileConfig.ch);
+	}
+#endif
+
 	{
 		SCOPE_CYCLE_COUNTER(STAT_Navigation_RecastBuildContours);
 		// Build contour set
@@ -3669,6 +3677,14 @@ bool FRecastTileGenerator::GenerateNavigationDataLayer(FNavMeshBuildContext& Bui
 		}
 	}
 
+#if RECAST_INTERNAL_DEBUG_DATA
+	if (GNavmeshDisplayStep == 110 && IsTileToDebug())
+	{
+		duDebugDrawTileCacheContours(&BuildContext.InternalDebugData, *GenerationContext.ContourSet, GenerationContext.Layer->header->bmin, TileConfig.cs, TileConfig.ch);
+	}
+#endif
+
+
 	{
 		SCOPE_CYCLE_COUNTER(STAT_Navigation_RecastBuildPolyMesh);
 		// Build poly mesh
@@ -3694,6 +3710,13 @@ bool FRecastTileGenerator::GenerateNavigationDataLayer(FNavMeshBuildContext& Bui
 		}
 	}
 
+#if RECAST_INTERNAL_DEBUG_DATA
+	if (GNavmeshDisplayStep == 120 && IsTileToDebug())
+	{
+		duDebugDrawTileCachePolyMesh(&BuildContext.InternalDebugData, *GenerationContext.PolyMesh, GenerationContext.Layer->header->bmin, TileConfig.cs, TileConfig.ch);
+	}
+#endif
+
 	// Build detail mesh
 	if (TileConfig.bGenerateDetailedMesh)
 	{
@@ -3714,6 +3737,13 @@ bool FRecastTileGenerator::GenerateNavigationDataLayer(FNavMeshBuildContext& Bui
 			BuildContext.log(RC_LOG_ERROR, "GenerateNavigationData: Failed to generate poly detail mesh.");
 			return false;
 		}
+
+#if RECAST_INTERNAL_DEBUG_DATA
+		if (GNavmeshDisplayStep == 130 && IsTileToDebug())
+		{
+			duDebugDrawTileCacheDetailMesh(&BuildContext.InternalDebugData, *GenerationContext.DetailMesh);
+		}
+#endif
 	}
 
 	unsigned char* NavData = nullptr;
diff --git a/Engine/Source/Runtime/NavigationSystem/Public/NavMesh/NavMeshRenderingComponent.h b/Engine/Source/Runtime/NavigationSystem/Public/NavMesh/NavMeshRenderingComponent.h
--- a/Engine/Source/Runtime/NavigationSystem/Public/NavMesh/NavMeshRenderingComponent.h
+++ b/Engine/Source/Runtime/NavigationSystem/Public/NavMesh/NavMeshRenderingComponent.h
@@ -56,11 +56,22 @@ struct NAVIGATIONSYSTEM_API FNavMeshSceneProxyData : public TSharedFromThis<FNav
 	};
 	TArray<FDebugMeshData> MeshBuilders;
 
+	struct FDebugPoint
+	{
+		FDebugPoint() {}
+		FDebugPoint(const FVector& InPosition, const FColor& InColor, const float InSize) : Position(InPosition), Color(InColor), Size(InSize) {}
+		FVector Position;
+		FColor Color;
+		float Size;
+	};
+
 	TArray<FDebugRenderSceneProxy::FDebugLine> ThickLineItems;
 	TArray<FDebugRenderSceneProxy::FDebugLine> TileEdgeLines;
 	TArray<FDebugRenderSceneProxy::FDebugLine> NavMeshEdgeLines;
 	TArray<FDebugRenderSceneProxy::FDebugLine> NavLinkLines;
 	TArray<FDebugRenderSceneProxy::FDebugLine> ClusterLinkLines;
+	TArray<FDebugRenderSceneProxy::FDebugLine> AuxLines;
+	TArray<FDebugPoint> AuxPoints;
 	TArray<FDebugRenderSceneProxy::FDebugBox> AuxBoxes;
 	TArray<FDebugRenderSceneProxy::FMesh> Meshes;
 
diff --git a/Engine/Source/Runtime/NavigationSystem/Public/NavMesh/RecastInternalDebugData.cpp b/Engine/Source/Runtime/NavigationSystem/Public/NavMesh/RecastInternalDebugData.cpp
--- a/Engine/Source/Runtime/NavigationSystem/Public/NavMesh/RecastInternalDebugData.cpp
+++ b/Engine/Source/Runtime/NavigationSystem/Public/NavMesh/RecastInternalDebugData.cpp
@@ -5,42 +5,50 @@
 
 void FRecastInternalDebugData::vertex(const float x, const float y, const float z, unsigned int color, const float u, const float v)
 {
-	float pos[3] = { x,y,z };
-	Vertices.Push(Recast2UnrealPoint(pos));
-	VertexColors.Push(color);
+	float RecastPos[3] = { x,y,z };
+	const FVector Pos = Recast2UnrealPoint(RecastPos);
+	switch(CurrentPrim)
+	{
+	case DU_DRAW_POINTS:
+		PointVertices.Push(Pos);
+		PointColors.Push(color);
+		break;
+	case DU_DRAW_LINES:
+		LineVertices.Push(Pos);
+		LineColors.Push(color);
+		break;
+	case DU_DRAW_TRIS:
+		// Fallthrough
+	case DU_DRAW_QUADS:
+		TriangleVertices.Push(Pos);
+		TriangleColors.Push(color);
+		break;
+	}
 }
 
 void FRecastInternalDebugData::end()
 {
-	if (ensure(Commands.Num() > 0))
+	if (CurrentPrim == DU_DRAW_QUADS)
 	{
-		Command& Cmd = Commands.Last();
-		const int FirstVertex = Cmd.Offset;
-
-		Cmd.Offset = Indices.Num();
-
 		// Turns quads to triangles
-		if (Cmd.Prim == DU_DRAW_QUADS)
+		for (int32 i = FirstVertexIndex; i < TriangleVertices.Num(); i += 4)
 		{
-			Cmd.Prim = DU_DRAW_TRIS;
-			for (int i = FirstVertex; i < Vertices.Num(); i += 4)
-			{
-				ensure(i + 3 < Vertices.Num());
-				Indices.Push(i + 0);
-				Indices.Push(i + 1);
-				Indices.Push(i + 3);
+			ensure(i + 3 < TriangleVertices.Num());
+			TriangleIndices.Push(i + 0);
+			TriangleIndices.Push(i + 1);
+			TriangleIndices.Push(i + 3);
 
-				Indices.Push(i + 3);
-				Indices.Push(i + 1);
-				Indices.Push(i + 2);
-			}
+			TriangleIndices.Push(i + 3);
+			TriangleIndices.Push(i + 1);
+			TriangleIndices.Push(i + 2);
 		}
-		else
+	}
+	else if (CurrentPrim == DU_DRAW_TRIS)
+	{
+		// Add indices for triangles.
+		for (int32 i = FirstVertexIndex; i < TriangleVertices.Num(); i++)
 		{
-			for (int i = FirstVertex; i < Vertices.Num(); i++)
-			{
-				Indices.Push(i);
-			}
+			TriangleIndices.Push(i);
 		}
 	}
 }
diff --git a/Engine/Source/Runtime/NavigationSystem/Public/NavMesh/RecastInternalDebugData.h b/Engine/Source/Runtime/NavigationSystem/Public/NavMesh/RecastInternalDebugData.h
--- a/Engine/Source/Runtime/NavigationSystem/Public/NavMesh/RecastInternalDebugData.h
+++ b/Engine/Source/Runtime/NavigationSystem/Public/NavMesh/RecastInternalDebugData.h
@@ -7,17 +7,18 @@
 
 struct FRecastInternalDebugData : public duDebugDraw
 {
-	struct Command
-	{
-		duDebugDrawPrimitives Prim;
-		int Offset; // Index offset
-		int Count;  // Index count
-	};
+	duDebugDrawPrimitives CurrentPrim;
+	int32 FirstVertexIndex;
+
+	TArray<uint32> TriangleIndices;
+	TArray<FVector> TriangleVertices;
+	TArray<uint32> TriangleColors;
+
+	TArray<FVector> LineVertices;
+	TArray<uint32>  LineColors;
 
-	TArray<Command> Commands;
-	TArray<uint32> Indices;
-	TArray<FVector> Vertices;
-	TArray<uint32>  VertexColors;
+	TArray<FVector> PointVertices;
+	TArray<uint32>  PointColors;
 
 	FRecastInternalDebugData() {}
 	virtual ~FRecastInternalDebugData() override {}
@@ -27,11 +28,8 @@ struct FRecastInternalDebugData : public duDebugDraw
 
 	virtual void begin(duDebugDrawPrimitives prim, float size = 1.0f) override
 	{
-		Command Cmd;
-		Cmd.Prim = prim;
-		Cmd.Offset = Vertices.Num();    // Misuse to store initial state.
-		Cmd.Count = 0;
-		Commands.Push(Cmd);
+		CurrentPrim = prim;
+		FirstVertexIndex = TriangleVertices.Num();
 	}
 
 	virtual void vertex(const float* pos, unsigned int color) override
diff --git a/Engine/Source/Runtime/Navmesh/Private/DetourTileCache/DetourTileCacheBuilder.cpp b/Engine/Source/Runtime/Navmesh/Private/DetourTileCache/DetourTileCacheBuilder.cpp
--- a/Engine/Source/Runtime/Navmesh/Private/DetourTileCache/DetourTileCacheBuilder.cpp
+++ b/Engine/Source/Runtime/Navmesh/Private/DetourTileCache/DetourTileCacheBuilder.cpp
@@ -27,7 +27,6 @@ static const int MAX_VERTS_PER_POLY = 6;	// TODO: use the DT_VERTS_PER_POLYGON
 static const int MAX_REM_EDGES = 48;		// TODO: make this an expression.
 
 
-
 dtTileCacheContourSet* dtAllocTileCacheContourSet(dtTileCacheAlloc* alloc)
 {
 	dtAssert(alloc);
@@ -198,7 +197,7 @@ static bool appendVertex(dtTempContour& cont, const int x, const int y, const in
 
 static void getNeighbourRegAndArea(dtTileCacheLayer& layer,
 	const int ax, const int ay, const int dir,
-	unsigned short& neiReg, unsigned char& neiArea)
+	unsigned short& neiReg, unsigned char& neiArea, unsigned char& cornerNeiArea)
 {
 	const int w = (int)layer.header->width;
 	const int ia = ax + ay*w;
@@ -207,6 +206,8 @@ static void getNeighbourRegAndArea(dtTileCacheLayer& layer,
 	const unsigned char portal = layer.cons[ia] >> 4;
 	const unsigned char mask = (unsigned char)(1<<dir);
 
+	cornerNeiArea = 0;
+
 	if ((con & mask) == 0)
 	{
 		// No connection, return portal or hard edge.
@@ -229,6 +230,27 @@ static void getNeighbourRegAndArea(dtTileCacheLayer& layer,
 
 		neiReg = layer.regs[ib];
 		neiArea = layer.areas[ib];
+
+		// Get area type of the cell diagonal [c] to current cell [a]. Where [b] is direct neighbour in the direction of 'dir'.
+		//   ^
+		//  [b][c]
+		//  [a]
+		const int cdir = (dir + 1) & 0x3;
+		const unsigned char bcon = layer.cons[ib] & 0xf;
+		const unsigned char bportal = layer.cons[ib] >> 4;
+		const unsigned char bmask = (unsigned char)(1 << cdir);
+		if ((bcon & bmask) == 0)
+		{
+			cornerNeiArea = 0;
+		}
+		else
+		{
+			const int cx = bx + getDirOffsetX(cdir);
+			const int cy = by + getDirOffsetY(cdir);
+			const int ic = cx + cy * w;
+			cornerNeiArea = layer.areas[ic];
+		}
+
 	}
 }
 
@@ -247,20 +269,40 @@ static bool walkContour(dtTileCacheLayer& layer, int x, int y, int idx, unsigned
 
 	unsigned short neiReg = 0xffff;
 	unsigned char neiArea = 0;
+	unsigned char cornerNeiArea = 0;
+	unsigned short prevNeiArea = 0;
+	unsigned short prevCornerNeiArea = 0;
+	bool checkForPinning = false;
 
 	const int maxIter = w * h * 2;
 	int iter = 0;
 	while (iter < maxIter)
 	{
-		getNeighbourRegAndArea(layer, x, y, dir, neiReg, neiArea);
-
 		int nx = x;
 		int ny = y;
 		unsigned char ndir = dir;
 
+		getNeighbourRegAndArea(layer, x, y, dir, neiReg, neiArea, cornerNeiArea);
+
 		if (neiReg != layer.regs[x+y*w])
 		{
 			// Solid edge.
+			if (checkForPinning)
+			{
+				// Detect if there was 8-connected area type change during the turn.
+				// If it did, we need to make sure the vertex adde during turn does not get simplified.
+				// AB
+				// xC
+				// x = current location, A = prevNeiArea, B = prevCornerNeiArea, C = neiArea.
+				const bool shouldPinVertex = prevNeiArea == neiArea && prevCornerNeiArea != neiArea;
+				if (shouldPinVertex && cont.nverts > 0)
+				{
+					unsigned short* v = &cont.verts[(cont.nverts - 1) * 5];
+					v[4] |= 0x8000; // Use high bits in the area type to flag pinning of the vertex, will be removed in the contour simplification.
+				}
+			}
+
+			// Get corner vertex
 			int px = x;
 			int pz = y;
 			switch(dir)
@@ -276,6 +318,7 @@ static bool walkContour(dtTileCacheLayer& layer, int x, int y, int idx, unsigned
 
 			flags[idx] &= ~(1 << dir); // Remove visited edges
 			ndir = (dir+1) & 0x3;  // Rotate CW
+			checkForPinning = true;
 		}
 		else
 		{
@@ -283,11 +326,15 @@ static bool walkContour(dtTileCacheLayer& layer, int x, int y, int idx, unsigned
 			nx = x + getDirOffsetX(dir);
 			ny = y + getDirOffsetY(dir);
 			ndir = (dir+3) & 0x3;	// Rotate CCW
+			checkForPinning = false;
 		}
 
 		if (iter > 0 && idx == startIdx && dir == startDir)
 			break;
 
+		prevNeiArea = neiArea;
+		prevCornerNeiArea = cornerNeiArea;
+
 		x = nx;
 		y = ny;
 		dir = ndir;
@@ -343,10 +390,11 @@ static void simplifyContour(unsigned char area, dtTempContour& cont, const float
 	for (int i = 0; i < cont.nverts; ++i)
 	{
 		int j = (i+1) % cont.nverts;
-		// Check for start of a wall segment.
-		unsigned short ra = cont.verts[j*5+3];
-		unsigned short rb = cont.verts[i*5+3];
-		if (ra != rb)
+		// Check for start of a wall segment or pinned vertices.
+		const unsigned short ra = cont.verts[j*5+3];
+		const unsigned short rb = cont.verts[i*5+3];
+		const bool pinnedVertex = (cont.verts[i*5+4] & 0x8000) != 0;
+		if (ra != rb || pinnedVertex)
 			cont.poly[cont.npoly++] = (unsigned short)i;
 	}
 	if (cont.npoly < 2)
@@ -476,7 +524,7 @@ static void simplifyContour(unsigned char area, dtTempContour& cont, const float
 		dst[1] = src[1];
 		dst[2] = src[2];
 		dst[3] = src[3];
-		dst[4] = src[4];
+		dst[4] = src[4] & 0xff; // Mask out pinned vertex flag.
 		cont.nverts++;
 	}
 }
diff --git a/Engine/Source/Runtime/Navmesh/Public/DebugUtils/DetourDebugDraw.h b/Engine/Source/Runtime/Navmesh/Public/DebugUtils/DetourDebugDraw.h
--- a/Engine/Source/Runtime/Navmesh/Public/DebugUtils/DetourDebugDraw.h
+++ b/Engine/Source/Runtime/Navmesh/Public/DebugUtils/DetourDebugDraw.h
@@ -37,25 +37,25 @@ enum DrawNavMeshFlags
 	DU_DRAWNAVMESH_COLOR_TILES = 0x04,
 };
 
-void duDebugDrawNavMesh(struct duDebugDraw* dd, const dtNavMesh& mesh, unsigned char flags);
-void duDebugDrawNavMeshWithClosedList(struct duDebugDraw* dd, const dtNavMesh& mesh, const dtNavMeshQuery& query, unsigned char flags);
-void duDebugDrawNavMeshNodes(struct duDebugDraw* dd, const dtNavMeshQuery& query);
-void duDebugDrawNavMeshBVTree(struct duDebugDraw* dd, const dtNavMesh& mesh);
-void duDebugDrawNavMeshPortals(struct duDebugDraw* dd, const dtNavMesh& mesh);
-void duDebugDrawNavMeshPolysWithFlags(struct duDebugDraw* dd, const dtNavMesh& mesh, const unsigned short polyFlags, const unsigned int col);
-void duDebugDrawNavMeshPoly(struct duDebugDraw* dd, const dtNavMesh& mesh, dtPolyRef ref, const unsigned int col);
-void duDebugDrawNavMeshClusters(struct duDebugDraw* dd, const dtNavMesh& mesh);
-void duDebugDrawNavMeshCluster(struct duDebugDraw* dd, const dtNavMesh& mesh, dtClusterRef ref, const unsigned int col);
-
-void duDebugDrawTileCacheLayerAreas(struct duDebugDraw* dd, const dtTileCacheLayer& layer, const float cs, const float ch);
-void duDebugDrawTileCacheLayerDistance(struct duDebugDraw* dd, const dtTileCacheLayer& layer,
+NAVMESH_API void duDebugDrawNavMesh(struct duDebugDraw* dd, const dtNavMesh& mesh, unsigned char flags);
+NAVMESH_API void duDebugDrawNavMeshWithClosedList(struct duDebugDraw* dd, const dtNavMesh& mesh, const dtNavMeshQuery& query, unsigned char flags);
+NAVMESH_API void duDebugDrawNavMeshNodes(struct duDebugDraw* dd, const dtNavMeshQuery& query);
+NAVMESH_API void duDebugDrawNavMeshBVTree(struct duDebugDraw* dd, const dtNavMesh& mesh);
+NAVMESH_API void duDebugDrawNavMeshPortals(struct duDebugDraw* dd, const dtNavMesh& mesh);
+NAVMESH_API void duDebugDrawNavMeshPolysWithFlags(struct duDebugDraw* dd, const dtNavMesh& mesh, const unsigned short polyFlags, const unsigned int col);
+NAVMESH_API void duDebugDrawNavMeshPoly(struct duDebugDraw* dd, const dtNavMesh& mesh, dtPolyRef ref, const unsigned int col);
+NAVMESH_API void duDebugDrawNavMeshClusters(struct duDebugDraw* dd, const dtNavMesh& mesh);
+NAVMESH_API void duDebugDrawNavMeshCluster(struct duDebugDraw* dd, const dtNavMesh& mesh, dtClusterRef ref, const unsigned int col);
+
+NAVMESH_API void duDebugDrawTileCacheLayerAreas(struct duDebugDraw* dd, const dtTileCacheLayer& layer, const float cs, const float ch);
+NAVMESH_API void duDebugDrawTileCacheLayerDistance(struct duDebugDraw* dd, const dtTileCacheLayer& layer,
 									   const dtTileCacheDistanceField& dfield, const float cs, const float ch);
-void duDebugDrawTileCacheLayerRegions(struct duDebugDraw* dd, const dtTileCacheLayer& layer, const float cs, const float ch);
-void duDebugDrawTileCacheContours(duDebugDraw* dd, const struct dtTileCacheContourSet& lcset,
+NAVMESH_API void duDebugDrawTileCacheLayerRegions(struct duDebugDraw* dd, const dtTileCacheLayer& layer, const float cs, const float ch);
+NAVMESH_API void duDebugDrawTileCacheContours(duDebugDraw* dd, const struct dtTileCacheContourSet& lcset,
 								  const float* orig, const float cs, const float ch);
-void duDebugDrawTileCacheClusters(duDebugDraw* dd, const struct dtTileCacheClusterSet& lclusters);
-void duDebugDrawTileCachePolyMesh(duDebugDraw* dd, const struct dtTileCachePolyMesh& lmesh,
+NAVMESH_API void duDebugDrawTileCacheClusters(duDebugDraw* dd, const struct dtTileCacheClusterSet& lclusters);
+NAVMESH_API void duDebugDrawTileCachePolyMesh(duDebugDraw* dd, const struct dtTileCachePolyMesh& lmesh,
 								  const float* orig, const float cs, const float ch);
-void duDebugDrawTileCacheDetailMesh(duDebugDraw* dd, const struct dtTileCachePolyMeshDetail& dmesh);
+NAVMESH_API void duDebugDrawTileCacheDetailMesh(duDebugDraw* dd, const struct dtTileCachePolyMeshDetail& dmesh);
 
 #endif // DETOURDEBUGDRAW_H
