Remove duplicated code and implemented the timesliced version of the geometry gathering

#rnx

diff --git a/Engine/Source/Runtime/Engine/Classes/AI/Navigation/NavRelevantInterface.h b/Engine/Source/Runtime/Engine/Classes/AI/Navigation/NavRelevantInterface.h
--- a/Engine/Source/Runtime/Engine/Classes/AI/Navigation/NavRelevantInterface.h
+++ b/Engine/Source/Runtime/Engine/Classes/AI/Navigation/NavRelevantInterface.h
@@ -94,6 +94,11 @@ struct ENGINE_API FNavigationRelevantData : public TSharedFromThis<FNavigationRe
 			(Modifiers.HasAgentHeightAdjust() ? ENavigationDirtyFlag::UseAgentHeight : 0);
 	}
 
+	FORCEINLINE FCompositeNavModifier GetModifierForAgent(const struct FNavAgentProperties* NavAgent = nullptr) const
+	{
+		return Modifiers.HasMetaAreas() ? Modifiers.GetInstantiatedMetaModifier(NavAgent, SourceObject) : Modifiers;
+	}
+
 	bool HasPerInstanceTransforms() const;
 	bool IsMatchingFilter(const FNavigationRelevantDataFilter& Filter) const;
 	void Shrink();
diff --git a/Engine/Source/Runtime/NavigationSystem/Private/NavMesh/RecastNavMeshGenerator.cpp b/Engine/Source/Runtime/NavigationSystem/Private/NavMesh/RecastNavMeshGenerator.cpp
--- a/Engine/Source/Runtime/NavigationSystem/Private/NavMesh/RecastNavMeshGenerator.cpp
+++ b/Engine/Source/Runtime/NavigationSystem/Private/NavMesh/RecastNavMeshGenerator.cpp
@@ -1932,7 +1932,7 @@ FRecastTileGenerator::FRecastTileGenerator(FRecastNavMeshGenerator& ParentGenera
 
 	RasterizeGeomRecastState = ERasterizeGeomRecastTimeSlicedState::MarkWalkableTriangles;
 	RasterizeGeomState = ERasterizeGeomTimeSlicedState::RasterizeGeometryTransformCoords;
-	DoWorkTimeSlicedState = EDoWorkTimeSlicedState::DoAsyncGeometryGathering;
+	DoWorkTimeSlicedState = EDoWorkTimeSlicedState::GatherGeometryFromSources;
 	GenerateTileTimeSlicedState = EGenerateTileTimeSlicedState::GenerateCompressedLayers;
 
 	GenerateNavDataTimeSlicedState = EGenerateNavDataTimeSlicedState::Init;
@@ -2014,7 +2014,15 @@ void FRecastTileGenerator::Setup(const FRecastNavMeshGenerator& ParentGenerator,
 			}
 		}
 		
-		if (ParentGenerator.GatherGeometryOnGameThread())
+		bool bGatherGeometryNow = ParentGenerator.GatherGeometryOnGameThread();
+#if !RECAST_ASYNC_REBUILDING
+		if (ParentGenerator.IsTimeSliceRegenActive())
+		{
+			bGatherGeometryNow = false;
+		}
+#endif // !RECAST_ASYNC_REBUILDING
+
+		if (bGatherGeometryNow)
 		{
 			GatherGeometry(ParentGenerator, bRegenerateCompressedLayers);
 		}
@@ -2055,20 +2063,18 @@ ETimeSliceWorkResult FRecastTileGenerator::DoWorkTimeSliced()
 			return ETimeSliceWorkResult::Failed;
 		}
 		break;
-		case EDoWorkTimeSlicedState::DoAsyncGeometryGathering:
+		case EDoWorkTimeSlicedState::GatherGeometryFromSources:
 		{
-			DoWorkTimeSlicedState = EDoWorkTimeSlicedState::GenerateTile;
-
 			if (InclusionBounds.Num())
 			{
-				const bool bHadNavigationRelevantData = DoAsyncGeometryGathering();
+				WorkResult = GatherGeometryFromSourcesTimeSliced();
 
-				//check bHadNavigationRelevantData as an optimization so we don't call TestTimeSliceFinished when its unnecessary 
-				if (bHadNavigationRelevantData && TimeSliceManager->GetTimeSlicer().TestTimeSliceFinished())
+				if (WorkResult == ETimeSliceWorkResult::CallAgainNextTimeSlice)
 				{
-					return ETimeSliceWorkResult::CallAgainNextTimeSlice;
+					break;
 				}
 			}
+			DoWorkTimeSlicedState = EDoWorkTimeSlicedState::GenerateTile;
 		} //fall through to next state
 		case EDoWorkTimeSlicedState::GenerateTile:
 		{
@@ -2103,7 +2109,7 @@ bool FRecastTileGenerator::DoWork()
 	{
 		if (InclusionBounds.Num())
 		{
-			DoAsyncGeometryGathering();
+			GatherGeometryFromSources();
 		}
 
 		bSucceess = GenerateTile();
@@ -2124,97 +2130,44 @@ void FRecastTileGenerator::DumpAsyncData()
 	NavOctree = nullptr;
 }
 
-bool FRecastTileGenerator::DoAsyncGeometryGathering()
+void FRecastTileGenerator::GatherGeometryFromSources()
 {
-	QUICK_SCOPE_CYCLE_COUNTER(STAT_RecastNavMeshGenerator_PrepareGeometrySources);
-
-	const bool bRetVal = NavigationRelevantData.Num() > 0;
+	QUICK_SCOPE_CYCLE_COUNTER(STAT_RecastNavMeshGenerator_GatherGeometryFromSources);
 
-	for (auto& ElementData : NavigationRelevantData)
+	for (TSharedRef<FNavigationRelevantData, ESPMode::ThreadSafe>& ElementData : NavigationRelevantData)
 	{
 		if (ElementData->GetOwner() == nullptr)
 		{
-			UE_LOG(LogNavigation, Warning, TEXT("DoAsyncGeometryGathering: skipping an element with no longer valid Owner"));
+			UE_LOG(LogNavigation, Warning, TEXT(__FUNCTION__": skipping an element with no longer valid Owner"));
 			continue;
 		}
 
-		bool bDumpGeometryData = false;
-		if (ElementData->IsPendingLazyGeometryGathering() && ElementData->SupportsGatheringGeometrySlices())
-		{
-			QUICK_SCOPE_CYCLE_COUNTER(STAT_RecastNavMeshGenerator_LandscapeSlicesExporting);
-
-			FRecastGeometryExport GeomExport(*ElementData);
-
-			INavRelevantInterface* NavRelevant = Cast<INavRelevantInterface>(ElementData->GetOwner());
-			if(NavRelevant)
-			{
-				NavRelevant->PrepareGeometryExportSync();
-				// adding a small bump to avoid special case of zero-expansion when tile bounds
-				// overlap landscape's tile bounds
-				NavRelevant->GatherGeometrySlice(GeomExport, TileBBExpandedForAgent);
-
-				RecastGeometryExport::CovertCoordDataToRecast(GeomExport.VertexBuffer);
-				RecastGeometryExport::StoreCollisionCache(GeomExport);
-				bDumpGeometryData = true;
-			}
-			else
-			{
-				UE_LOG(LogNavigation, Error, TEXT("DoAsyncGeometryGathering: got an invalid NavRelevant instance!"));
-			}
-		}
-
-		if (ElementData->IsPendingLazyGeometryGathering() || ElementData->IsPendingLazyModifiersGathering())
-		{
-			NavOctree->DemandLazyDataGathering(*ElementData);
-		}
+		GatherNavigationDataGeometry(ElementData, NavOctree.Get(), NavDataConfig, bUpdateGeometry);
+	}
+}
 
-		const FCompositeNavModifier ModifierInstance = ElementData->Modifiers.HasMetaAreas() ? ElementData->Modifiers.GetInstantiatedMetaModifier(&NavDataConfig, ElementData->SourceObject) : ElementData->Modifiers;
+ETimeSliceWorkResult FRecastTileGenerator::GatherGeometryFromSourcesTimeSliced()
+{
+	QUICK_SCOPE_CYCLE_COUNTER(STAT_RecastNavMeshGenerator_GatherGeometryFromSources);
 
-		const bool bExportGeometry = bUpdateGeometry && ElementData->HasGeometry();
-		if (bExportGeometry)
+	while(NavigationRelevantData.Num())
+	{
+		TSharedRef<FNavigationRelevantData, ESPMode::ThreadSafe> ElementData = NavigationRelevantData.Pop(false/*bAllowShrinking*/);
+		if (ElementData->GetOwner() == nullptr)
 		{
-			if (ARecastNavMesh::IsVoxelCacheEnabled())
-			{
-				TNavStatArray<rcSpanCache> SpanData;
-				rcSpanCache* CachedVoxels = 0;
-				int32 NumCachedVoxels = 0;
-
-				DECLARE_SCOPE_CYCLE_COUNTER(TEXT("Rasterization: prepare voxel cache"), Stat_RecastRasterCachePrep, STATGROUP_Navigation);
-
-				if (!HasVoxelCache(ElementData->VoxelData, CachedVoxels, NumCachedVoxels))
-				{
-					// rasterize
-					PrepareVoxelCache(ElementData->CollisionData, ModifierInstance, SpanData);
-					CachedVoxels = SpanData.GetData();
-					NumCachedVoxels = SpanData.Num();
-
-					// encode
-					const int32 PrevElementMemory = ElementData->GetAllocatedSize();
-					FNavigationRelevantData* ModData = (FNavigationRelevantData*)&ElementData;
-					AddVoxelCache(ModData->VoxelData, CachedVoxels, NumCachedVoxels);
-
-					const int32 NewElementMemory = ElementData->GetAllocatedSize();
-					const int32 ElementMemoryDelta = NewElementMemory - PrevElementMemory;
-					INC_MEMORY_STAT_BY(STAT_Navigation_CollisionTreeMemory, ElementMemoryDelta);
-				}
-			}
-			else
-			{
-				ValidateAndAppendGeometry(ElementData, ModifierInstance);
-			}
-
-			if (bDumpGeometryData)
-			{
-				const_cast<FNavigationRelevantData&>(*ElementData).CollisionData.Empty();
-			}
+			UE_LOG(LogNavigation, Warning, TEXT(__FUNCTION__": skipping an element with no longer valid Owner"));
+			continue;
 		}
 
-		if (ModifierInstance.IsEmpty() == false)
+		GatherNavigationDataGeometry(ElementData, NavOctree.Get(), NavDataConfig, bUpdateGeometry);
+		
+		if (TimeSliceManager->GetTimeSlicer().TestTimeSliceFinished())
 		{
-			AppendModifier(ModifierInstance, ElementData->NavDataPerInstanceTransformDelegate);
+			return ETimeSliceWorkResult::CallAgainNextTimeSlice;
 		}
 	}
-	return bRetVal;
+
+	return ETimeSliceWorkResult::Succeeded;
 }
 
 void FRecastTileGenerator::PrepareGeometrySources(const FRecastNavMeshGenerator& ParentGenerator, bool bGeometryChanged)
@@ -2237,8 +2190,10 @@ void FRecastTileGenerator::PrepareGeometrySources(const FRecastNavMeshGenerator&
 		if (bShouldUse)
 		{
 			const bool bExportGeometry = bGeometryChanged && (Element.Data->HasGeometry() || Element.Data->IsPendingLazyGeometryGathering());
-			if (bExportGeometry
-				|| (Element.Data->IsPendingLazyModifiersGathering() || Element.Data->Modifiers.HasMetaAreas() == true || Element.Data->Modifiers.IsEmpty() == false))
+			if (bExportGeometry || 
+				Element.Data->IsPendingLazyModifiersGathering() || 
+				Element.Data->Modifiers.HasMetaAreas() == true || 
+				Element.Data->Modifiers.IsEmpty() == false)
 			{
 				NavigationRelevantData.Add(Element.Data);
 			}
@@ -2266,87 +2221,88 @@ void FRecastTileGenerator::GatherGeometry(const FRecastNavMeshGenerator& ParentG
 		const bool bShouldUse = Element.ShouldUseGeometry(OwnerNavDataConfig);
 		if (bShouldUse)
 		{
-			bool bDumpGeometryData = false;
-			if (Element.Data->IsPendingLazyGeometryGathering() || Element.Data->IsPendingLazyModifiersGathering())
-			{
-				const bool bSupportsSlices = Element.Data->SupportsGatheringGeometrySlices();
-				if (bSupportsSlices == false || Element.Data->IsPendingLazyModifiersGathering() == true)
-				{
-					QUICK_SCOPE_CYCLE_COUNTER(STAT_RecastNavMeshGenerator_LazyGeometryExport);
-					NavigationOctree->DemandLazyDataGathering(Element);
-				}
+			GatherNavigationDataGeometry(Element.Data, NavigationOctree, OwnerNavDataConfig, bGeometryChanged);
+		}
+	}
+}
+
+void FRecastTileGenerator::GatherNavigationDataGeometry(const TSharedRef<FNavigationRelevantData, ESPMode::ThreadSafe>& ElementData, FNavigationOctree* NavigationOctree, const FNavDataConfig& OwnerNavDataConfig, const bool bGeometryChanged)
+{
+	bool bDumpGeometryData = false;
+	if (ElementData->IsPendingLazyGeometryGathering() || ElementData->IsPendingLazyModifiersGathering())
+	{
+		QUICK_SCOPE_CYCLE_COUNTER(STAT_RecastNavMeshGenerator_LazyGeometryExport);
+		NavigationOctree->DemandLazyDataGathering(*ElementData);
+	}
 				
-				if (bSupportsSlices == true)
-				{
-					QUICK_SCOPE_CYCLE_COUNTER(STAT_RecastNavMeshGenerator_LandscapeSlicesExporting);
+	if (ElementData->IsPendingLazyGeometryGathering() && ElementData->SupportsGatheringGeometrySlices())
+	{
+		QUICK_SCOPE_CYCLE_COUNTER(STAT_RecastNavMeshGenerator_LandscapeSlicesExporting);
 
-					FRecastGeometryExport GeomExport(const_cast<FNavigationRelevantData&>(*Element.Data));
+		FRecastGeometryExport GeomExport(const_cast<FNavigationRelevantData&>(*ElementData));
 
-					INavRelevantInterface* NavRelevant = const_cast<INavRelevantInterface*>(Cast<const INavRelevantInterface>(Element.GetOwner()));
-					if (NavRelevant)
-					{
-						NavRelevant->PrepareGeometryExportSync();
-						// adding a small bump to avoid special case of zero-expansion when tile bounds
-						// overlap landscape's tile bounds
-						NavRelevant->GatherGeometrySlice(GeomExport, TileBBExpandedForAgent);
-
-						RecastGeometryExport::CovertCoordDataToRecast(GeomExport.VertexBuffer);
-						RecastGeometryExport::StoreCollisionCache(GeomExport);
-						bDumpGeometryData = true;
-					}
-					else
-					{
-						UE_LOG(LogNavigation, Error, TEXT("GatherGeometry: got an invalid NavRelevant instance!"));
-					}
-				}
-			}
+		INavRelevantInterface* NavRelevant = const_cast<INavRelevantInterface*>(Cast<const INavRelevantInterface>(ElementData->GetOwner()));
+		if (NavRelevant)
+		{
+			NavRelevant->PrepareGeometryExportSync();
+			// adding a small bump to avoid special case of zero-expansion when tile bounds
+			// overlap landscape's tile bounds
+			NavRelevant->GatherGeometrySlice(GeomExport, TileBBExpandedForAgent);
+
+			RecastGeometryExport::CovertCoordDataToRecast(GeomExport.VertexBuffer);
+			RecastGeometryExport::StoreCollisionCache(GeomExport);
+			bDumpGeometryData = true;
+		}
+		else
+		{
+			UE_LOG(LogNavigation, Error, TEXT("GatherGeometry: got an invalid NavRelevant instance!"));
+		}
+	}
 
-			const FCompositeNavModifier ModifierInstance = Element.GetModifierForAgent(&OwnerNavDataConfig);
+	const FCompositeNavModifier ModifierInstance = ElementData->GetModifierForAgent(&OwnerNavDataConfig);
 
-			const bool bExportGeometry = bGeometryChanged && Element.Data->HasGeometry();
-			if (bExportGeometry)
-			{
-				if (ARecastNavMesh::IsVoxelCacheEnabled())
-				{
-					TNavStatArray<rcSpanCache> SpanData;
-					rcSpanCache* CachedVoxels = 0;
-					int32 NumCachedVoxels = 0;
+	const bool bExportGeometry = bGeometryChanged && ElementData->HasGeometry();
+	if (bExportGeometry)
+	{
+		if (ARecastNavMesh::IsVoxelCacheEnabled())
+		{
+			TNavStatArray<rcSpanCache> SpanData;
+			rcSpanCache* CachedVoxels = 0;
+			int32 NumCachedVoxels = 0;
 
-					DECLARE_SCOPE_CYCLE_COUNTER(TEXT("Rasterization: prepare voxel cache"), Stat_RecastRasterCachePrep, STATGROUP_Navigation);
+			DECLARE_SCOPE_CYCLE_COUNTER(TEXT("Rasterization: prepare voxel cache"), Stat_RecastRasterCachePrep, STATGROUP_Navigation);
 
-					if (!HasVoxelCache(Element.Data->VoxelData, CachedVoxels, NumCachedVoxels))
-					{
+			if (!HasVoxelCache(ElementData->VoxelData, CachedVoxels, NumCachedVoxels))
+			{
 						// rasterize
-						PrepareVoxelCache(Element.Data->CollisionData, ModifierInstance, SpanData);
-						CachedVoxels = SpanData.GetData();
-						NumCachedVoxels = SpanData.Num();
+				PrepareVoxelCache(ElementData->CollisionData, ModifierInstance, SpanData);
+				CachedVoxels = SpanData.GetData();
+				NumCachedVoxels = SpanData.Num();
 
 						// encode
-						const int32 PrevElementMemory = Element.Data->GetAllocatedSize();
-						FNavigationRelevantData* ModData = (FNavigationRelevantData*)&Element.Data;
-						AddVoxelCache(ModData->VoxelData, CachedVoxels, NumCachedVoxels);
-
-						const int32 NewElementMemory = Element.Data->GetAllocatedSize();
-						const int32 ElementMemoryDelta = NewElementMemory - PrevElementMemory;
-						INC_MEMORY_STAT_BY(STAT_Navigation_CollisionTreeMemory, ElementMemoryDelta);
-					}
-				}
-				else
-				{
-					ValidateAndAppendGeometry(Element.Data, ModifierInstance);
-				}
+				const int32 PrevElementMemory = ElementData->GetAllocatedSize();
+				FNavigationRelevantData* ModData = (FNavigationRelevantData*)&ElementData;
+				AddVoxelCache(ModData->VoxelData, CachedVoxels, NumCachedVoxels);
 
-				if (bDumpGeometryData)
-				{
-					const_cast<FNavigationRelevantData&>(*Element.Data).CollisionData.Empty();
-				}
-			}
-						
-			if (ModifierInstance.IsEmpty() == false)
-			{
-				AppendModifier(ModifierInstance, Element.Data->NavDataPerInstanceTransformDelegate);
+				const int32 NewElementMemory = ElementData->GetAllocatedSize();
+				const int32 ElementMemoryDelta = NewElementMemory - PrevElementMemory;
+				INC_MEMORY_STAT_BY(STAT_Navigation_CollisionTreeMemory, ElementMemoryDelta);
 			}
 		}
+		else
+		{
+			ValidateAndAppendGeometry(ElementData, ModifierInstance);
+		}
+
+		if (bDumpGeometryData)
+		{
+			const_cast<FNavigationRelevantData&>(*ElementData).CollisionData.Empty();
+		}
+	}
+				
+	if (ModifierInstance.IsEmpty() == false)
+	{
+		AppendModifier(ModifierInstance, ElementData->NavDataPerInstanceTransformDelegate);
 	}
 }
 
@@ -2596,7 +2552,7 @@ void FRecastTileGenerator::AppendModifier(const FCompositeNavModifier& Modifier,
 	Modifiers.Add(MoveTemp(ModifierElement));
 }
 
-void FRecastTileGenerator::ValidateAndAppendGeometry(TSharedRef<FNavigationRelevantData, ESPMode::ThreadSafe> ElementData, const FCompositeNavModifier& InModifier)
+void FRecastTileGenerator::ValidateAndAppendGeometry(const TSharedRef<FNavigationRelevantData, ESPMode::ThreadSafe>& ElementData, const FCompositeNavModifier& InModifier)
 {
 	const FNavigationRelevantData& DataRef = ElementData.Get();
 	if (DataRef.IsCollisionDataValid())
@@ -6318,6 +6274,11 @@ bool FRecastNavMeshGenerator::GatherGeometryOnGameThread() const
 	return DestNavMesh == nullptr || DestNavMesh->ShouldGatherDataOnGameThread() == true;
 }
 
+bool FRecastNavMeshGenerator::IsTimeSliceRegenActive() const
+{
+	return SyncTimeSlicedData.bTimeSliceRegenActive;
+}
+
 bool FRecastNavMeshGenerator::IsTileChanged(int32 TileIdx) const
 {
 #if WITH_EDITOR	
diff --git a/Engine/Source/Runtime/NavigationSystem/Public/NavMesh/RecastNavMeshGenerator.h b/Engine/Source/Runtime/NavigationSystem/Public/NavMesh/RecastNavMeshGenerator.h
--- a/Engine/Source/Runtime/NavigationSystem/Public/NavMesh/RecastNavMeshGenerator.h
+++ b/Engine/Source/Runtime/NavigationSystem/Public/NavMesh/RecastNavMeshGenerator.h
@@ -239,7 +239,7 @@ enum class ERasterizeGeomTimeSlicedState : uint8
 enum class EDoWorkTimeSlicedState : uint8
 {
 	Invalid,
-	DoAsyncGeometryGathering,
+	GatherGeometryFromSources,
 	GenerateTile,
 };
 
@@ -328,10 +328,16 @@ protected:
 
 	void Setup(const FRecastNavMeshGenerator& ParentGenerator, const TArray<FBox>& DirtyAreas);
 	
+	/** Gather geometry */
 	void GatherGeometry(const FRecastNavMeshGenerator& ParentGenerator, bool bGeometryChanged);
+	/** Gather geometry sources to be processed later by the GatherGeometryFromSources */
 	void PrepareGeometrySources(const FRecastNavMeshGenerator& ParentGenerator, bool bGeometryChanged);
-	/** Returns true if there is any NavigationRelevantData */
-	bool DoAsyncGeometryGathering();
+	/** Gather geometry from the prefetched sources */
+	void GatherGeometryFromSources();
+	/** Gather geometry from the prefetched sources time sliced version */
+	ETimeSliceWorkResult GatherGeometryFromSourcesTimeSliced();
+	/** Gather geometry from a specified Navigation Data */
+	void GatherNavigationDataGeometry(const TSharedRef<FNavigationRelevantData, ESPMode::ThreadSafe>& ElementData, FNavigationOctree* NavigationOctree, const FNavDataConfig& OwnerNavDataConfig, bool bGeometryChanged);
 
 	/** Start functions used by GenerateCompressedLayersTimeSliced / GenerateCompressedLayers */
 	bool CreateHeightField(FNavMeshBuildContext& BuildContext, FTileRasterizationContext& RasterContext);
@@ -378,7 +384,7 @@ protected:
 
 	void AppendModifier(const FCompositeNavModifier& Modifier, const FNavDataPerInstanceTransformDelegate& InTransformsDelegate);
 	/** Appends specified geometry to tile's geometry */
-	void ValidateAndAppendGeometry(TSharedRef<FNavigationRelevantData, ESPMode::ThreadSafe> ElementData, const FCompositeNavModifier& InModifier);
+	void ValidateAndAppendGeometry(const TSharedRef<FNavigationRelevantData, ESPMode::ThreadSafe>& ElementData, const FCompositeNavModifier& InModifier);
 	void AppendGeometry(const TNavStatArray<uint8>& RawCollisionCache, const FCompositeNavModifier& InModifier, const FNavDataPerInstanceTransformDelegate& InTransformsDelegate);
 	void AppendVoxels(rcSpanCache* SpanData, int32 NumSpans);
 	
@@ -652,6 +658,7 @@ public:
 	int32 GetDirtyTilesCount(const FBox& AreaBounds) const;
 
 	bool GatherGeometryOnGameThread() const;
+	bool IsTimeSliceRegenActive() const;
 
 	FBox GrowBoundingBox(const FBox& BBox, bool bIncludeAgentHeight) const;
 	
diff --git a/Engine/Source/Runtime/NavigationSystem/Public/NavigationOctree.h b/Engine/Source/Runtime/NavigationSystem/Public/NavigationOctree.h
--- a/Engine/Source/Runtime/NavigationSystem/Public/NavigationOctree.h
+++ b/Engine/Source/Runtime/NavigationSystem/Public/NavigationOctree.h
@@ -56,9 +56,9 @@ public:
 	 *	If it does then retrieves a copy with meta areas substituted with
 	 *	appropriate non-meta areas, depending on NavAgent
 	 */
-	FORCEINLINE FCompositeNavModifier GetModifierForAgent(const struct FNavAgentProperties* NavAgent = NULL) const 
+	FORCEINLINE FCompositeNavModifier GetModifierForAgent(const struct FNavAgentProperties* NavAgent = nullptr) const 
 	{ 
-		return Data->Modifiers.HasMetaAreas() ? Data->Modifiers.GetInstantiatedMetaModifier(NavAgent, Data->SourceObject) : Data->Modifiers;
+		return Data->GetModifierForAgent(NavAgent);
 	}
 
 	FORCEINLINE bool ShouldUseGeometry(const FNavDataConfig& NavConfig) const
