Build fix

diff --git a/Engine/Source/Runtime/NavigationSystem/Private/NavMesh/RecastNavMeshGenerator.cpp b/Engine/Source/Runtime/NavigationSystem/Private/NavMesh/RecastNavMeshGenerator.cpp
--- a/Engine/Source/Runtime/NavigationSystem/Private/NavMesh/RecastNavMeshGenerator.cpp
+++ b/Engine/Source/Runtime/NavigationSystem/Private/NavMesh/RecastNavMeshGenerator.cpp
@@ -2138,7 +2138,7 @@ void FRecastTileGenerator::GatherGeometryFromSources()
 	{
 		if (ElementData->GetOwner() == nullptr)
 		{
-			UE_LOG(LogNavigation, Warning, TEXT(__FUNCTION__": skipping an element with no longer valid Owner"));
+			UE_LOG(LogNavigation, Warning, TEXT("%s: skipping an element with no longer valid Owner"), ANSI_TO_TCHAR(__FUNCTION__));
 			continue;
 		}
 
@@ -2155,7 +2155,7 @@ ETimeSliceWorkResult FRecastTileGenerator::GatherGeometryFromSourcesTimeSliced()
 		TSharedRef<FNavigationRelevantData, ESPMode::ThreadSafe> ElementData = NavigationRelevantData.Pop(false/*bAllowShrinking*/);
 		if (ElementData->GetOwner() == nullptr)
 		{
-			UE_LOG(LogNavigation, Warning, TEXT(__FUNCTION__": skipping an element with no longer valid Owner"));
+			UE_LOG(LogNavigation, Warning, TEXT("%s: skipping an element with no longer valid Owner"), ANSI_TO_TCHAR(__FUNCTION__));
 			continue;
 		}
 
