Support lazy gathering on child of octree nodes

#rnx

diff --git a/Engine/Source/Runtime/Engine/Classes/AI/Navigation/NavRelevantInterface.h b/Engine/Source/Runtime/Engine/Classes/AI/Navigation/NavRelevantInterface.h
--- a/Engine/Source/Runtime/Engine/Classes/AI/Navigation/NavRelevantInterface.h
+++ b/Engine/Source/Runtime/Engine/Classes/AI/Navigation/NavRelevantInterface.h
@@ -70,6 +70,7 @@ struct ENGINE_API FNavigationRelevantData : public TSharedFromThis<FNavigationRe
 	*	Mind that this flag can go back to 'true' if related data gets cleared out. */
 	uint32 bPendingLazyGeometryGathering : 1;
 	uint32 bPendingLazyModifiersGathering : 1;
+	uint32 bPendingChildLazyModifiersGathering : 1;
 
 	uint32 bSupportsGatheringGeometrySlices : 1;
 
@@ -77,12 +78,15 @@ struct ENGINE_API FNavigationRelevantData : public TSharedFromThis<FNavigationRe
 		: SourceObject(&Source)
 		, bPendingLazyGeometryGathering(false)
 		, bPendingLazyModifiersGathering(false)
+		, bPendingChildLazyModifiersGathering(false)
 	{}
 
 	FORCEINLINE bool HasGeometry() const { return VoxelData.Num() || CollisionData.Num(); }
 	FORCEINLINE bool HasModifiers() const { return !Modifiers.IsEmpty(); }
 	FORCEINLINE bool IsPendingLazyGeometryGathering() const { return bPendingLazyGeometryGathering; }
 	FORCEINLINE bool IsPendingLazyModifiersGathering() const { return bPendingLazyModifiersGathering; }
+	FORCEINLINE bool IsPendingChildLazyModifiersGathering() const { return bPendingChildLazyModifiersGathering; }
+	FORCEINLINE bool NeedAnyPendingLazyModifiersGathering() const { return bPendingLazyModifiersGathering || bPendingChildLazyModifiersGathering; }
 	FORCEINLINE bool SupportsGatheringGeometrySlices() const { return bSupportsGatheringGeometrySlices; }
 	FORCEINLINE bool IsEmpty() const { return !HasGeometry() && !HasModifiers(); }
 	FORCEINLINE uint32 GetAllocatedSize() const { return CollisionData.GetAllocatedSize() + VoxelData.GetAllocatedSize() + Modifiers.GetAllocatedSize(); }
@@ -90,7 +94,7 @@ struct ENGINE_API FNavigationRelevantData : public TSharedFromThis<FNavigationRe
 	FORCEINLINE int32 GetDirtyFlag() const
 	{
 		return ((HasGeometry() || IsPendingLazyGeometryGathering() || Modifiers.GetMaskFillCollisionUnderneathForNavmesh()) ? ENavigationDirtyFlag::Geometry : 0) |
-			((HasModifiers() || IsPendingLazyModifiersGathering()) ? ENavigationDirtyFlag::DynamicModifier : 0) |
+			((HasModifiers() || NeedAnyPendingLazyModifiersGathering()) ? ENavigationDirtyFlag::DynamicModifier : 0) |
 			(Modifiers.HasAgentHeightAdjust() ? ENavigationDirtyFlag::UseAgentHeight : 0);
 	}
 
diff --git a/Engine/Source/Runtime/NavigationSystem/Private/NavMesh/RecastNavMeshGenerator.cpp b/Engine/Source/Runtime/NavigationSystem/Private/NavMesh/RecastNavMeshGenerator.cpp
--- a/Engine/Source/Runtime/NavigationSystem/Private/NavMesh/RecastNavMeshGenerator.cpp
+++ b/Engine/Source/Runtime/NavigationSystem/Private/NavMesh/RecastNavMeshGenerator.cpp
@@ -2127,13 +2127,19 @@ void FRecastTileGenerator::DumpAsyncData()
 	OffmeshLinks.Empty();
 
 	NavigationRelevantData.Empty();
-	NavOctree = nullptr;
+	NavSystem = nullptr;
 }
 
 void FRecastTileGenerator::GatherGeometryFromSources()
 {
 	QUICK_SCOPE_CYCLE_COUNTER(STAT_RecastNavMeshGenerator_GatherGeometryFromSources);
 
+	UNavigationSystemV1* NavSys = NavSystem.Get();
+	if (NavSys == nullptr)
+	{
+		return;
+	}
+
 	for (TSharedRef<FNavigationRelevantData, ESPMode::ThreadSafe>& ElementData : NavigationRelevantData)
 	{
 		if (ElementData->GetOwner() == nullptr)
@@ -2142,7 +2148,7 @@ void FRecastTileGenerator::GatherGeometryFromSources()
 			continue;
 		}
 
-		GatherNavigationDataGeometry(ElementData, NavOctree.Get(), NavDataConfig, bUpdateGeometry);
+		GatherNavigationDataGeometry(ElementData, *NavSys, NavDataConfig, bUpdateGeometry);
 	}
 }
 
@@ -2150,6 +2156,12 @@ ETimeSliceWorkResult FRecastTileGenerator::GatherGeometryFromSourcesTimeSliced()
 {
 	QUICK_SCOPE_CYCLE_COUNTER(STAT_RecastNavMeshGenerator_GatherGeometryFromSources);
 
+	UNavigationSystemV1* NavSys = NavSystem.Get();
+	if (NavSys == nullptr)
+	{
+		return ETimeSliceWorkResult::Failed;
+	}
+
 	while(NavigationRelevantData.Num())
 	{
 		TSharedRef<FNavigationRelevantData, ESPMode::ThreadSafe> ElementData = NavigationRelevantData.Pop(false/*bAllowShrinking*/);
@@ -2159,7 +2171,7 @@ ETimeSliceWorkResult FRecastTileGenerator::GatherGeometryFromSourcesTimeSliced()
 			continue;
 		}
 
-		GatherNavigationDataGeometry(ElementData, NavOctree.Get(), NavDataConfig, bUpdateGeometry);
+		GatherNavigationDataGeometry(ElementData, *NavSys, NavDataConfig, bUpdateGeometry);
 		
 		if (TimeSliceManager->GetTimeSlicer().TestTimeSliceFinished())
 		{
@@ -2178,7 +2190,7 @@ void FRecastTileGenerator::PrepareGeometrySources(const FRecastNavMeshGenerator&
 	FNavigationOctree* NavOctreeInstance = NavSys ? NavSys->GetMutableNavOctree() : nullptr;
 	check(NavOctreeInstance);
 	NavigationRelevantData.Reset();
-	NavOctree = NavOctreeInstance->AsShared();
+	NavSystem = NavSys;
 	bUpdateGeometry = bGeometryChanged;
 
 	for (FNavigationOctree::TConstElementBoxIterator<FNavigationOctree::DefaultStackAllocator> It(*NavOctreeInstance, ParentGenerator.GrowBoundingBox(TileBB, /*bIncludeAgentHeight*/ false));
@@ -2191,7 +2203,7 @@ void FRecastTileGenerator::PrepareGeometrySources(const FRecastNavMeshGenerator&
 		{
 			const bool bExportGeometry = bGeometryChanged && (Element.Data->HasGeometry() || Element.Data->IsPendingLazyGeometryGathering());
 			if (bExportGeometry || 
-				Element.Data->IsPendingLazyModifiersGathering() || 
+				Element.Data->NeedAnyPendingLazyModifiersGathering() ||
 				Element.Data->Modifiers.HasMetaAreas() == true || 
 				Element.Data->Modifiers.IsEmpty() == false)
 			{
@@ -2221,18 +2233,18 @@ void FRecastTileGenerator::GatherGeometry(const FRecastNavMeshGenerator& ParentG
 		const bool bShouldUse = Element.ShouldUseGeometry(OwnerNavDataConfig);
 		if (bShouldUse)
 		{
-			GatherNavigationDataGeometry(Element.Data, NavigationOctree, OwnerNavDataConfig, bGeometryChanged);
+			GatherNavigationDataGeometry(Element.Data, *NavSys, OwnerNavDataConfig, bGeometryChanged);
 		}
 	}
 }
 
-void FRecastTileGenerator::GatherNavigationDataGeometry(const TSharedRef<FNavigationRelevantData, ESPMode::ThreadSafe>& ElementData, FNavigationOctree* NavigationOctree, const FNavDataConfig& OwnerNavDataConfig, const bool bGeometryChanged)
+void FRecastTileGenerator::GatherNavigationDataGeometry(const TSharedRef<FNavigationRelevantData, ESPMode::ThreadSafe>& ElementData, UNavigationSystemV1& NavSys, const FNavDataConfig& OwnerNavDataConfig, const bool bGeometryChanged)
 {
 	bool bDumpGeometryData = false;
-	if (ElementData->IsPendingLazyGeometryGathering() || ElementData->IsPendingLazyModifiersGathering())
+	if (ElementData->IsPendingLazyGeometryGathering() || ElementData->NeedAnyPendingLazyModifiersGathering())
 	{
 		QUICK_SCOPE_CYCLE_COUNTER(STAT_RecastNavMeshGenerator_LazyGeometryExport);
-		NavigationOctree->DemandLazyDataGathering(*ElementData);
+		NavSys.DemandLazyDataGathering(*ElementData);
 	}
 				
 	if (ElementData->IsPendingLazyGeometryGathering() && ElementData->SupportsGatheringGeometrySlices())
diff --git a/Engine/Source/Runtime/NavigationSystem/Private/NavigationDataHandler.cpp b/Engine/Source/Runtime/NavigationSystem/Private/NavigationDataHandler.cpp
--- a/Engine/Source/Runtime/NavigationSystem/Private/NavigationDataHandler.cpp
+++ b/Engine/Source/Runtime/NavigationSystem/Private/NavigationDataHandler.cpp
@@ -499,4 +499,31 @@ void FNavigationDataHandler::ProcessPendingOctreeUpdates()
 		}
 	}
 	OctreeController.PendingOctreeUpdates.Empty(32);
-}
\ No newline at end of file
+}
+
+void FNavigationDataHandler::DemandLazyDataGathering(FNavigationRelevantData& ElementData)
+{
+	// Do the lazy gathering on the element
+	OctreeController.NavOctree->DemandLazyDataGathering(ElementData);
+
+    // Check if any child asked for some lazy gathering
+	if (ElementData.IsPendingChildLazyModifiersGathering())
+	{
+		TArray<FWeakObjectPtr> ChildNodes;
+		OctreeController.OctreeChildNodesMap.MultiFind(ElementData.GetOwner(), ChildNodes);
+
+		for (FWeakObjectPtr& ChildNode : ChildNodes)
+		{
+			if (ChildNode.IsValid())
+			{
+				UObject* ChildNodeOb = ChildNode.Get();
+				INavRelevantInterface* ChildNavInterface = ChildNodeOb ? Cast<INavRelevantInterface>(ChildNodeOb) : nullptr;
+				if (ChildNavInterface)
+				{
+					OctreeController.NavOctree->DemandChildLazyDataGathering(ElementData, *ChildNavInterface);
+				}
+			}
+		}
+		ElementData.bPendingChildLazyModifiersGathering = false;
+	}
+}
diff --git a/Engine/Source/Runtime/NavigationSystem/Private/NavigationOctree.cpp b/Engine/Source/Runtime/NavigationSystem/Private/NavigationOctree.cpp
--- a/Engine/Source/Runtime/NavigationSystem/Private/NavigationOctree.cpp
+++ b/Engine/Source/Runtime/NavigationSystem/Private/NavigationOctree.cpp
@@ -82,10 +82,22 @@ void FNavigationOctree::DemandLazyDataGathering(FNavigationRelevantData& Element
 	INC_MEMORY_STAT_BY(STAT_Navigation_CollisionTreeMemory, ElementMemoryChange);
 }
 
-void FNavigationOctree::DemandLazyDataGathering(const FNavigationOctreeElement& Element)
+void FNavigationOctree::DemandChildLazyDataGathering(FNavigationRelevantData& ElementData, INavRelevantInterface& ChildNavInterface)
 {
-	FNavigationRelevantData& MutableData = const_cast<FNavigationRelevantData&>(*Element.Data);
-	DemandLazyDataGathering(MutableData);
+	if (IsLazyGathering(ChildNavInterface))
+	{
+		ChildNavInterface.GetNavigationData(ElementData);
+		ElementData.ValidateAndShrink();
+	}
+}
+
+bool FNavigationOctree::IsLazyGathering(const INavRelevantInterface& ChildNavInterface) const
+{
+	const ENavDataGatheringMode GatheringMode = ChildNavInterface.GetGeometryGatheringMode();
+	const bool bDoInstantGathering = ((GatheringMode == ENavDataGatheringMode::Default && DefaultGeometryGatheringMode == ENavDataGatheringMode::Instant)
+		|| GatheringMode == ENavDataGatheringMode::Instant);
+
+	return !bDoInstantGathering;
 }
 
 void FNavigationOctree::AddNode(UObject* ElementOb, INavRelevantInterface* NavElement, const FBox& Bounds, FNavigationOctreeElement& Element)
@@ -95,9 +107,7 @@ void FNavigationOctree::AddNode(UObject* ElementOb, INavRelevantInterface* NavEl
 
 	if (NavElement)
 	{
-		const ENavDataGatheringMode GatheringMode = NavElement->GetGeometryGatheringMode();
-		bool bDoInstantGathering = ((GatheringMode == ENavDataGatheringMode::Default && DefaultGeometryGatheringMode == ENavDataGatheringMode::Instant) 
-			|| GatheringMode == ENavDataGatheringMode::Instant);
+		const bool bDoInstantGathering = !IsLazyGathering(*NavElement);
 
 		if (bGatherGeometry)
 		{
@@ -149,7 +159,16 @@ void FNavigationOctree::AppendToNode(const FOctreeElementId& Id, INavRelevantInt
 	if (NavElement)
 	{
 		SCOPE_CYCLE_COUNTER(STAT_Navigation_GatheringNavigationModifiersSync);
-		NavElement->GetNavigationData(*Element.Data);
+		const bool bDoInstantGathering = !IsLazyGathering(*NavElement);
+
+		if (bDoInstantGathering)
+		{
+			NavElement->GetNavigationData(*Element.Data);
+		}
+		else
+		{
+			Element.Data->bPendingChildLazyModifiersGathering = true;
+		}
 	}
 
 	// validate exported data
diff --git a/Engine/Source/Runtime/NavigationSystem/Private/NavigationSystem.cpp b/Engine/Source/Runtime/NavigationSystem/Private/NavigationSystem.cpp
--- a/Engine/Source/Runtime/NavigationSystem/Private/NavigationSystem.cpp
+++ b/Engine/Source/Runtime/NavigationSystem/Private/NavigationSystem.cpp
@@ -2703,6 +2703,11 @@ void UNavigationSystemV1::RemoveNavOctreeElementId(const FOctreeElementId& Eleme
 	FNavigationDataHandler(DefaultOctreeController, DefaultDirtyAreasController).RemoveNavOctreeElementId(ElementId, UpdateFlags);
 }
 
+void UNavigationSystemV1::DemandLazyDataGathering(FNavigationRelevantData& ElementData)
+{
+	FNavigationDataHandler(DefaultOctreeController, DefaultDirtyAreasController).DemandLazyDataGathering(ElementData);
+}
+
 const FNavigationRelevantData* UNavigationSystemV1::GetDataForObject(const UObject& Object) const
 {
 	return DefaultOctreeController.GetDataForObject(Object);
diff --git a/Engine/Source/Runtime/NavigationSystem/Public/NavMesh/RecastNavMeshGenerator.h b/Engine/Source/Runtime/NavigationSystem/Public/NavMesh/RecastNavMeshGenerator.h
--- a/Engine/Source/Runtime/NavigationSystem/Public/NavMesh/RecastNavMeshGenerator.h
+++ b/Engine/Source/Runtime/NavigationSystem/Public/NavMesh/RecastNavMeshGenerator.h
@@ -38,6 +38,7 @@ struct FTileCacheAllocator;
 struct FTileGenerationContext;
 class dtNavMesh;
 class FNavRegenTimeSliceManager;
+class UNavigationSystemV1;
 
 #define MAX_VERTS_PER_POLY	6
 
@@ -337,7 +338,7 @@ protected:
 	/** Gather geometry from the prefetched sources time sliced version */
 	ETimeSliceWorkResult GatherGeometryFromSourcesTimeSliced();
 	/** Gather geometry from a specified Navigation Data */
-	void GatherNavigationDataGeometry(const TSharedRef<FNavigationRelevantData, ESPMode::ThreadSafe>& ElementData, FNavigationOctree* NavigationOctree, const FNavDataConfig& OwnerNavDataConfig, bool bGeometryChanged);
+	void GatherNavigationDataGeometry(const TSharedRef<FNavigationRelevantData, ESPMode::ThreadSafe>& ElementData, UNavigationSystemV1& NavSys, const FNavDataConfig& OwnerNavDataConfig, bool bGeometryChanged);
 
 	/** Start functions used by GenerateCompressedLayersTimeSliced / GenerateCompressedLayers */
 	bool CreateHeightField(FNavMeshBuildContext& BuildContext, FTileRasterizationContext& RasterContext);
@@ -462,7 +463,7 @@ protected:
 	TWeakPtr<FNavDataGenerator, ESPMode::ThreadSafe> ParentGeneratorWeakPtr;
 
 	TNavStatArray<TSharedRef<FNavigationRelevantData, ESPMode::ThreadSafe> > NavigationRelevantData;
-	TSharedPtr<FNavigationOctree, ESPMode::ThreadSafe> NavOctree; 
+	TWeakObjectPtr<UNavigationSystemV1> NavSystem; 
 	FNavDataConfig NavDataConfig;
 
 #if RECAST_INTERNAL_DEBUG_DATA
diff --git a/Engine/Source/Runtime/NavigationSystem/Public/NavigationDataHandler.h b/Engine/Source/Runtime/NavigationSystem/Public/NavigationDataHandler.h
--- a/Engine/Source/Runtime/NavigationSystem/Public/NavigationDataHandler.h
+++ b/Engine/Source/Runtime/NavigationSystem/Public/NavigationDataHandler.h
@@ -34,4 +34,5 @@ struct NAVIGATIONSYSTEM_API FNavigationDataHandler
 	void RemoveLevelCollisionFromOctree(ULevel& Level);
 	void UpdateActorAndComponentsInNavOctree(AActor& Actor);
 	void ProcessPendingOctreeUpdates();
+	void DemandLazyDataGathering(FNavigationRelevantData& ElementData);
 };
diff --git a/Engine/Source/Runtime/NavigationSystem/Public/NavigationOctree.h b/Engine/Source/Runtime/NavigationSystem/Public/NavigationOctree.h
--- a/Engine/Source/Runtime/NavigationSystem/Public/NavigationOctree.h
+++ b/Engine/Source/Runtime/NavigationSystem/Public/NavigationOctree.h
@@ -152,10 +152,10 @@ public:
 
 	void SetDataGatheringMode(ENavDataGatheringModeConfig Mode);
 	
-	// @hack! TO BE FIXED
-	void DemandLazyDataGathering(const FNavigationOctreeElement& Element);
+	// Lazy data gathering methods
+	bool IsLazyGathering(const INavRelevantInterface& ChildNavInterface) const;
 	void DemandLazyDataGathering(FNavigationRelevantData& ElementData);
-
+	void DemandChildLazyDataGathering(FNavigationRelevantData& ElementData, INavRelevantInterface& ChildNavInterface);
 
 	FORCEINLINE static uint32 HashObject(const UObject& Object)
 	{
diff --git a/Engine/Source/Runtime/NavigationSystem/Public/NavigationSystem.h b/Engine/Source/Runtime/NavigationSystem/Public/NavigationSystem.h
--- a/Engine/Source/Runtime/NavigationSystem/Public/NavigationSystem.h
+++ b/Engine/Source/Runtime/NavigationSystem/Public/NavigationSystem.h
@@ -1089,6 +1089,9 @@ protected:
 
 	virtual void SpawnMissingNavigationData();
 
+public:
+	void DemandLazyDataGathering(FNavigationRelevantData& ElementData);
+
 protected:
 	virtual void RebuildDirtyAreas(float DeltaSeconds);
 
