[NavigationSystem] Extracted methods from SpawnMissingNavigationData to be able to setup navigation data per ULevel.

#rnx

diff --git a/Engine/Source/Runtime/NavigationSystem/Private/NavigationSystem.cpp b/Engine/Source/Runtime/NavigationSystem/Private/NavigationSystem.cpp
--- a/Engine/Source/Runtime/NavigationSystem/Private/NavigationSystem.cpp
+++ b/Engine/Source/Runtime/NavigationSystem/Private/NavigationSystem.cpp
@@ -2202,12 +2202,12 @@ UNavigationSystemV1::ERegistrationResult UNavigationSystemV1::RegisterNavData(AN
 		}
 		else
 		{
-		// check if this kind of agent has already its navigation implemented
-		TWeakObjectPtr<ANavigationData>* NavDataForAgent = AgentToNavDataMap.Find(NavConfig);
-		ANavigationData* NavDataInstanceForAgent = NavDataForAgent ? NavDataForAgent->Get() : nullptr;
+			// check if this kind of agent has already its navigation implemented
+			TWeakObjectPtr<ANavigationData>* NavDataForAgent = AgentToNavDataMap.Find(NavConfig);
+			ANavigationData* NavDataInstanceForAgent = NavDataForAgent ? NavDataForAgent->Get() : nullptr;
 
-		if (NavDataInstanceForAgent == nullptr)
-		{
+			if (NavDataInstanceForAgent == nullptr)
+			{
 				// ok, so this navigation agent doesn't have its navmesh registered yet, but do we want to support it?
 				bool bAgentSupported = false;
 
@@ -2235,19 +2235,19 @@ UNavigationSystemV1::ERegistrationResult UNavigationSystemV1::RegisterNavData(AN
 				}
 				Result = bAgentSupported == true ? RegistrationSuccessful : RegistrationFailed_AgentNotValid;
 			}
-		else if (NavDataInstanceForAgent == NavData)
-		{
-			ensure(NavDataSet.Find(NavData) != INDEX_NONE);
-			// let's treat double registration of the same nav data with the same agent as a success
-			Result = RegistrationSuccessful;
-		}
-		else
-		{
-			// otherwise specified agent type already has its navmesh implemented, fail redundant instance
-			Result = RegistrationFailed_AgentAlreadySupported;
+			else if (NavDataInstanceForAgent == NavData)
+			{
+				ensure(NavDataSet.Find(NavData) != INDEX_NONE);
+				// let's treat double registration of the same nav data with the same agent as a success
+				Result = RegistrationSuccessful;
+			}
+			else
+			{
+				// otherwise specified agent type already has its navmesh implemented, fail redundant instance
+				Result = RegistrationFailed_AgentAlreadySupported;
+			}
 		}
 	}
-	}
 	else
 	{
 		Result = RegistrationFailed_AgentNotValid;
@@ -3366,79 +3366,117 @@ void UNavigationSystemV1::SpawnMissingNavigationData()
 	}
 	
 	// Bit array might be a bit of an overkill here, but this function will be called very rarely
-	TBitArray<> AlreadyInstantiated(false, AllSupportedAgentsCount);
+	TBitArray<> AlreadyInstantiated;
 	uint8 NumberFound = 0;
-	UWorld* NavWorld = GetWorld();
 
 	// 1. check whether any of required navigation data has already been instantiated
-	for (TActorIterator<ANavigationData> It(NavWorld); It && NumberFound < AllSupportedAgentsCount; ++It)
+	NumberFound = FillInstantiatedDataMask(AlreadyInstantiated);
+
+	// 2. for any not already instantiated navigation data call creator functions
+	if (NumberFound < ValidSupportedAgentsCount)
 	{
-		ANavigationData* Nav = (*It);
-		if (Nav != nullptr 
-			&& Nav->IsPendingKill() == false
-			// mz@todo the 'is level in' condition is temporary
-			&& (Nav->GetTypedOuter<UWorld>() == NavWorld || NavWorld->GetLevels().Contains(Nav->GetLevel())))
+		SpawnMissingNavigationDataInLevel(AlreadyInstantiated);
+	}
+
+	if (MainNavData == nullptr || MainNavData->IsPendingKillPending())
+	{
+		MainNavData = GetDefaultNavDataInstance(FNavigationSystem::DontCreate);
+	}
+}
+
+uint8 UNavigationSystemV1::FillInstantiatedDataMask(TBitArray<>& OutInstantiatedMask, ULevel* InLevel /*= nullptr*/)
+{
+	int32 AllSupportedAgentsCount = SupportedAgents.Num();
+	OutInstantiatedMask.Init(false, AllSupportedAgentsCount);
+	uint8 NumberFound = 0;
+
+	auto SetMatchingAgentIndexFunc = [&](ANavigationData* Nav) {
+		for (int32 AgentIndex = 0; AgentIndex < AllSupportedAgentsCount; ++AgentIndex)
+		{
+			if (OutInstantiatedMask[AgentIndex] == false
+				&& Nav->GetClass() == SupportedAgents[AgentIndex].GetNavDataClass<ANavigationData>()
+				&& Nav->DoesSupportAgent(SupportedAgents[AgentIndex]) == true)
+			{
+				OutInstantiatedMask[AgentIndex] = true;
+				++NumberFound;
+				break;
+			}
+		}
+	};
+
+	if (InLevel != nullptr)
+	{
+		for (AActor* Actor: InLevel->Actors)
 		{
-			// find out which one it is
-			for (int32 AgentIndex = 0; AgentIndex < AllSupportedAgentsCount; ++AgentIndex)
+			if (ANavigationData* NavData = Cast<ANavigationData>(Actor))
 			{
-				if (AlreadyInstantiated[AgentIndex] == false
-					&& Nav->GetClass() == SupportedAgents[AgentIndex].GetNavDataClass<ANavigationData>()
-					&& Nav->DoesSupportAgent(SupportedAgents[AgentIndex]) == true)
+				SetMatchingAgentIndexFunc(NavData);
+				if (NumberFound >= AllSupportedAgentsCount)
 				{
-					AlreadyInstantiated[AgentIndex] = true;
-					++NumberFound;
 					break;
 				}
-			}				
+			}
+		}
+	} 
+	else
+	{
+		UWorld* NavWorld = GetWorld();	
+		for (TActorIterator<ANavigationData> It(NavWorld); It && NumberFound < AllSupportedAgentsCount; ++It)
+		{
+			ANavigationData* Nav = (*It);
+			if (IsValid(Nav)
+				// mz@todo the 'is level in' condition is temporary
+				&& (Nav->GetTypedOuter<UWorld>() == NavWorld || NavWorld->GetLevels().Contains(Nav->GetLevel())))
+			{
+				// find out which one it is
+				SetMatchingAgentIndexFunc(Nav);
+			}
 		}
 	}
 
-	// 2. for any not already instantiated navigation data call creator functions
-	if (NumberFound < ValidSupportedAgentsCount)
+	return NumberFound;
+}
+
+void UNavigationSystemV1::SpawnMissingNavigationDataInLevel(const TBitArray<>& InInstantiatedMask, ULevel* InLevel/*=nullptr*/)
+{
+	UWorld* NavWorld = GetWorld();
+
+	ensure(SupportedAgents.Num() == InInstantiatedMask.Num());
+	int32 AllSupportedAgentsCount = InInstantiatedMask.Num();
+
+	for (int32 AgentIndex = 0; AgentIndex < AllSupportedAgentsCount; ++AgentIndex)
 	{
-		for (int32 AgentIndex = 0; AgentIndex < AllSupportedAgentsCount; ++AgentIndex)
+		const FNavDataConfig& NavConfig = SupportedAgents[AgentIndex];
+		if (InInstantiatedMask[AgentIndex] == false
+			&& SupportedAgentsMask.Contains(AgentIndex)
+			&& NavConfig.GetNavDataClass<ANavigationData>() != nullptr)
 		{
-			const FNavDataConfig& NavConfig = SupportedAgents[AgentIndex];
-			if (AlreadyInstantiated[AgentIndex] == false 
-				&& SupportedAgentsMask.Contains(AgentIndex)
-				&& NavConfig.GetNavDataClass<ANavigationData>() != nullptr)
+			const ANavigationData* NavDataCDO = NavConfig.GetNavDataClass<ANavigationData>()->GetDefaultObject<ANavigationData>();
+			if (NavDataCDO == nullptr || !NavDataCDO->CanSpawnOnRebuild())
 			{
-				bool bHandled = false;
-
-				const ANavigationData* NavDataCDO = NavConfig.GetNavDataClass<ANavigationData>()->GetDefaultObject<ANavigationData>();
-				if (NavDataCDO == nullptr || !NavDataCDO->CanSpawnOnRebuild())
-				{
-					continue;
-				}
+				continue;
+			}
 
-				if (NavWorld->WorldType != EWorldType::Editor && NavDataCDO->GetRuntimeGenerationMode() == ERuntimeGenerationType::Static)
-				{
-					// if we're not in the editor, and specified navigation class is configured 
-					// to be static, then we don't want to create an instance					
-					UE_LOG(LogNavigation, Log, TEXT("Not spawning navigation data for %s since indivated NavigationData type is not configured for dynamic generation")
-						, *NavConfig.Name.ToString());
-					continue;
-				}
+			if (NavWorld->WorldType != EWorldType::Editor && NavDataCDO->GetRuntimeGenerationMode() == ERuntimeGenerationType::Static)
+			{
+				// if we're not in the editor, and specified navigation class is configured 
+				// to be static, then we don't want to create an instance					
+				UE_LOG(LogNavigation, Log, TEXT("Not spawning navigation data for %s since indicated NavigationData type is not configured for dynamic generation")
+					, *NavConfig.Name.ToString());
+				continue;
+			}
 
-				ANavigationData* Instance = CreateNavigationDataInstanceInLevel(NavConfig, nullptr);
-				if (Instance)
-				{
-					RequestRegistrationDeferred(*Instance);
-				}
-				else
-				{
-					UE_LOG(LogNavigation, Warning, TEXT("Was not able to create navigation data for SupportedAgent[%d]: %s"), AgentIndex, *NavConfig.Name.ToString());
-				}
+			ANavigationData* Instance = CreateNavigationDataInstanceInLevel(NavConfig, InLevel);
+			if (Instance)
+			{
+				RequestRegistrationDeferred(*Instance);
+			}
+			else
+			{
+				UE_LOG(LogNavigation, Warning, TEXT("Was not able to create navigation data for SupportedAgent[%d]: %s"), AgentIndex, *NavConfig.Name.ToString());
 			}
 		}
 	}
-	
-	if (MainNavData == nullptr || MainNavData->IsPendingKillPending())
-	{
-		// update 
-		MainNavData = GetDefaultNavDataInstance(FNavigationSystem::DontCreate);
-	}
 }
 
 ANavigationData* UNavigationSystemV1::CreateNavigationDataInstanceInLevel(const FNavDataConfig& NavConfig, ULevel* SpawnLevel)
@@ -4880,4 +4918,4 @@ void UNavigationSystemModuleConfig::PostEditChangeProperty(FPropertyChangedEvent
 }
 #endif // WITH_EDITOR
 
-#undef LOCTEXT_NAMESPACE
\ No newline at end of file
+#undef LOCTEXT_NAMESPACE
diff --git a/Engine/Source/Runtime/NavigationSystem/Public/NavigationSystem.h b/Engine/Source/Runtime/NavigationSystem/Public/NavigationSystem.h
--- a/Engine/Source/Runtime/NavigationSystem/Public/NavigationSystem.h
+++ b/Engine/Source/Runtime/NavigationSystem/Public/NavigationSystem.h
@@ -1089,6 +1089,21 @@ protected:
 
 	virtual void SpawnMissingNavigationData();
 
+	/** 
+	 * Fills a mask indicating which navigation data associated to the supported agent mask are already instantiated.
+	 * @param OutInstantiatedMask The mask that will represent already instantiated navigation data.
+	 * @param InLevel If specified will be used to search for navigation data; will use the owning world otherwise.
+	 * @return Number of instantiated navigation data.
+	 */
+	uint8 FillInstantiatedDataMask(TBitArray<>& OutInstantiatedMask, ULevel* InLevel = nullptr);
+
+	/**
+	 * Spawns missing navigation data.
+	 * @param InInstantiatedMask The mask representing already instantiated navigation data.
+	 * @param InLevel Level in which the new data must be added. See CreateNavigationDataInstanceInLevel doc.
+	 */
+	void SpawnMissingNavigationDataInLevel(const TBitArray<>& InInstantiatedMask, ULevel* InLevel = nullptr);
+
 public:
 	void DemandLazyDataGathering(FNavigationRelevantData& ElementData);
 
