[NavSystemConfigOverride] Moved the config setup from PostLoad to PostRegisterAllComponents/PostUnregisterAllComponents for Editor worlds to cleanup a dependency in the WorldSettings. This solves a problem where a level could not be removed and was causing a Fatal Error.

#rnx

diff --git a/Engine/Source/Runtime/Engine/Classes/GameFramework/WorldSettings.h b/Engine/Source/Runtime/Engine/Classes/GameFramework/WorldSettings.h
--- a/Engine/Source/Runtime/Engine/Classes/GameFramework/WorldSettings.h
+++ b/Engine/Source/Runtime/Engine/Classes/GameFramework/WorldSettings.h
@@ -764,8 +764,15 @@ public:
 	 *	no navigation system will be created*/
 	UNavigationSystemConfig* const GetNavigationSystemConfig() const { return NavigationSystemConfigOverride ? NavigationSystemConfigOverride : NavigationSystemConfig; }
 
+	/** 
+	 * Sets a configuration override for NavigationSystem's creation. 
+	 * If set, GetNavigationSystemConfig will return this configuration instead NavigationSystemConfig. 
+	 */
 	void SetNavigationSystemConfigOverride(UNavigationSystemConfig* NewConfig);
 
+	/** @return current configuration override for NavigationSystem's creation, if any. */
+	const UNavigationSystemConfig* GetNavigationSystemConfigOverride() const { return NavigationSystemConfigOverride; }
+
 	/** @return whether given world is configured to host any NavigationSystem */
 	bool IsNavigationSystemEnabled() const;
 
diff --git a/Engine/Source/Runtime/NavigationSystem/Private/NavSystemConfigOverride.cpp b/Engine/Source/Runtime/NavigationSystem/Private/NavSystemConfigOverride.cpp
--- a/Engine/Source/Runtime/NavigationSystem/Private/NavSystemConfigOverride.cpp
+++ b/Engine/Source/Runtime/NavigationSystem/Private/NavSystemConfigOverride.cpp
@@ -60,20 +60,6 @@ ANavSystemConfigOverride::ANavSystemConfigOverride(const FObjectInitializer& Obj
 	bNetLoadOnClient = false;
 }
 
-void ANavSystemConfigOverride::PostLoad()
-{
-	Super::PostLoad();
-	
-#if WITH_EDITOR
-	// ApplyConfig in PostLoad only for Editor worlds; applied in BeginPlay for Game worlds.
-	UWorld* World = GetWorld();
-	if ((World != nullptr) && (World->IsGameWorld() == false))
-	{
-		ApplyConfig();
-	}
-#endif // WITH_EDITOR
-}
-
 void ANavSystemConfigOverride::BeginPlay()
 {
 	Super::BeginPlay();
@@ -158,6 +144,60 @@ void ANavSystemConfigOverride::OverrideNavSystem()
 }
 
 #if WITH_EDITOR
+void ANavSystemConfigOverride::PostRegisterAllComponents()
+{
+	Super::PostRegisterAllComponents();
+
+	// ApplyConfig in PostRegisterAllComponents only for Editor worlds; applied in BeginPlay for Game worlds.
+	UWorld* World = GetWorld();
+	if (World == nullptr || World->IsGameWorld())
+	{
+		return;
+	}
+
+	// While refreshing streaming levels, components are unregistered then registered again
+	// and we don't want to modify the configuration and recreate the NavigationSystem.
+	if (World->IsRefreshingStreamingLevels())
+	{
+		return;
+	}
+
+	ApplyConfig();
+}
+
+void ANavSystemConfigOverride::PostUnregisterAllComponents()
+{
+	Super::PostUnregisterAllComponents();
+
+	if (NavigationSystemConfig == nullptr)
+	{
+		return;
+	}
+
+	// ApplyConfig was performed in PostRegisterAllComponents for Editor worlds so nothing to unregister for Game worlds.
+	UWorld* World = GetWorld();
+	if (World == nullptr || World->IsGameWorld())
+	{
+		return;
+	}
+
+	// While refreshing streaming levels, components are unregistered then registered again
+	// and we don't want to modify the configuration and recreate the NavigationSystem.
+	if (World->IsRefreshingStreamingLevels())
+	{
+		return;
+	}
+
+	// If our override was used to create the navigation system, we remove the dependency
+	// and recreate the navigation system (if needed after removing the override)
+	AWorldSettings* WorldSetting = World->GetWorldSettings();
+	if (WorldSetting && WorldSetting->GetNavigationSystemConfigOverride() == NavigationSystemConfig)
+	{
+		WorldSetting->SetNavigationSystemConfigOverride(nullptr);
+		FNavigationSystem::AddNavigationSystemToWorld(*World, FNavigationSystemRunMode::EditorMode, nullptr, /*bInitializeForWorld=*/true, /*bOverridePreviousNavSys=*/true);
+	}
+}
+
 void ANavSystemConfigOverride::InitializeForWorld(UNavigationSystemBase* NewNavSys, UWorld* World, const FNavigationSystemRunMode RunMode)
 {
 	if (NewNavSys && World)
diff --git a/Engine/Source/Runtime/NavigationSystem/Public/NavSystemConfigOverride.h b/Engine/Source/Runtime/NavigationSystem/Public/NavSystemConfigOverride.h
--- a/Engine/Source/Runtime/NavigationSystem/Public/NavSystemConfigOverride.h
+++ b/Engine/Source/Runtime/NavigationSystem/Public/NavSystemConfigOverride.h
@@ -46,7 +46,6 @@ public:
 	ANavSystemConfigOverride(const FObjectInitializer& ObjectInitializer = FObjectInitializer::Get());
 
 	//~ Begin UObject Interface
-	virtual void PostLoad() override;
 	virtual void PostInitProperties() override;
 #if WITH_EDITOR
 	virtual void PostEditChangeProperty(FPropertyChangedEvent& PropertyChangedEvent) override;
@@ -55,6 +54,10 @@ public:
 
 	//~ Begin AActor Interface
 	virtual void BeginPlay() override;
+#if WITH_EDITOR
+	virtual void PostRegisterAllComponents() override;
+	virtual void PostUnregisterAllComponents() override;
+#endif
 	//~ End AActor Interface
 
 #if WITH_EDITOR
