RecastNavMesh -> Remove deprecated AgentMaxHeight
#rnx

diff --git a/Engine/Source/Runtime/NavigationSystem/Public/NavMesh/RecastNavMesh.h b/Engine/Source/Runtime/NavigationSystem/Public/NavMesh/RecastNavMesh.h
--- a/Engine/Source/Runtime/NavigationSystem/Public/NavMesh/RecastNavMesh.h
+++ b/Engine/Source/Runtime/NavigationSystem/Public/NavMesh/RecastNavMesh.h
@@ -491,10 +491,6 @@ class NAVIGATIONSYSTEM_API ARecastNavMesh : public ANavigationData
 	UPROPERTY(EditAnywhere, Category = Generation, config, meta = (ClampMin = "0.0"))
 	float AgentHeight;
 
-	UE_DEPRECATED(4.24, "ARecastNavMesh.AgentMaxHeight has been deprecated as it has no use. Use AgentHeight instead.")
-	UPROPERTY(VisibleAnywhere, Category=Generation, config, meta=(ClampMin = "0.0", DisplayName="DEPRECATED_AgentMaxHeight"))
-	float AgentMaxHeight;
-
 	/* The maximum slope (angle) that the agent can move on. */ 
 	UPROPERTY(EditAnywhere, Category=Generation, config, meta=(ClampMin = "0.0", ClampMax = "89.0", UIMin = "0.0", UIMax = "89.0" ))
 	float AgentMaxSlope;
