Adding FindMoveAlongSurface to find a nearby position on navmesh relative to another (faster than doing nav raycast)

#rnx

diff --git a/Engine/Source/Runtime/Engine/Classes/AI/Navigation/NavigationDataInterface.h b/Engine/Source/Runtime/Engine/Classes/AI/Navigation/NavigationDataInterface.h
--- a/Engine/Source/Runtime/Engine/Classes/AI/Navigation/NavigationDataInterface.h
+++ b/Engine/Source/Runtime/Engine/Classes/AI/Navigation/NavigationDataInterface.h
@@ -24,6 +24,12 @@ class INavigationDataInterface
 {
 	GENERATED_IINTERFACE_BODY()
 public:
+	/**	Tries to move current nav location towards target constrained to navigable area.
+	 *	@param OutLocation if successful this variable will be filed with result
+	 *	@return true if successful, false otherwise
+	 */
+	virtual bool FindMoveAlongSurface(const FNavLocation& StartLocation, const FVector& TargetPosition, FNavLocation& OutLocation, FSharedConstNavQueryFilter Filter = nullptr, const UObject* Querier = nullptr) const PURE_VIRTUAL(INavigationDataInterface::FindMoveAlongSurface, return false;);
+
 	/**	Tries to project given Point to this navigation type, within given Extent.
 	*	@param OutLocation if successful this variable will be filed with result
 	*	@return true if successful, false otherwise
diff --git a/Engine/Source/Runtime/NavigationSystem/Private/NavMesh/PImplRecastNavMesh.cpp b/Engine/Source/Runtime/NavigationSystem/Private/NavMesh/PImplRecastNavMesh.cpp
--- a/Engine/Source/Runtime/NavigationSystem/Private/NavMesh/PImplRecastNavMesh.cpp
+++ b/Engine/Source/Runtime/NavigationSystem/Private/NavMesh/PImplRecastNavMesh.cpp
@@ -1389,6 +1389,53 @@ bool FPImplRecastNavMesh::GetRandomPointInCluster(NavNodeRef ClusterRef, FNavLoc
 	return false;
 }
 
+bool FPImplRecastNavMesh::FindMoveAlongSurface(const FNavLocation& StartLocation, const FVector& TargetPosition, FNavLocation& OutLocation, const FNavigationQueryFilter& Filter, const UObject* Owner) const
+{
+	// sanity check
+	if (DetourNavMesh == NULL)
+	{
+		return false;
+	}
+
+	FRecastSpeciaLinkFilter LinkFilter(FNavigationSystem::GetCurrent<UNavigationSystemV1>(NavMeshOwner->GetWorld()), Owner);
+	// using 0 as NumNodes since findNearestPoly2D, being the only dtNavMeshQuery
+	// function we're using, is not utilizing m_nodePool
+	INITIALIZE_NAVQUERY(NavQuery, /*NumNodes=*/0, LinkFilter);
+
+	const dtQueryFilter* QueryFilter = ((const FRecastQueryFilter*)(Filter.GetImplementation()))->GetAsDetourQueryFilter();
+	ensure(QueryFilter);
+	if (!QueryFilter)
+	{
+		return false;
+	}
+
+	FVector RcStartPos = Unreal2RecastPoint(StartLocation.Location);
+	FVector RcEndPos = Unreal2RecastPoint(TargetPosition);
+
+	float Result[3];
+	static const int MAX_VISITED = 16;
+	dtPolyRef Visited[MAX_VISITED];
+	int VisitedCount = 0;
+
+	dtStatus status = NavQuery.moveAlongSurface(StartLocation.NodeRef, &RcStartPos.X, &RcEndPos.X, QueryFilter, Result, Visited, &VisitedCount, MAX_VISITED);
+	if (dtStatusFailed(status))
+	{
+		return false;
+	}
+	dtPolyRef ResultPoly = Visited[VisitedCount - 1];
+
+	// Adjust the position to stay on top of the navmesh.
+	float h = RcStartPos.Y;
+	NavQuery.getPolyHeight(ResultPoly, Result, &h);
+	Result[1] = h;
+
+	const FVector UnrealResult = Recast2UnrVector(Result);
+
+	OutLocation = FNavLocation(UnrealResult, ResultPoly);
+
+	return true;
+}
+
 bool FPImplRecastNavMesh::ProjectPointToNavMesh(const FVector& Point, FNavLocation& Result, const FVector& Extent, const FNavigationQueryFilter& Filter, const UObject* Owner) const
 {
 	// sanity check
diff --git a/Engine/Source/Runtime/NavigationSystem/Private/NavMesh/RecastNavMesh.cpp b/Engine/Source/Runtime/NavigationSystem/Private/NavMesh/RecastNavMesh.cpp
--- a/Engine/Source/Runtime/NavigationSystem/Private/NavMesh/RecastNavMesh.cpp
+++ b/Engine/Source/Runtime/NavigationSystem/Private/NavMesh/RecastNavMesh.cpp
@@ -1107,6 +1107,17 @@ NavNodeRef ARecastNavMesh::GetClusterRef(NavNodeRef PolyRef) const
 	return ClusterRef;
 }
 
+bool ARecastNavMesh::FindMoveAlongSurface(const FNavLocation& StartLocation, const FVector& TargetPosition, FNavLocation& OutLocation, FSharedConstNavQueryFilter Filter, const UObject* QueryOwner) const
+{
+	bool bSuccess = false;
+	if (RecastNavMeshImpl)
+	{
+		bSuccess = RecastNavMeshImpl->FindMoveAlongSurface(StartLocation, TargetPosition, OutLocation, GetRightFilterRef(Filter), QueryOwner);
+	}
+
+	return bSuccess;
+}
+
 bool ARecastNavMesh::ProjectPoint(const FVector& Point, FNavLocation& OutLocation, const FVector& Extent, FSharedConstNavQueryFilter Filter, const UObject* QueryOwner) const
 {
 	bool bSuccess = false;
diff --git a/Engine/Source/Runtime/NavigationSystem/Public/NavMesh/PImplRecastNavMesh.h b/Engine/Source/Runtime/NavigationSystem/Public/NavMesh/PImplRecastNavMesh.h
--- a/Engine/Source/Runtime/NavigationSystem/Public/NavMesh/PImplRecastNavMesh.h
+++ b/Engine/Source/Runtime/NavigationSystem/Public/NavMesh/PImplRecastNavMesh.h
@@ -101,6 +101,12 @@ public:
 	/** Returns a random location on the navmesh within cluster */
 	bool GetRandomPointInCluster(NavNodeRef ClusterRef, FNavLocation& OutLocation) const;
 
+	/**	Tries to move current nav location towards target constrained to navigable area. Faster than ProjectPointToNavmesh.
+	 *	@param OutLocation if successful this variable will be filed with result
+	 *	@return true if successful, false otherwise
+	 */
+	bool FindMoveAlongSurface(const FNavLocation& StartLocation, const FVector& TargetPosition, FNavLocation& OutLocation, const FNavigationQueryFilter& Filter, const UObject* Owner) const;
+
 	bool ProjectPointToNavMesh(const FVector& Point, FNavLocation& Result, const FVector& Extent, const FNavigationQueryFilter& Filter, const UObject* Owner) const;
 	
 	/** Project single point and grab all vertical intersections */
diff --git a/Engine/Source/Runtime/NavigationSystem/Public/NavMesh/RecastNavMesh.h b/Engine/Source/Runtime/NavigationSystem/Public/NavMesh/RecastNavMesh.h
--- a/Engine/Source/Runtime/NavigationSystem/Public/NavMesh/RecastNavMesh.h
+++ b/Engine/Source/Runtime/NavigationSystem/Public/NavMesh/RecastNavMesh.h
@@ -711,6 +711,7 @@ public:
 	virtual bool GetRandomReachablePointInRadius(const FVector& Origin, float Radius, FNavLocation& OutResult, FSharedConstNavQueryFilter Filter = NULL, const UObject* Querier = NULL) const override;
 	virtual bool GetRandomPointInNavigableRadius(const FVector& Origin, float Radius, FNavLocation& OutResult, FSharedConstNavQueryFilter Filter = NULL, const UObject* Querier = NULL) const override;
 
+	virtual bool FindMoveAlongSurface(const FNavLocation& StartLocation, const FVector& TargetPosition, FNavLocation& OutLocation, FSharedConstNavQueryFilter Filter = NULL, const UObject* Querier = NULL) const override;
 	virtual bool ProjectPoint(const FVector& Point, FNavLocation& OutLocation, const FVector& Extent, FSharedConstNavQueryFilter Filter = NULL, const UObject* Querier = NULL) const override;
 	virtual bool IsNodeRefValid(NavNodeRef NodeRef) const override;
 
diff --git a/Engine/Source/Runtime/NavigationSystem/Public/NavigationData.h b/Engine/Source/Runtime/NavigationSystem/Public/NavigationData.h
--- a/Engine/Source/Runtime/NavigationSystem/Public/NavigationData.h
+++ b/Engine/Source/Runtime/NavigationSystem/Public/NavigationData.h
@@ -808,6 +808,12 @@ public:
 	/** Raycasts batched for efficiency */
 	virtual void BatchRaycast(TArray<FNavigationRaycastWork>& Workload, FSharedConstNavQueryFilter QueryFilter, const UObject* Querier = NULL) const PURE_VIRTUAL(ANavigationData::BatchRaycast, );
 
+	/**	Tries to move current nav location towards target constrained to navigable area. Faster than ProjectPointToNavmesh.
+	 *	@param OutLocation if successful this variable will be filed with result
+	 *	@return true if successful, false otherwise
+	 */
+	virtual bool FindMoveAlongSurface(const FNavLocation& StartLocation, const FVector& TargetPosition, FNavLocation& OutLocation, FSharedConstNavQueryFilter Filter = NULL, const UObject* Querier = NULL) const PURE_VIRTUAL(ANavigationData::FindMoveAlongSurface, return false;);
+
 	virtual FNavLocation GetRandomPoint(FSharedConstNavQueryFilter Filter = NULL, const UObject* Querier = NULL) const PURE_VIRTUAL(ANavigationData::GetRandomPoint, return FNavLocation(););
 
 	/** finds a random location in Radius, reachable from Origin */
