Remove unused fields - AI

diff --git a/Engine/Plugins/AI/EnvironmentQueryEditor/Source/EnvironmentQueryEditor/Private/EnvironmentQueryEditor.h b/Engine/Plugins/AI/EnvironmentQueryEditor/Source/EnvironmentQueryEditor/Private/EnvironmentQueryEditor.h
--- a/Engine/Plugins/AI/EnvironmentQueryEditor/Source/EnvironmentQueryEditor/Private/EnvironmentQueryEditor.h
+++ b/Engine/Plugins/AI/EnvironmentQueryEditor/Source/EnvironmentQueryEditor/Private/EnvironmentQueryEditor.h
@@ -65,8 +65,6 @@ private:
 	TSharedPtr<class IDetailsView> DetailsView;
 	TSharedPtr<SEnvQueryProfiler> ProfilerView;
 
-	uint32 SelectedNodesCount;
-
 	/**	Graph editor tab */
 	static const FName EQSUpdateGraphTabId;
 	static const FName EQSPropertiesTabId;
diff --git a/Engine/Source/Runtime/NavigationSystem/Private/NavGraph/NavGraphGenerator.h b/Engine/Source/Runtime/NavigationSystem/Private/NavGraph/NavGraphGenerator.h
--- a/Engine/Source/Runtime/NavigationSystem/Private/NavGraph/NavGraphGenerator.h
+++ b/Engine/Source/Runtime/NavigationSystem/Private/NavGraph/NavGraphGenerator.h
@@ -33,8 +33,4 @@ private:
 	TArray<AVolume const*> InclusionVolumes;
 
 	FCriticalSection GraphChangingLock;
-
-	ANavigationGraph* DestNavGraph;
-
-	uint32 bInitialized:1;
 };
diff --git a/Engine/Source/Runtime/NavigationSystem/Private/NavMesh/NavMeshRenderingComponent.cpp b/Engine/Source/Runtime/NavigationSystem/Private/NavMesh/NavMeshRenderingComponent.cpp
--- a/Engine/Source/Runtime/NavigationSystem/Private/NavMesh/NavMeshRenderingComponent.cpp
+++ b/Engine/Source/Runtime/NavigationSystem/Private/NavMesh/NavMeshRenderingComponent.cpp
@@ -1026,7 +1026,6 @@ SIZE_T FNavMeshSceneProxy::GetTypeHash() const
 FNavMeshSceneProxy::FNavMeshSceneProxy(const UPrimitiveComponent* InComponent, FNavMeshSceneProxyData* InProxyData, bool ForceToRender)
 	: FDebugRenderSceneProxy(InComponent)
 	, VertexFactory(GetScene().GetFeatureLevel(), "FNavMeshSceneProxy")
-	, bRequestedData(false)
 	, bForceRendering(ForceToRender)
 {
 	DrawType = EDrawType::SolidAndWireMeshes;
diff --git a/Engine/Source/Runtime/NavigationSystem/Public/NavMesh/NavMeshRenderingComponent.h b/Engine/Source/Runtime/NavigationSystem/Public/NavMesh/NavMeshRenderingComponent.h
--- a/Engine/Source/Runtime/NavigationSystem/Public/NavMesh/NavMeshRenderingComponent.h
+++ b/Engine/Source/Runtime/NavigationSystem/Public/NavMesh/NavMeshRenderingComponent.h
@@ -143,7 +143,6 @@ private:
 	FDebugDrawDelegate DebugTextDrawingDelegate;
 	FDelegateHandle DebugTextDrawingDelegateHandle;
 	TWeakObjectPtr<UNavMeshRenderingComponent> RenderingComponent;
-	uint32 bRequestedData : 1;
 	uint32 bForceRendering : 1;
 	uint32 bSkipDistanceCheck : 1;
 	uint32 bUseThickLines : 1;
diff --git a/Engine/Source/Runtime/Navmesh/Private/DetourCrowd/DetourProximityGrid.cpp b/Engine/Source/Runtime/Navmesh/Private/DetourCrowd/DetourProximityGrid.cpp
--- a/Engine/Source/Runtime/Navmesh/Private/DetourCrowd/DetourProximityGrid.cpp
+++ b/Engine/Source/Runtime/Navmesh/Private/DetourCrowd/DetourProximityGrid.cpp
@@ -47,7 +47,6 @@ inline int hashPos2(int x, int y, int n)
 
 
 dtProximityGrid::dtProximityGrid() :
-	m_maxItems(0),
 	m_cellSize(0),
 	m_pool(0),
 	m_poolHead(0),
diff --git a/Engine/Source/Runtime/Navmesh/Public/DetourCrowd/DetourProximityGrid.h b/Engine/Source/Runtime/Navmesh/Public/DetourCrowd/DetourProximityGrid.h
--- a/Engine/Source/Runtime/Navmesh/Public/DetourCrowd/DetourProximityGrid.h
+++ b/Engine/Source/Runtime/Navmesh/Public/DetourCrowd/DetourProximityGrid.h
@@ -26,7 +26,6 @@
 
 class dtProximityGrid
 {
-	int m_maxItems;
 	float m_cellSize;
 	float m_invCellSize;
 	
