Relaxed radius condition in RecastNavMesh::GetRandomReachablePointInRadius so that 'radius == 0' can produce a relevant result.

diff --git a/Engine/Source/Runtime/NavigationSystem/Private/NavMesh/RecastNavMesh.cpp b/Engine/Source/Runtime/NavigationSystem/Private/NavMesh/RecastNavMesh.cpp
--- a/Engine/Source/Runtime/NavigationSystem/Private/NavMesh/RecastNavMesh.cpp
+++ b/Engine/Source/Runtime/NavigationSystem/Private/NavMesh/RecastNavMesh.cpp
@@ -1045,7 +1045,7 @@ FNavLocation ARecastNavMesh::GetRandomPoint(FSharedConstNavQueryFilter Filter, c
 
 bool ARecastNavMesh::GetRandomReachablePointInRadius(const FVector& Origin, float Radius, FNavLocation& OutResult, FSharedConstNavQueryFilter Filter, const UObject* QueryOwner) const
 {
-	if (RecastNavMeshImpl == nullptr || RecastNavMeshImpl->DetourNavMesh == nullptr || Radius <= 0.f)
+	if (RecastNavMeshImpl == nullptr || RecastNavMeshImpl->DetourNavMesh == nullptr || Radius < 0.f)
 	{
 		return false;
 	}
diff --git a/Engine/Source/Runtime/NavigationSystem/Public/NavMesh/RecastNavMesh.h b/Engine/Source/Runtime/NavigationSystem/Public/NavMesh/RecastNavMesh.h
--- a/Engine/Source/Runtime/NavigationSystem/Public/NavMesh/RecastNavMesh.h
+++ b/Engine/Source/Runtime/NavigationSystem/Public/NavMesh/RecastNavMesh.h
@@ -733,6 +733,8 @@ public:
 
 	//~ Begin ANavigationData Interface
 	virtual FNavLocation GetRandomPoint(FSharedConstNavQueryFilter Filter = NULL, const UObject* Querier = NULL) const override;
+	/** finds a random location in Radius, reachable from Origin.
+	 *  @param Radius needs to be non-negative. The function fails for Radius < 0. Radius being 0 is still rasults in a valid request. */
 	virtual bool GetRandomReachablePointInRadius(const FVector& Origin, float Radius, FNavLocation& OutResult, FSharedConstNavQueryFilter Filter = NULL, const UObject* Querier = NULL) const override;
 	virtual bool GetRandomPointInNavigableRadius(const FVector& Origin, float Radius, FNavLocation& OutResult, FSharedConstNavQueryFilter Filter = NULL, const UObject* Querier = NULL) const override;
 
