Pixing PVS 7.7 Issues:

warning V614: Potentially uninitialized pointer 'Tiles' used.

diff --git a/Engine/Source/Runtime/NavigationSystem/Private/NavMesh/RecastNavMeshGenerator.cpp b/Engine/Source/Runtime/NavigationSystem/Private/NavMesh/RecastNavMeshGenerator.cpp
--- a/Engine/Source/Runtime/NavigationSystem/Private/NavMesh/RecastNavMeshGenerator.cpp
+++ b/Engine/Source/Runtime/NavigationSystem/Private/NavMesh/RecastNavMeshGenerator.cpp
@@ -244,35 +244,36 @@ FRecastVoxelCache::FRecastVoxelCache(const uint8* Memory)
 	if (Memory)
 	{
 		NumTiles = *((int32*)BytesArr);	BytesArr += sizeof(int32);
-		Tiles = (FTileInfo*)BytesArr;
-	}
-	else
-	{
-		NumTiles = 0;
-	}
-
-	FTileInfo* iTile = Tiles;	
-	for (int i = 0; i < NumTiles; i++)
-	{
-		iTile = (FTileInfo*)BytesArr; BytesArr += sizeof(FTileInfo);
-		if (iTile->NumSpans)
+		if (NumTiles > 0)
 		{
-			iTile->SpanData = (rcSpanCache*)BytesArr; BytesArr += sizeof(rcSpanCache) * iTile->NumSpans;
+			Tiles = (FTileInfo*)BytesArr;
+
+			FTileInfo* iTile = Tiles;
+			for (int i = 0; i < NumTiles; i++)
+			{
+				iTile = (FTileInfo*)BytesArr; BytesArr += sizeof(FTileInfo);
+				if (iTile->NumSpans)
+				{
+					iTile->SpanData = (rcSpanCache*)BytesArr; BytesArr += sizeof(rcSpanCache) * iTile->NumSpans;
+				}
+				else
+				{
+					iTile->SpanData = 0;
+				}
+
+				iTile->NextTile = (FTileInfo*)BytesArr;
+			}
+
+			iTile->NextTile = 0;
 		}
 		else
 		{
-			iTile->SpanData = 0;
+			Tiles = 0;
 		}
-
-		iTile->NextTile = (FTileInfo*)BytesArr;
-	}
-
-	if (NumTiles > 0)
-	{
-		iTile->NextTile = 0;
 	}
 	else
 	{
+		NumTiles = 0;
 		Tiles = 0;
 	}
 }
