Add missing DDC debug context in NavCollision
Use the Super default implementation instead of copy pasting the same string everywhere

diff --git a/Engine/Source/Runtime/Engine/Private/PhysicsEngine/Experimental/ChaosDerivedData.cpp b/Engine/Source/Runtime/Engine/Private/PhysicsEngine/Experimental/ChaosDerivedData.cpp
--- a/Engine/Source/Runtime/Engine/Private/PhysicsEngine/Experimental/ChaosDerivedData.cpp
+++ b/Engine/Source/Runtime/Engine/Private/PhysicsEngine/Experimental/ChaosDerivedData.cpp
@@ -70,7 +70,7 @@ FString FChaosDerivedDataCooker::GetDebugContextString() const
 		}
 	}
 
-	return TEXT("Unknown Context");
+	return FDerivedDataPluginInterface::GetDebugContextString();
 }
 
 FString FChaosDerivedDataCooker::GetPluginSpecificCacheKeySuffix() const
diff --git a/Engine/Source/Runtime/NavigationSystem/Private/NavCollision.cpp b/Engine/Source/Runtime/NavigationSystem/Private/NavCollision.cpp
--- a/Engine/Source/Runtime/NavigationSystem/Private/NavCollision.cpp
+++ b/Engine/Source/Runtime/NavigationSystem/Private/NavCollision.cpp
@@ -115,6 +115,7 @@ public:
 	}
 
 	virtual bool Build( TArray<uint8>& OutData ) override;
+	virtual FString GetDebugContextString() const override;
 
 	/** Return true if we can build **/
 	bool CanBuild()
@@ -161,6 +162,20 @@ bool FDerivedDataNavCollisionCooker::Build( TArray<uint8>& OutData )
 	return true;
 }
 
+FString FDerivedDataNavCollisionCooker::GetDebugContextString() const
+{
+	if (NavCollisionInstance)
+	{
+		UObject* Outer = NavCollisionInstance->GetOuter();
+		if (Outer)
+		{
+			return Outer->GetFullName();
+		}
+	}
+
+	return FDerivedDataPluginInterface::GetDebugContextString();
+}
+
 namespace
 {
 	UNavCollisionBase* CreateNewNavCollisionInstance(UObject& Outer)
