[Navigation] Addition of ARecastNavMesh::FindEdges() method to collect navmesh wall edges

diff --git a/Engine/Source/Runtime/NavigationSystem/Private/NavMesh/RecastNavMesh.cpp b/Engine/Source/Runtime/NavigationSystem/Private/NavMesh/RecastNavMesh.cpp
--- a/Engine/Source/Runtime/NavigationSystem/Private/NavMesh/RecastNavMesh.cpp
+++ b/Engine/Source/Runtime/NavigationSystem/Private/NavMesh/RecastNavMesh.cpp
@@ -1347,6 +1347,43 @@ bool ARecastNavMesh::GetPolysInBox(const FBox& Box, TArray<FNavPoly>& Polys, FSh
 	return bSuccess;
 }
 
+bool ARecastNavMesh::FindEdges(const NavNodeRef CenterNodeRef, const FVector Center, const float Radius, const FSharedConstNavQueryFilter Filter, TArray<FNavigationWallEdge>& OutEdges) const
+{
+	const FNavigationQueryFilter& FilterToUse = GetRightFilterRef(Filter);
+	INITIALIZE_NAVQUERY(NavQuery, FilterToUse.GetMaxSearchNodes());
+	const dtQueryFilter* QueryFilter = static_cast<const FRecastQueryFilter*>(FilterToUse.GetImplementation())->GetAsDetourQueryFilter();
+
+	const int32 MaxWalls = 64;
+	int32 NumWalls = 0;
+	float WallSegments[MaxWalls * 3 * 2] = { 0 };
+	dtPolyRef WallPolys[MaxWalls * 2] = { 0 };
+
+	const int32 MaxNeis = 64;
+	int32 NumNeis = 0;
+	dtPolyRef NeiPolys[MaxNeis] = { 0 };
+
+	const FVector RcCenter = Unreal2RecastPoint(Center);
+
+	dtStatus Status = NavQuery.findWallsInNeighbourhood(CenterNodeRef, &RcCenter.X, Radius, QueryFilter,
+		NeiPolys, &NumNeis, MaxNeis, WallSegments, WallPolys, &NumWalls, MaxWalls);
+
+	if (dtStatusSucceed(Status))
+	{
+		OutEdges.Reset(NumWalls);
+		FNavigationWallEdge NewEdge;
+		for (int32 Idx = 0; Idx < NumWalls; Idx++)
+		{
+			NewEdge.Start = Recast2UnrealPoint(&WallSegments[Idx * 6]);
+			NewEdge.End = Recast2UnrealPoint(&WallSegments[Idx * 6 + 3]);
+			OutEdges.Add(NewEdge);
+		}
+
+		return true;
+	}
+
+	return false;
+}
+
 bool ARecastNavMesh::ProjectPointMulti(const FVector& Point, TArray<FNavLocation>& OutLocations, const FVector& Extent,
 	float MinZ, float MaxZ, FSharedConstNavQueryFilter Filter, const UObject* QueryOwner) const
 {
diff --git a/Engine/Source/Runtime/NavigationSystem/Public/NavMesh/RecastNavMesh.h b/Engine/Source/Runtime/NavigationSystem/Public/NavMesh/RecastNavMesh.h
--- a/Engine/Source/Runtime/NavigationSystem/Public/NavMesh/RecastNavMesh.h
+++ b/Engine/Source/Runtime/NavigationSystem/Public/NavMesh/RecastNavMesh.h
@@ -218,6 +218,12 @@ namespace ERecastNamedFilter
 		NamedFiltersCount,
 	};
 }
+
+struct FNavigationWallEdge
+{
+	FVector Start = FVector::ZeroVector;
+	FVector End = FVector::ZeroVector;
+};
 #endif //WITH_RECAST
 
 
@@ -1013,9 +1019,12 @@ public:
 	/** Get all polys from tile */
 	bool GetPolysInTile(int32 TileIndex, TArray<FNavPoly>& Polys) const;
 
-	/** Get all polys that overlap the specified box */
+	/** Get up to 256 polys that overlap the specified box */
 	bool GetPolysInBox(const FBox& Box, TArray<FNavPoly>& Polys, FSharedConstNavQueryFilter Filter = nullptr, const UObject* Owner = nullptr) const;
 
+	/** Find up to 64 navmesh eges in up to 64 polys around the center */
+	bool FindEdges(const NavNodeRef CenterNodeRef, const FVector Center, const float Radius, const FSharedConstNavQueryFilter Filter, TArray<FNavigationWallEdge>& OutEdges) const;
+
 	/** Get all polys from tile */
 	bool GetNavLinksInTile(const int32 TileIndex, TArray<FNavPoly>& Polys, const bool bIncludeLinksFromNeighborTiles) const;
 
