[Navigation Navmesh] Switch to using bSortNavigationAreasByCost enabled by default

diff --git a/Engine/Source/Runtime/NavigationSystem/Private/NavMesh/RecastNavMesh.cpp b/Engine/Source/Runtime/NavigationSystem/Private/NavMesh/RecastNavMesh.cpp
--- a/Engine/Source/Runtime/NavigationSystem/Private/NavMesh/RecastNavMesh.cpp
+++ b/Engine/Source/Runtime/NavigationSystem/Private/NavMesh/RecastNavMesh.cpp
@@ -254,7 +254,7 @@ FRecastNavMeshGenerationProperties::FRecastNavMeshGenerationProperties()
 	LayerPartitioning = ERecastPartitioning::Watershed;
 	RegionChunkSplits = 2;
 	LayerChunkSplits = 2;
-	bSortNavigationAreasByCost = false;
+	bSortNavigationAreasByCost = true;
 	bPerformVoxelFiltering = true;
 	bMarkLowHeightAreas = false;
 	bUseExtraTopCellWhenMarkingAreas = true;
@@ -305,6 +305,7 @@ ARecastNavMesh::ARecastNavMesh(const FObjectInitializer& ObjectInitializer)
 	, MaxSimplificationError(1.3f)	// from RecastDemo
 	, DefaultMaxSearchNodes(RECAST_MAX_SEARCH_NODES)
 	, DefaultMaxHierarchicalSearchNodes(RECAST_MAX_SEARCH_NODES)
+	, bSortNavigationAreasByCost(true)
 	, bPerformVoxelFiltering(true)	
 	, bMarkLowHeightAreas(false)
 	, bUseExtraTopCellWhenMarkingAreas(true)
