Fix typo in dev doc

diff --git a/Engine/Source/Runtime/NavigationSystem/DevDocs/How To Optimize Navmesh Generation.md b/Engine/Source/Runtime/NavigationSystem/DevDocs/How To Optimize Navmesh Generation.md
--- a/Engine/Source/Runtime/NavigationSystem/DevDocs/How To Optimize Navmesh Generation.md	
+++ b/Engine/Source/Runtime/NavigationSystem/DevDocs/How To Optimize Navmesh Generation.md	
@@ -2,7 +2,7 @@
 Using a dynamic navmesh and generating navmesh at runtime can take a lot of CPU resources. Here is a list of tips and strategies that can limit that cost.
 
 ## Tips
-1. Use the highest `CellHeight` and `CellSize` possible (in Project Settings ->Navigation Settings)
+1. Use the highest `CellHeight` and `CellSize` possible (in Project Settings ->Navigation Mesh)
     - Those parameters are used to define the size of the voxels used to generate the navmesh. The bigger they are, the less voxels are needed and the faster the tile generation runs. Be aware that increasing those will reduce the precision of the navmesh (how well it will fit the geometry shapes).
 
 
