PR #8362: Prevent compiler error when using WITH_RECAST (Contributed by [REDACTED_PERSON])

diff --git a/Engine/Source/Runtime/NavigationSystem/Public/NavMesh/RecastNavMesh.h b/Engine/Source/Runtime/NavigationSystem/Public/NavMesh/RecastNavMesh.h
--- a/Engine/Source/Runtime/NavigationSystem/Public/NavMesh/RecastNavMesh.h
+++ b/Engine/Source/Runtime/NavigationSystem/Public/NavMesh/RecastNavMesh.h
@@ -1139,6 +1139,11 @@ private:
 	static const FRecastQueryFilter* NamedFilters[ERecastNamedFilter::NamedFiltersCount];
 #endif // WITH_RECAST
 
+public:
+	//----------------------------------------------------------------------//
+	// Blueprint functions
+	//----------------------------------------------------------------------//
+
 	/** @return true if any polygon/link has been touched */
 	UFUNCTION(BlueprintCallable, Category = NavMesh, meta = (DisplayName = "ReplaceAreaInTileBounds"))
 	bool K2_ReplaceAreaInTileBounds(FBox Bounds, TSubclassOf<UNavArea> OldArea, TSubclassOf<UNavArea> NewArea, bool ReplaceLinks = true);
