[Navmesh] Add comment on borderSize computation

diff --git a/Engine/Source/Runtime/NavigationSystem/Private/NavMesh/RecastNavMeshGenerator.cpp b/Engine/Source/Runtime/NavigationSystem/Private/NavMesh/RecastNavMeshGenerator.cpp
--- a/Engine/Source/Runtime/NavigationSystem/Private/NavMesh/RecastNavMeshGenerator.cpp
+++ b/Engine/Source/Runtime/NavigationSystem/Private/NavMesh/RecastNavMeshGenerator.cpp
@@ -4502,7 +4502,7 @@ void FRecastNavMeshGenerator::ConfigureBuildProperties(FRecastBuildConfig& OutCo
 	OutConfig.AgentMaxClimb = AgentMaxClimb;
 	OutConfig.AgentRadius = AgentRadius;
 
-	OutConfig.borderSize = WalkableRadius + 3;
+	OutConfig.borderSize = WalkableRadius + 3; // +1 for voxelization rounding, +1 for ledge neighbor access, +1 for occasional errors
 	OutConfig.maxEdgeLen = (int32)(1200.0f / CellSize);
 	OutConfig.maxSimplificationError = 1.3f;
 	// hardcoded, but can be overridden by RecastNavMesh params later
