[DebugDrawDelegate]
Update registration flow since CreateRenderState_Concurrent can't be used to register our delegate since it needs to be initialized first (CreateSceneProxy) and in some code paths for loaded actors the call order might be different since primitive registration gets deferred (i.e. FRegisterComponentContext != nullptr).
In that case initialization won't called and an ensure will fire in UnregisterDebugDrawDelegate since registration failed (was not initialized).
Case 1: FRegisterComponentContext == nullptr  ==> Super::CreateRenderState_Concurrent + sync CreateSceneProxy (init) + RegisterDebugDrawDelegate ==> works fine
Case 2: FRegisterComponentContext != nullptr  ==> Super::CreateRenderState_Concurrent + deferred CreateSceneProxy + RegisterDebugDrawDelegate (skipped since not init) + CreateSceneProxy ==> ensures in UnregisterDebugDrawDelegate

Bonus:
- some code analysis fixes
- removed some 'CoreMinimal.h' includes
- exported log category 'LogVisual'
- fixed some uninitialized properties (FNavTestDebugDrawDelegateHelper, UNavMeshRenderingComponent, FNavMeshSceneProxyData)
- fixed some methods hiding non-virtual from base class (GetAllocatedSize)
- fixed FGameplayDebuggerCompositeSceneProxy::GetMemoryFootprint that was not considering base class allocations

#rnx

diff --git a/Engine/Source/Developer/GameplayDebugger/Private/GameplayDebuggerRenderingComponent.cpp b/Engine/Source/Developer/GameplayDebugger/Private/GameplayDebuggerRenderingComponent.cpp
--- a/Engine/Source/Developer/GameplayDebugger/Private/GameplayDebuggerRenderingComponent.cpp
+++ b/Engine/Source/Developer/GameplayDebugger/Private/GameplayDebuggerRenderingComponent.cpp
@@ -13,7 +13,7 @@ class FGameplayDebuggerCompositeSceneProxy : public FDebugRenderSceneProxy
 public:
 	FGameplayDebuggerCompositeSceneProxy(const UPrimitiveComponent* InComponent) : FDebugRenderSceneProxy(InComponent) { }
 
-	virtual ~FGameplayDebuggerCompositeSceneProxy()
+	virtual ~FGameplayDebuggerCompositeSceneProxy() override
 	{
 		for (int32 Idx = 0; Idx < ChildProxies.Num(); Idx++)
 		{
@@ -49,18 +49,7 @@ public:
 
 	virtual uint32 GetMemoryFootprint(void) const override
 	{
-		return sizeof(*this) + GetAllocatedSize();
-	}
-
-	uint32 GetAllocatedSize(void) const
-	{
-		uint32 Size = ChildProxies.GetAllocatedSize();
-		for (int32 Idx = 0; Idx < ChildProxies.Num(); Idx++)
-		{
-			Size += ChildProxies[Idx]->GetMemoryFootprint();
-		}
-
-		return Size;
+		return sizeof(*this) + GetAllocatedSizeInternal();
 	}
 
 	void AddChild(FDebugRenderSceneProxy* NewChild)
@@ -73,6 +62,18 @@ public:
 		ChildProxies.Append(Children);
 	}
 
+private:
+	uint32 GetAllocatedSizeInternal(void) const
+	{
+		uint32 Size = FDebugRenderSceneProxy::GetAllocatedSize() + ChildProxies.GetAllocatedSize();
+		for (int32 Idx = 0; Idx < ChildProxies.Num(); Idx++)
+		{
+			Size += ChildProxies[Idx]->GetMemoryFootprint();
+		}
+
+		return Size;
+	}
+
 protected:
 	TArray<FDebugRenderSceneProxy*> ChildProxies;
 };
@@ -137,7 +138,7 @@ FPrimitiveSceneProxy* UGameplayDebuggerRenderingComponent::CreateSceneProxy()
 		TArray<FDebugRenderSceneProxy*> SceneProxies;
 		for (int32 Idx = 0; Idx < OwnerReplicator->GetNumCategories(); Idx++)
 		{
-			TSharedRef<FGameplayDebuggerCategory> Category = OwnerReplicator->GetCategory(Idx);
+			const TSharedRef<FGameplayDebuggerCategory> Category = OwnerReplicator->GetCategory(Idx);
 			if (Category->IsCategoryEnabled())
 			{
 				FDebugDrawDelegateHelper* DebugDrawDelegateHelper = nullptr;
@@ -164,7 +165,7 @@ FPrimitiveSceneProxy* UGameplayDebuggerRenderingComponent::CreateSceneProxy()
 	if (CompositeProxy)
 	{
 		GameplayDebuggerDebugDrawDelegateHelper.InitDelegateHelper(CompositeProxy);
-		GameplayDebuggerDebugDrawDelegateHelper.ReregisterDebugDrawDelegate();
+		GameplayDebuggerDebugDrawDelegateHelper.RegisterDebugDrawDelegate();
 	}
 	return CompositeProxy;
 }
@@ -174,13 +175,6 @@ FBoxSphereBounds UGameplayDebuggerRenderingComponent::CalcBounds(const FTransfor
 	return FBoxSphereBounds(FBox::BuildAABB(FVector::ZeroVector, FVector(1000000.0f, 1000000.0f, 1000000.0f)));
 }
 
-void UGameplayDebuggerRenderingComponent::CreateRenderState_Concurrent(FRegisterComponentContext* Context)
-{
-	Super::CreateRenderState_Concurrent(Context);
-
-	GameplayDebuggerDebugDrawDelegateHelper.RegisterDebugDrawDelegate();
-}
-
 void UGameplayDebuggerRenderingComponent::DestroyRenderState_Concurrent()
 {
 	GameplayDebuggerDebugDrawDelegateHelper.UnregisterDebugDrawDelegate();
diff --git a/Engine/Source/Developer/GameplayDebugger/Public/GameplayDebuggerRenderingComponent.h b/Engine/Source/Developer/GameplayDebugger/Public/GameplayDebuggerRenderingComponent.h
--- a/Engine/Source/Developer/GameplayDebugger/Public/GameplayDebuggerRenderingComponent.h
+++ b/Engine/Source/Developer/GameplayDebugger/Public/GameplayDebuggerRenderingComponent.h
@@ -2,8 +2,6 @@
 
 #pragma once
 
-#include "CoreMinimal.h"
-#include "UObject/ObjectMacros.h"
 #include "Components/PrimitiveComponent.h"
 #include "DebugRenderSceneProxy.h"
 #include "GameplayDebuggerRenderingComponent.generated.h"
@@ -15,7 +13,7 @@ class FGameplayDebuggerDebugDrawDelegateHelper : public FDebugDrawDelegateHelper
 	typedef FDebugDrawDelegateHelper Super;
 
 public:
-	~FGameplayDebuggerDebugDrawDelegateHelper()
+	virtual ~FGameplayDebuggerDebugDrawDelegateHelper() override
 	{
 		Reset();
 	}
@@ -40,8 +38,6 @@ class UGameplayDebuggerRenderingComponent : public UPrimitiveComponent
 
 	virtual FPrimitiveSceneProxy* CreateSceneProxy() override;
 	virtual FBoxSphereBounds CalcBounds(const FTransform &LocalToWorld) const override;
-
-	virtual void CreateRenderState_Concurrent(FRegisterComponentContext* Context) override;
 	virtual void DestroyRenderState_Concurrent() override;
 
 	FGameplayDebuggerDebugDrawDelegateHelper GameplayDebuggerDebugDrawDelegateHelper;
diff --git a/Engine/Source/Developer/LogVisualizer/Private/VisualLoggerRenderingActorBase.cpp b/Engine/Source/Developer/LogVisualizer/Private/VisualLoggerRenderingActorBase.cpp
--- a/Engine/Source/Developer/LogVisualizer/Private/VisualLoggerRenderingActorBase.cpp
+++ b/Engine/Source/Developer/LogVisualizer/Private/VisualLoggerRenderingActorBase.cpp
@@ -1,15 +1,12 @@
 // Copyright Epic Games, Inc. All Rights Reserved.
 
 #include "VisualLoggerRenderingActorBase.h"
-#include "AI/NavigationSystemBase.h"
-#include "VisualLogger/VisualLogger.h"
 #include "LogVisualizerSettings.h"
-#include "VisualLoggerDatabase.h"
-#include "LogVisualizerPublic.h"
 #if WITH_EDITOR
 #include "GeomTools.h"
 #endif // WITH_EDITOR
 #include "VisualLoggerRenderingComponent.h"
+#include "VisualLogger/VisualLogger.h"
 
 namespace FDebugDrawing
 {
@@ -20,7 +17,7 @@ class UVisualLoggerRenderingComponent;
 class FVisualLoggerSceneProxy final : public FDebugRenderSceneProxy
 {
 public:
-	SIZE_T GetTypeHash() const override
+	virtual SIZE_T GetTypeHash() const override
 	{
 		static size_t UniquePointer;
 		return reinterpret_cast<size_t>(&UniquePointer);
@@ -46,11 +43,6 @@ public:
 	}
 
 	virtual uint32 GetMemoryFootprint(void) const override { return sizeof(*this) + GetAllocatedSize(); }
-
-	uint32 GetAllocatedSize(void) const
-	{
-		return FDebugRenderSceneProxy::GetAllocatedSize();
-	}
 };
 
 UVisualLoggerRenderingComponent::UVisualLoggerRenderingComponent(const FObjectInitializer& ObjectInitializer) : Super(ObjectInitializer)
@@ -60,9 +52,9 @@ UVisualLoggerRenderingComponent::UVisualLoggerRenderingComponent(const FObjectIn
 FPrimitiveSceneProxy* UVisualLoggerRenderingComponent::CreateSceneProxy()
 {
 	AVisualLoggerRenderingActorBase* RenderingActor = Cast<AVisualLoggerRenderingActorBase>(GetOuter());
-	if (RenderingActor == NULL)
+	if (RenderingActor == nullptr)
 	{
-		return NULL;
+		return nullptr;
 	}
 
 	ULogVisualizerSettings *Settings = ULogVisualizerSettings::StaticClass()->GetDefaultObject<ULogVisualizerSettings>();
@@ -86,7 +78,7 @@ FPrimitiveSceneProxy* UVisualLoggerRenderingComponent::CreateSceneProxy()
 	if (VLogSceneProxy)
 	{
 		DebugDrawDelegateHelper.InitDelegateHelper(VLogSceneProxy);
-		DebugDrawDelegateHelper.ReregisterDebugDrawDelegate();
+		DebugDrawDelegateHelper.RegisterDebugDrawDelegate();
 	}
 #endif
 	return VLogSceneProxy;
@@ -102,15 +94,6 @@ FBoxSphereBounds UVisualLoggerRenderingComponent::CalcBounds(const FTransform& L
 	return MyBounds;
 }
 
-void UVisualLoggerRenderingComponent::CreateRenderState_Concurrent(FRegisterComponentContext* Context)
-{
-	Super::CreateRenderState_Concurrent(Context);
-
-#if WITH_EDITOR
-	DebugDrawDelegateHelper.RegisterDebugDrawDelegate();
-#endif
-}
-
 void UVisualLoggerRenderingComponent::DestroyRenderState_Concurrent()
 {
 #if WITH_EDITOR
@@ -204,7 +187,7 @@ void AVisualLoggerRenderingActorBase::GetDebugShapes(const FVisualLogEntry& InEn
 
 	if (bAddEntryLocationPointer)
 	{
-		const float Length = 100;
+		constexpr float Length = 100;
 		const FVector DirectionNorm = FVector(0, 0, 1).GetSafeNormal();
 		FVector YAxis, ZAxis;
 		DirectionNorm.FindBestAxisVectors(YAxis, ZAxis);
@@ -358,18 +341,18 @@ void AVisualLoggerRenderingActorBase::GetDebugShapes(const FVisualLogEntry& InEn
 			const bool bDrawLabel = ElementToDraw->Description.IsEmpty() == false;
 			for (int32 Index = 0; Index + 2 < ElementToDraw->Points.Num(); Index += 3)
 			{
-				const FVector Orgin = ElementToDraw->Points[Index];
+				const FVector Origin = ElementToDraw->Points[Index];
 				const FVector Direction = ElementToDraw->Points[Index + 1].GetSafeNormal();
 				const FVector Angles = ElementToDraw->Points[Index + 2];
 				const float Length = Angles.X;
 
 				FVector YAxis, ZAxis;
 				Direction.FindBestAxisVectors(YAxis, ZAxis);
-				DebugShapes.Cones.Add(FDebugRenderSceneProxy::FCone(FScaleMatrix(FVector(Length)) * FMatrix(Direction, YAxis, ZAxis, Orgin), Angles.Y, Angles.Z, Color));
+				DebugShapes.Cones.Add(FDebugRenderSceneProxy::FCone(FScaleMatrix(FVector(Length)) * FMatrix(Direction, YAxis, ZAxis, Origin), Angles.Y, Angles.Z, Color));
 
 				if (bDrawLabel)
 				{
-					DebugShapes.Texts.Add(FDebugRenderSceneProxy::FText3d(ElementToDraw->Description, Orgin, Color));
+					DebugShapes.Texts.Add(FDebugRenderSceneProxy::FText3d(ElementToDraw->Description, Origin, Color));
 				}
 			}
 		}
@@ -523,7 +506,7 @@ void AVisualLoggerRenderingActorBase::GetDebugShapes(const FVisualLogEntry& InEn
 				const FVector XAxis = Rotation.RotateVector(FVector::XAxisVector);
 				const FVector YAxis = Rotation.RotateVector(FVector::YAxisVector);
 				
-				static const int32 CircleDivs = 12;
+				static constexpr int32 CircleDivs = 12;
 				FVector PrevPosition = FVector::ZeroVector;
 				for (int32 Div = 0; Div <= CircleDivs; Div++)
 				{
@@ -544,7 +527,12 @@ void AVisualLoggerRenderingActorBase::GetDebugShapes(const FVisualLogEntry& InEn
 			}
 		}
 			break;
-			
+		case EVisualLoggerShapeElement::Invalid:
+			UE_LOG(LogVisual, Warning, TEXT("Invalid element type"));
+			break;
+
+		default:
+			UE_LOG(LogVisual, Warning, TEXT("Unhandled element type: %s"), ElementToDraw->GetType());
 		}
 	}
 }
diff --git a/Engine/Source/Developer/LogVisualizer/Private/VisualLoggerRenderingComponent.h b/Engine/Source/Developer/LogVisualizer/Private/VisualLoggerRenderingComponent.h
--- a/Engine/Source/Developer/LogVisualizer/Private/VisualLoggerRenderingComponent.h
+++ b/Engine/Source/Developer/LogVisualizer/Private/VisualLoggerRenderingComponent.h
@@ -2,8 +2,6 @@
 
 #pragma once
 
-#include "CoreMinimal.h"
-#include "UObject/ObjectMacros.h"
 #include "DebugRenderSceneProxy.h"
 #include "Components/PrimitiveComponent.h"
 #include "VisualLoggerRenderingComponent.generated.h"
@@ -20,7 +18,6 @@ public:
 
 	virtual FPrimitiveSceneProxy* CreateSceneProxy() override;
 	virtual FBoxSphereBounds CalcBounds(const FTransform &LocalToWorld) const override;
-	virtual void CreateRenderState_Concurrent(FRegisterComponentContext* Context) override;
 	virtual void DestroyRenderState_Concurrent() override;
 
 private:
diff --git a/Engine/Source/Runtime/AIModule/Classes/EnvironmentQuery/EQSRenderingComponent.h b/Engine/Source/Runtime/AIModule/Classes/EnvironmentQuery/EQSRenderingComponent.h
--- a/Engine/Source/Runtime/AIModule/Classes/EnvironmentQuery/EQSRenderingComponent.h
+++ b/Engine/Source/Runtime/AIModule/Classes/EnvironmentQuery/EQSRenderingComponent.h
@@ -2,8 +2,6 @@
 
 #pragma once
 
-#include "CoreMinimal.h"
-#include "UObject/ObjectMacros.h"
 #include "EngineDefines.h"
 #include "EnvironmentQuery/EnvQueryTypes.h"
 #include "PrimitiveViewRelevance.h"
@@ -20,7 +18,7 @@ class AIMODULE_API FEQSSceneProxy final : public FDebugRenderSceneProxy
 {
 	friend class FEQSRenderingDebugDrawDelegateHelper;
 public:
-	SIZE_T GetTypeHash() const override;
+	virtual SIZE_T GetTypeHash() const override;
 
 	explicit FEQSSceneProxy(const UPrimitiveComponent& InComponent, const FString& ViewFlagName = TEXT("DebugAI"), const TArray<FSphere>& Spheres = TArray<FSphere>(), const TArray<FText3d>& Texts = TArray<FText3d>());
 	
@@ -49,8 +47,8 @@ class FEQSRenderingDebugDrawDelegateHelper : public FDebugDrawDelegateHelper
 
 public:
 	FEQSRenderingDebugDrawDelegateHelper()
-		: ActorOwner(NULL)
-		, QueryDataSource(NULL)
+		: ActorOwner(nullptr)
+		, QueryDataSource(nullptr)
 		, bDrawOnlyWhenSelected(false)
 	{
 	}
@@ -95,7 +93,6 @@ class AIMODULE_API UEQSRenderingComponent : public UPrimitiveComponent
 
 	virtual FPrimitiveSceneProxy* CreateSceneProxy() override;
 	virtual FBoxSphereBounds CalcBounds(const FTransform &LocalToWorld) const override;
-	virtual void CreateRenderState_Concurrent(FRegisterComponentContext* Context) override;
 	virtual void DestroyRenderState_Concurrent() override;
 
 	void ClearStoredDebugData();
diff --git a/Engine/Source/Runtime/AIModule/Private/EnvironmentQuery/EQSRenderingComponent.cpp b/Engine/Source/Runtime/AIModule/Private/EnvironmentQuery/EQSRenderingComponent.cpp
--- a/Engine/Source/Runtime/AIModule/Private/EnvironmentQuery/EQSRenderingComponent.cpp
+++ b/Engine/Source/Runtime/AIModule/Private/EnvironmentQuery/EQSRenderingComponent.cpp
@@ -5,7 +5,7 @@
 #include "Engine/Canvas.h"
 #include "EnvironmentQuery/EQSQueryResultSourceInterface.h"
 
-static const int32 EQSMaxItemsDrawn = 10000;
+static constexpr int32 EQSMaxItemsDrawn = 10000;
 
 namespace FEQSRenderingHelper
 {
@@ -14,7 +14,7 @@ namespace FEQSRenderingHelper
 		if (Items.IsValidIndex(Index) &&
 			ItemType->IsChildOf(UEnvQueryItemType_VectorBase::StaticClass()))
 		{
-			UEnvQueryItemType_VectorBase* DefTypeOb = ItemType->GetDefaultObject<UEnvQueryItemType_VectorBase>();
+			const UEnvQueryItemType_VectorBase* DefTypeOb = ItemType->GetDefaultObject<UEnvQueryItemType_VectorBase>();
 			return DefTypeOb->GetItemLocation(RawData.GetData() + Items[Index].DataOffset);
 		}
 		return FVector::ZeroVector;
@@ -68,15 +68,15 @@ FEQSSceneProxy::FEQSSceneProxy(const UPrimitiveComponent& InComponent, const FSt
 #if  USE_EQS_DEBUGGER 
 void FEQSSceneProxy::CollectEQSData(const UPrimitiveComponent* InComponent, const IEQSQueryResultSourceInterface* InQueryDataSource, TArray<FSphere>& Spheres, TArray<FText3d>& Texts, TArray<EQSDebug::FDebugHelper>& DebugItems)
 {
-	AActor* ActorOwner = InComponent ? InComponent->GetOwner() : NULL;
+	AActor* ActorOwner = InComponent ? InComponent->GetOwner() : nullptr;
 	IEQSQueryResultSourceInterface* QueryDataSource = const_cast<IEQSQueryResultSourceInterface*>(InQueryDataSource);
-	if (QueryDataSource == NULL)
+	if (QueryDataSource == nullptr)
 	{
 		QueryDataSource = Cast<IEQSQueryResultSourceInterface>(ActorOwner);
-		if (QueryDataSource == NULL)
+		if (QueryDataSource == nullptr)
 		{
 			QueryDataSource = Cast<IEQSQueryResultSourceInterface>(const_cast<UPrimitiveComponent*>(InComponent));
-			if (QueryDataSource == NULL)
+			if (QueryDataSource == nullptr)
 			{
 				return;
 			}
@@ -101,9 +101,9 @@ void FEQSSceneProxy::CollectEQSData(const UPrimitiveComponent* InComponent, cons
 
 void FEQSSceneProxy::CollectEQSData(const FEnvQueryResult* ResultItems, const FEnvQueryInstance* QueryInstance, float HighlightRangePct, bool ShouldDrawFailedItems, TArray<FSphere>& Spheres, TArray<FText3d>& Texts, TArray<EQSDebug::FDebugHelper>& DebugItems)
 {
-	if (ResultItems == NULL)
+	if (ResultItems == nullptr)
 	{
-		if (QueryInstance == NULL)
+		if (QueryInstance == nullptr)
 		{
 			return;
 		}
@@ -120,9 +120,9 @@ void FEQSSceneProxy::CollectEQSData(const FEnvQueryResult* ResultItems, const FE
 	const TArray<FEnvQueryItem>& Items = bUseMidResults ? QueryInstance->DebugData.DebugItems : ResultItems->Items; //-V595
 	const TArray<uint8>& RawData = bUseMidResults ? QueryInstance->DebugData.RawData : ResultItems->RawData;
 	const int32 ItemCountLimit = FMath::Clamp(Items.Num(), 0, EQSMaxItemsDrawn);
-	const bool bNoTestsPerformed = QueryInstance != NULL && QueryInstance->CurrentTest <= 0;
+	const bool bNoTestsPerformed = QueryInstance != nullptr && QueryInstance->CurrentTest <= 0;
 	
-	const bool bSingleItemResult = QueryInstance != NULL && QueryInstance->DebugData.bSingleItemResult;
+	const bool bSingleItemResult = QueryInstance != nullptr && QueryInstance->DebugData.bSingleItemResult;
 
 	float MinScore = 0.f;
 	float MaxScore = -BIG_NUMBER;
@@ -168,7 +168,7 @@ void FEQSSceneProxy::CollectEQSData(const FEnvQueryResult* ResultItems, const FE
 		if (Items[0].IsValid())
 		{
 			const float Score = Items[0].Score;
-			const bool bLowRadius = false;
+			constexpr bool bLowRadius = false;
 			const float Radius = ItemDrawRadius.X * (bLowRadius ? 0.2f : 1.0f);
 			const FVector Loc = FEQSRenderingHelper::ExtractLocation(ResultItems->ItemType, RawData, Items, 0);
 			Spheres.Add(FSphere(Radius, Loc, FLinearColor(0.0, 1.0, 0.12, 1)));
@@ -212,7 +212,6 @@ void FEQSSceneProxy::CollectEQSData(const FEnvQueryResult* ResultItems, const FE
 				continue;
 			}
 
-			const float Score = bNoTestsPerformed ? 1 : Items[ItemIndex].Score;
 			const bool bLowRadius = (HighlightThreshold < FLT_MAX) && (bNoTestsPerformed || (Items[ItemIndex].Score < HighlightThreshold));
 			const float Radius = ItemDrawRadius.X * (bLowRadius ? 0.2f : 1.0f);
 			const FVector Loc = FEQSRenderingHelper::ExtractLocation(QueryInstance->ItemType, InstanceDebugData.RawData, DebugQueryItems, ItemIndex);
@@ -290,7 +289,7 @@ FPrimitiveSceneProxy* UEQSRenderingComponent::CreateSceneProxy()
 	if (NewSceneProxy)
 	{
 		EQSRenderingDebugDrawDelegateHelper.InitDelegateHelper(NewSceneProxy);
-		EQSRenderingDebugDrawDelegateHelper.ReregisterDebugDrawDelegate();
+		EQSRenderingDebugDrawDelegateHelper.RegisterDebugDrawDelegate();
 	}
 #endif
 
@@ -351,15 +350,6 @@ FBoxSphereBounds UEQSRenderingComponent::CalcBounds(const FTransform& LocalToWor
 	return FBoxSphereBounds(BoundingSphere).TransformBy(LocalToWorld);
 }
 
-void UEQSRenderingComponent::CreateRenderState_Concurrent(FRegisterComponentContext* Context)
-{
-	Super::CreateRenderState_Concurrent(Context);
-
-#if USE_EQS_DEBUGGER
-	EQSRenderingDebugDrawDelegateHelper.RegisterDebugDrawDelegate();
-#endif
-}
-
 void UEQSRenderingComponent::DestroyRenderState_Concurrent()
 {
 #if USE_EQS_DEBUGGER
diff --git a/Engine/Source/Runtime/Engine/Public/VisualLogger/VisualLogger.h b/Engine/Source/Runtime/Engine/Public/VisualLogger/VisualLogger.h
--- a/Engine/Source/Runtime/Engine/Public/VisualLogger/VisualLogger.h
+++ b/Engine/Source/Runtime/Engine/Public/VisualLogger/VisualLogger.h
@@ -143,8 +143,7 @@
 
 class FVisualLogger;
 
-DECLARE_LOG_CATEGORY_EXTERN(LogVisual, Display, All);
-
+ENGINE_API DECLARE_LOG_CATEGORY_EXTERN(LogVisual, Display, All);
 
 #if ENABLE_VISUAL_LOG
 
diff --git a/Engine/Source/Runtime/NavigationSystem/Private/NavMesh/NavMeshRenderingComponent.cpp b/Engine/Source/Runtime/NavigationSystem/Private/NavMesh/NavMeshRenderingComponent.cpp
--- a/Engine/Source/Runtime/NavigationSystem/Private/NavMesh/NavMeshRenderingComponent.cpp
+++ b/Engine/Source/Runtime/NavigationSystem/Private/NavMesh/NavMeshRenderingComponent.cpp
@@ -5,10 +5,8 @@
  =============================================================================*/
 
 #include "NavMesh/NavMeshRenderingComponent.h"
-#include "EngineGlobals.h"
 #include "NavigationSystem.h"
 #include "Engine/Canvas.h"
-#include "Engine/CollisionProfile.h"
 #include "Engine/Engine.h"
 #include "NavigationOctree.h"
 #include "NavMesh/RecastHelpers.h"
@@ -27,24 +25,45 @@
 #include "EditorViewportClient.h"
 #endif
 
-static const FColor NavMeshRenderColor_Recast_TriangleEdges(255, 255, 255);
-static const FColor NavMeshRenderColor_Recast_TileEdges(16, 16, 16, 32);
-static const FColor NavMeshRenderColor_Recast_NavMeshEdges(32, 63, 0, 220);
-static const FColor NavMeshRenderColor_RecastMesh(140, 255, 0, 164);
-static const FColor NavMeshRenderColor_TileBounds(255, 255, 64, 255);
-static const FColor NavMeshRenderColor_PathCollidingGeom(255, 255, 255, 40);
-static const FColor NavMeshRenderColor_RecastTileBeingRebuilt(255, 0, 0, 64);
-static const FColor NavMeshRenderColor_OffMeshConnectionInvalid(64, 64, 64);
+static constexpr FColor NavMeshRenderColor_Recast_TriangleEdges(255, 255, 255);
+static constexpr FColor NavMeshRenderColor_Recast_TileEdges(16, 16, 16, 32);
+static constexpr FColor NavMeshRenderColor_RecastMesh(140, 255, 0, 164);
+static constexpr FColor NavMeshRenderColor_TileBounds(255, 255, 64, 255);
+static constexpr FColor NavMeshRenderColor_PathCollidingGeom(255, 255, 255, 40);
+static constexpr FColor NavMeshRenderColor_RecastTileBeingRebuilt(255, 0, 0, 64);
+static constexpr FColor NavMeshRenderColor_OffMeshConnectionInvalid(64, 64, 64);
 static const FColor NavMeshRenderColor_PolyForbidden(FColorList::Black);
 
-static const float DefaultEdges_LineThickness = 0.0f;
-static const float PolyEdges_LineThickness = 1.5f;
-static const float NavMeshEdges_LineThickness = 3.5f;
-static const float LinkLines_LineThickness = 2.0f;
-static const float ClusterLinkLines_LineThickness = 2.0f;
+static constexpr float DefaultEdges_LineThickness = 0.0f;
+static constexpr float PolyEdges_LineThickness = 1.5f;
+static constexpr float NavMeshEdges_LineThickness = 3.5f;
+static constexpr float LinkLines_LineThickness = 2.0f;
+static constexpr float ClusterLinkLines_LineThickness = 2.0f;
 
 namespace FNavMeshRenderingHelpers
 {
+	void DrawDebugBox(FPrimitiveDrawInterface* PDI, FVector const& Center, FVector const& Box, FColor const& Color)
+	{
+		// no debug line drawing on dedicated server
+		if (PDI != nullptr)
+		{
+			PDI->DrawLine(Center + FVector(Box.X, Box.Y, Box.Z), Center + FVector(Box.X, -Box.Y, Box.Z), Color, SDPG_World);
+			PDI->DrawLine(Center + FVector(Box.X, -Box.Y, Box.Z), Center + FVector(-Box.X, -Box.Y, Box.Z), Color, SDPG_World);
+			PDI->DrawLine(Center + FVector(-Box.X, -Box.Y, Box.Z), Center + FVector(-Box.X, Box.Y, Box.Z), Color, SDPG_World);
+			PDI->DrawLine(Center + FVector(-Box.X, Box.Y, Box.Z), Center + FVector(Box.X, Box.Y, Box.Z), Color, SDPG_World);
+
+			PDI->DrawLine(Center + FVector(Box.X, Box.Y, -Box.Z), Center + FVector(Box.X, -Box.Y, -Box.Z), Color, SDPG_World);
+			PDI->DrawLine(Center + FVector(Box.X, -Box.Y, -Box.Z), Center + FVector(-Box.X, -Box.Y, -Box.Z), Color, SDPG_World);
+			PDI->DrawLine(Center + FVector(-Box.X, -Box.Y, -Box.Z), Center + FVector(-Box.X, Box.Y, -Box.Z), Color, SDPG_World);
+			PDI->DrawLine(Center + FVector(-Box.X, Box.Y, -Box.Z), Center + FVector(Box.X, Box.Y, -Box.Z), Color, SDPG_World);
+
+			PDI->DrawLine(Center + FVector(Box.X, Box.Y, Box.Z), Center + FVector(Box.X, Box.Y, -Box.Z), Color, SDPG_World);
+			PDI->DrawLine(Center + FVector(Box.X, -Box.Y, Box.Z), Center + FVector(Box.X, -Box.Y, -Box.Z), Color, SDPG_World);
+			PDI->DrawLine(Center + FVector(-Box.X, -Box.Y, Box.Z), Center + FVector(-Box.X, -Box.Y, -Box.Z), Color, SDPG_World);
+			PDI->DrawLine(Center + FVector(-Box.X, Box.Y, Box.Z), Center + FVector(-Box.X, Box.Y, -Box.Z), Color, SDPG_World);
+		}
+	}
+
 	bool LineInView(const FVector& Start, const FVector& End, const FSceneView* View)
 	{
 		if (FVector::DistSquaredXY(Start, View->ViewMatrices.GetViewOrigin()) > ARecastNavMesh::GetDrawDistanceSq() ||
@@ -123,12 +142,10 @@ namespace FNavMeshRenderingHelpers
 
 	void CacheArrowHead(TArray<FDebugRenderSceneProxy::FDebugLine>& DebugLines, const FVector& Tip, const FVector& Origin, const float Size, const FLinearColor& Color, float LineThickness = 0)
 	{
-		FVector Ax, Ay, Az(0, 1, 0);
-		Ay = Origin - Tip;
-		Ay.Normalize();
-		Ax = FVector::CrossProduct(Az, Ay);
+		const FVector Az(0.f, 1.f, 0.f);
+		const FVector Ay = (Origin - Tip).GetSafeNormal();
+		const FVector Ax = FVector::CrossProduct(Az, Ay);
 
-		FHitProxyId HitProxyId;
 		DebugLines.Add(FDebugRenderSceneProxy::FDebugLine(Tip, FVector(Tip.X + Ay.X*Size + Ax.X*Size / 3, Tip.Y + Ay.Y*Size + Ax.Y*Size / 3, Tip.Z + Ay.Z*Size + Ax.Z*Size / 3), Color.ToFColor(true)));
 		DebugLines.Add(FDebugRenderSceneProxy::FDebugLine(Tip, FVector(Tip.X + Ay.X*Size - Ax.X*Size / 3, Tip.Y + Ay.Y*Size - Ax.Y*Size / 3, Tip.Z + Ay.Z*Size - Ax.Z*Size / 3), Color.ToFColor(true)));
 	}
@@ -150,7 +167,6 @@ namespace FNavMeshRenderingHelpers
 		const float	AngleDelta = 2.0f * PI / NumSides;
 		FVector	LastVertex = Base + X * Radius;
 
-		FHitProxyId HitProxyId;
 		for (int32 SideIndex = 0; SideIndex < NumSides; SideIndex++)
 		{
 			const FVector Vertex = Base + (X * FMath::Cos(AngleDelta * (SideIndex + 1)) + Y * FMath::Sin(AngleDelta * (SideIndex + 1))) * Radius;
@@ -170,9 +186,9 @@ namespace FNavMeshRenderingHelpers
 
 	FColor GetClusterColor(int32 Idx)
 	{
-		uint8 r = 1 + GetBit(Idx, 1) + GetBit(Idx, 3) * 2;
-		uint8 g = 1 + GetBit(Idx, 2) + GetBit(Idx, 4) * 2;
-		uint8 b = 1 + GetBit(Idx, 0) + GetBit(Idx, 5) * 2;
+		const uint8 r = 1 + GetBit(Idx, 1) + GetBit(Idx, 3) * 2;
+		const uint8 g = 1 + GetBit(Idx, 2) + GetBit(Idx, 4) * 2;
+		const uint8 b = 1 + GetBit(Idx, 0) + GetBit(Idx, 5) * 2;
 		return FColor(r * 63, g * 63, b * 63, 164);
 	}
 
@@ -194,7 +210,6 @@ namespace FNavMeshRenderingHelpers
 
 	void AddVertex(FNavMeshSceneProxyData::FDebugMeshData& MeshData, const FVector& Pos, const FColor Color)
 	{
-		const int32 VertexIndex = MeshData.Vertices.Num();
 		FDynamicMeshVertex* Vertex = new(MeshData.Vertices) FDynamicMeshVertex;
 		Vertex->Position = Pos;
 		Vertex->TextureCoordinate[0] = FVector2D::ZeroVector;
@@ -254,6 +269,31 @@ namespace FNavMeshRenderingHelpers
 	{
 		return (Flags & (1 << static_cast<int32>(TestFlag))) != 0;
 	}
+
+#if WITH_RECAST
+	int32 GetDetailFlags(const ARecastNavMesh* NavMesh)
+	{
+		return (NavMesh == nullptr) ? 0 : 0 |
+			(NavMesh->bDrawTriangleEdges ? (1 << static_cast<int32>(ENavMeshDetailFlags::TriangleEdges)) : 0) |
+			(NavMesh->bDrawPolyEdges ? (1 << static_cast<int32>(ENavMeshDetailFlags::PolyEdges)) : 0) |
+			(NavMesh->bDrawFilledPolys ? (1 << static_cast<int32>(ENavMeshDetailFlags::FilledPolys)) : 0) |
+			(NavMesh->bDrawNavMeshEdges ? (1 << static_cast<int32>(ENavMeshDetailFlags::BoundaryEdges)) : 0) |
+			(NavMesh->bDrawTileBounds ? (1 << static_cast<int32>(ENavMeshDetailFlags::TileBounds)) : 0) |
+			(NavMesh->bDrawPathCollidingGeometry ? (1 << static_cast<int32>(ENavMeshDetailFlags::PathCollidingGeometry)) : 0) |
+			(NavMesh->bDrawTileLabels ? (1 << static_cast<int32>(ENavMeshDetailFlags::TileLabels)) : 0) |
+			(NavMesh->bDrawPolygonLabels ? (1 << static_cast<int32>(ENavMeshDetailFlags::PolygonLabels)) : 0) |
+			(NavMesh->bDrawDefaultPolygonCost ? (1 << static_cast<int32>(ENavMeshDetailFlags::PolygonCost)) : 0) |
+			(NavMesh->bDrawPolygonFlags ? (1 << static_cast<int32>(ENavMeshDetailFlags::PolygonFlags)) : 0) |
+			(NavMesh->bDrawLabelsOnPathNodes ? (1 << static_cast<int32>(ENavMeshDetailFlags::PathLabels)) : 0) |
+			(NavMesh->bDrawNavLinks ? (1 << static_cast<int32>(ENavMeshDetailFlags::NavLinks)) : 0) |
+			(NavMesh->bDrawFailedNavLinks ? (1 << static_cast<int32>(ENavMeshDetailFlags::FailedNavLinks)) : 0) |
+			(NavMesh->bDrawClusters ? (1 << static_cast<int32>(ENavMeshDetailFlags::Clusters)) : 0) |
+			(NavMesh->bDrawOctree ? (1 << static_cast<int32>(ENavMeshDetailFlags::NavOctree)) : 0) |
+			(NavMesh->bDrawOctreeDetails ? (1 << static_cast<int32>(ENavMeshDetailFlags::NavOctreeDetails)) : 0) |
+			(NavMesh->bDrawMarkedForbiddenPolys ? (1 << static_cast<int32>(ENavMeshDetailFlags::MarkForbiddenPolys)) : 0);
+	}
+#endif // WITH_RECAST
+
 }
 
 //////////////////////////////////////////////////////////////////////////
@@ -437,28 +477,6 @@ uint32 FNavMeshSceneProxyData::GetAllocatedSize() const
 
 #if WITH_RECAST
 
-int32 FNavMeshSceneProxyData::GetDetailFlags(const ARecastNavMesh* NavMesh) const
-{
-	return (NavMesh == nullptr) ? 0 : 0 |
-		(NavMesh->bDrawTriangleEdges ? (1 << static_cast<int32>(ENavMeshDetailFlags::TriangleEdges)) : 0) |
-		(NavMesh->bDrawPolyEdges ? (1 << static_cast<int32>(ENavMeshDetailFlags::PolyEdges)) : 0) |
-		(NavMesh->bDrawFilledPolys ? (1 << static_cast<int32>(ENavMeshDetailFlags::FilledPolys)) : 0) |
-		(NavMesh->bDrawNavMeshEdges ? (1 << static_cast<int32>(ENavMeshDetailFlags::BoundaryEdges)) : 0) |
-		(NavMesh->bDrawTileBounds ? (1 << static_cast<int32>(ENavMeshDetailFlags::TileBounds)) : 0) |
-		(NavMesh->bDrawPathCollidingGeometry ? (1 << static_cast<int32>(ENavMeshDetailFlags::PathCollidingGeometry)) : 0) |
-		(NavMesh->bDrawTileLabels ? (1 << static_cast<int32>(ENavMeshDetailFlags::TileLabels)) : 0) |
-		(NavMesh->bDrawPolygonLabels ? (1 << static_cast<int32>(ENavMeshDetailFlags::PolygonLabels)) : 0) |
-		(NavMesh->bDrawDefaultPolygonCost ? (1 << static_cast<int32>(ENavMeshDetailFlags::PolygonCost)) : 0) |
-		(NavMesh->bDrawPolygonFlags ? (1 << static_cast<int32>(ENavMeshDetailFlags::PolygonFlags)) : 0) |
-		(NavMesh->bDrawLabelsOnPathNodes ? (1 << static_cast<int32>(ENavMeshDetailFlags::PathLabels)) : 0) |
-		(NavMesh->bDrawNavLinks ? (1 << static_cast<int32>(ENavMeshDetailFlags::NavLinks)) : 0) |
-		(NavMesh->bDrawFailedNavLinks ? (1 << static_cast<int32>(ENavMeshDetailFlags::FailedNavLinks)) : 0) |
-		(NavMesh->bDrawClusters ? (1 << static_cast<int32>(ENavMeshDetailFlags::Clusters)) : 0) |
-		(NavMesh->bDrawOctree ? (1 << static_cast<int32>(ENavMeshDetailFlags::NavOctree)) : 0) | 
-		(NavMesh->bDrawOctreeDetails ? (1 << static_cast<int32>(ENavMeshDetailFlags::NavOctreeDetails)) : 0) | 
-		(NavMesh->bDrawMarkedForbiddenPolys ? (1 << static_cast<int32>(ENavMeshDetailFlags::MarkForbiddenPolys)) : 0);
-}
-
 void FNavMeshSceneProxyData::GatherData(const ARecastNavMesh* NavMesh, int32 InNavDetailFlags, const TArray<int32>& TileSet)
 {
 	QUICK_SCOPE_CYCLE_COUNTER(STAT_NavMesh_GatherDebugDrawingGeometry);
@@ -470,7 +488,6 @@ void FNavMeshSceneProxyData::GatherData(const ARecastNavMesh* NavMesh, int32 InN
 		bNeedsNewData = false;
 		bDataGathered = true;
 
-		FHitProxyId HitProxyId = FHitProxyId();
 		NavMeshDrawOffset.Z = NavMesh->DrawOffset;
 
 		FRecastDebugGeometry NavMeshGeometry;
@@ -895,8 +912,8 @@ void FNavMeshSceneProxyData::GatherData(const ARecastNavMesh* NavMesh, int32 InN
 					const FColor ColA = (SegInfo.ValidEnds & FRecastDebugGeometry::OMLE_Left) ? FColor::White : FColor::Black;
 					const FColor ColB = (SegInfo.ValidEnds & FRecastDebugGeometry::OMLE_Right) ? FColor::White : FColor::Black;
 
-					const int32 NumArcPoints = 8;
-					const float ArcPtsScale = 1.0f / NumArcPoints;
+					constexpr int32 NumArcPoints = 8;
+					constexpr float ArcPtsScale = 1.0f / NumArcPoints;
 
 					FVector Prev0 = FNavMeshRenderingHelpers::EvalArc(A0, Edge0, Len0*0.25f, 0);
 					FVector Prev1 = FNavMeshRenderingHelpers::EvalArc(A1, Edge1, Len1*0.25f, 0);
@@ -987,7 +1004,7 @@ void FNavMeshSceneProxyData::AddMeshForInternalData(const FRecastInternalDebugDa
 			}
 
 			// Add triangles
-			for (FNavMeshRenderingHelpers::FUniqueColor CurrentColor : UniqueColors)
+			for (const FNavMeshRenderingHelpers::FUniqueColor& CurrentColor : UniqueColors)
 			{
 				const uint32 VertexCount = CurrentColor.Count * 3;
 				FDebugMeshData& MeshData = MeshBuilders.AddDefaulted_GetRef();
@@ -1155,28 +1172,6 @@ void FNavMeshDebugDrawDelegateHelper::UnregisterDebugDrawDelegate()
 }
 #endif
 
-void FNavMeshSceneProxy::DrawDebugBox(FPrimitiveDrawInterface* PDI, FVector const& Center, FVector const& Box, FColor const& Color) const
-{
-	// no debug line drawing on dedicated server
-	if (PDI != NULL)
-	{
-		PDI->DrawLine(Center + FVector(Box.X, Box.Y, Box.Z), Center + FVector(Box.X, -Box.Y, Box.Z), Color, SDPG_World);
-		PDI->DrawLine(Center + FVector(Box.X, -Box.Y, Box.Z), Center + FVector(-Box.X, -Box.Y, Box.Z), Color, SDPG_World);
-		PDI->DrawLine(Center + FVector(-Box.X, -Box.Y, Box.Z), Center + FVector(-Box.X, Box.Y, Box.Z), Color, SDPG_World);
-		PDI->DrawLine(Center + FVector(-Box.X, Box.Y, Box.Z), Center + FVector(Box.X, Box.Y, Box.Z), Color, SDPG_World);
-
-		PDI->DrawLine(Center + FVector(Box.X, Box.Y, -Box.Z), Center + FVector(Box.X, -Box.Y, -Box.Z), Color, SDPG_World);
-		PDI->DrawLine(Center + FVector(Box.X, -Box.Y, -Box.Z), Center + FVector(-Box.X, -Box.Y, -Box.Z), Color, SDPG_World);
-		PDI->DrawLine(Center + FVector(-Box.X, -Box.Y, -Box.Z), Center + FVector(-Box.X, Box.Y, -Box.Z), Color, SDPG_World);
-		PDI->DrawLine(Center + FVector(-Box.X, Box.Y, -Box.Z), Center + FVector(Box.X, Box.Y, -Box.Z), Color, SDPG_World);
-
-		PDI->DrawLine(Center + FVector(Box.X, Box.Y, Box.Z), Center + FVector(Box.X, Box.Y, -Box.Z), Color, SDPG_World);
-		PDI->DrawLine(Center + FVector(Box.X, -Box.Y, Box.Z), Center + FVector(Box.X, -Box.Y, -Box.Z), Color, SDPG_World);
-		PDI->DrawLine(Center + FVector(-Box.X, -Box.Y, Box.Z), Center + FVector(-Box.X, -Box.Y, -Box.Z), Color, SDPG_World);
-		PDI->DrawLine(Center + FVector(-Box.X, Box.Y, Box.Z), Center + FVector(-Box.X, Box.Y, -Box.Z), Color, SDPG_World);
-	}
-}
-
 void FNavMeshSceneProxy::GetDynamicMeshElements(const TArray<const FSceneView*>& Views, const FSceneViewFamily& ViewFamily, uint32 VisibilityMap, FMeshElementCollector& Collector) const
 {
 	QUICK_SCOPE_CYCLE_COUNTER(STAT_RecastRenderingSceneProxy_GetDynamicMeshElements);
@@ -1198,7 +1193,7 @@ void FNavMeshSceneProxy::GetDynamicMeshElements(const TArray<const FSceneView*>&
 			for (int32 Index = 0; Index < ProxyData.OctreeBounds.Num(); ++Index)
 			{
 				const FBoxCenterAndExtent& ProxyBounds = ProxyData.OctreeBounds[Index];
-				DrawDebugBox(PDI, ProxyBounds.Center, ProxyBounds.Extent, FColor::White);
+				FNavMeshRenderingHelpers::DrawDebugBox(PDI, ProxyBounds.Center, ProxyBounds.Extent, FColor::White);
 			}
 
 			// Draw Mesh
@@ -1366,7 +1361,7 @@ void FNavMeshDebugDrawDelegateHelper::DrawDebugLabels(UCanvas* Canvas, APlayerCo
 	const FColor OldDrawColor = Canvas->DrawColor;
 	Canvas->SetDrawColor(FColor::White);
 	const FSceneView* View = Canvas->SceneView;
-	UFont* Font = GEngine->GetSmallFont();
+	const UFont* Font = GEngine->GetSmallFont();
 	const FNavMeshSceneProxyData::FDebugText* DebugText = DebugLabels.GetData();
 	for (int32 Idx = 0; Idx < DebugLabels.Num(); ++Idx, ++DebugText)
 	{
@@ -1392,7 +1387,7 @@ FPrimitiveViewRelevance FNavMeshSceneProxy::GetViewRelevance(const FSceneView* V
 	return Result;
 }
 
-uint32 FNavMeshSceneProxy::GetAllocatedSize() const
+uint32 FNavMeshSceneProxy::GetAllocatedSizeInternal() const
 {
 	return FDebugRenderSceneProxy::GetAllocatedSize() +
 		ProxyData.GetAllocatedSize() +
@@ -1438,6 +1433,7 @@ UNavMeshRenderingComponent::UNavMeshRenderingComponent(const FObjectInitializer&
 	bIsEditorOnly = true;
 	bSelectable = false;
 	bCollectNavigationData = false;
+	bForceUpdate = false;
 }
 
 bool UNavMeshRenderingComponent::IsNavigationShowFlagSet(const UWorld* World)
@@ -1532,8 +1528,8 @@ void UNavMeshRenderingComponent::OnUnregister()
 
 void UNavMeshRenderingComponent::GatherData(const ARecastNavMesh& NavMesh, FNavMeshSceneProxyData& OutProxyData) const
 {
-	const int32 DetailFlags = OutProxyData.GetDetailFlags(&NavMesh);
-	TArray<int32> EmptyTileSet;
+	const int32 DetailFlags = FNavMeshRenderingHelpers::GetDetailFlags(&NavMesh);
+	const TArray<int32> EmptyTileSet;
 	OutProxyData.GatherData(&NavMesh, DetailFlags, EmptyTileSet);
 }
 
@@ -1561,7 +1557,7 @@ FPrimitiveSceneProxy* UNavMeshRenderingComponent::CreateSceneProxy()
 	if (NavMeshSceneProxy)
 	{
 		NavMeshDebugDrawDelegateManager.InitDelegateHelper(NavMeshSceneProxy);
-		NavMeshDebugDrawDelegateManager.ReregisterDebugDrawDelegate();
+		NavMeshDebugDrawDelegateManager.RegisterDebugDrawDelegate();
 	}
 	return NavMeshSceneProxy;
 #else
@@ -1569,15 +1565,6 @@ FPrimitiveSceneProxy* UNavMeshRenderingComponent::CreateSceneProxy()
 #endif //WITH_RECAST && !UE_BUILD_SHIPPING && !UE_BUILD_TEST
 }
 
-void UNavMeshRenderingComponent::CreateRenderState_Concurrent(FRegisterComponentContext* Context)
-{
-	Super::CreateRenderState_Concurrent(Context);
-
-#if WITH_RECAST && !UE_BUILD_SHIPPING && !UE_BUILD_TEST
-	NavMeshDebugDrawDelegateManager.RegisterDebugDrawDelegate();
-#endif
-}
-
 void UNavMeshRenderingComponent::DestroyRenderState_Concurrent()
 {
 #if WITH_RECAST && !UE_BUILD_SHIPPING && !UE_BUILD_TEST
@@ -1601,7 +1588,7 @@ FBoxSphereBounds UNavMeshRenderingComponent::CalcBounds(const FTransform& LocalT
 			const FNavigationOctree* NavOctree = NavSys ? NavSys->GetNavOctree() : nullptr;
 			if (NavOctree)
 			{
-				//this is only iterating over the rootnode
+				//this is only iterating over the root node
 				BoundingBox += NavOctree->GetRootBounds().GetBox();
 			}
 		}
diff --git a/Engine/Source/Runtime/NavigationSystem/Private/NavMesh/NavTestRenderingComponent.cpp b/Engine/Source/Runtime/NavigationSystem/Private/NavMesh/NavTestRenderingComponent.cpp
--- a/Engine/Source/Runtime/NavigationSystem/Private/NavMesh/NavTestRenderingComponent.cpp
+++ b/Engine/Source/Runtime/NavigationSystem/Private/NavMesh/NavTestRenderingComponent.cpp
@@ -1,7 +1,6 @@
 // Copyright Epic Games, Inc. All Rights Reserved.
 
 #include "NavMesh/NavTestRenderingComponent.h"
-#include "EngineGlobals.h"
 #include "Engine/Engine.h"
 #include "Engine/Canvas.h"
 #include "SceneManagement.h"
@@ -9,10 +8,10 @@
 #include "NavMesh/RecastNavMesh.h"
 #include "Debug/DebugDrawService.h"
 
-static const FColor NavMeshRenderColor_OpenSet(255,128,0,255);
-static const FColor NavMeshRenderColor_ClosedSet(255,196,0,255);
-static const uint8 NavMeshRenderAlpha_Modifed = 255;
-static const uint8 NavMeshRenderAlpha_NonModified = 64;
+static constexpr FColor NavMeshRenderColor_OpenSet(255,128,0,255);
+static constexpr FColor NavMeshRenderColor_ClosedSet(255,196,0,255);
+static constexpr uint8 NavMeshRenderAlpha_Modified = 255;
+static constexpr uint8 NavMeshRenderAlpha_NonModified = 64;
 
 SIZE_T FNavTestSceneProxy::GetTypeHash() const
 {
@@ -23,15 +22,15 @@ SIZE_T FNavTestSceneProxy::GetTypeHash() const
 FNavTestSceneProxy::FNavTestSceneProxy(const UNavTestRenderingComponent* InComponent)
 	: FDebugRenderSceneProxy(InComponent)
 	, NavMeshDrawOffset(0,0,10)
-	, NavTestActor(NULL)
+	, NavTestActor(nullptr)
 {
-	if (InComponent == NULL)
+	if (InComponent == nullptr)
 	{
 		return;
 	}
 
 	NavTestActor = Cast<ANavigationTestingActor>(InComponent->GetOwner());
-	if (NavTestActor == NULL)
+	if (NavTestActor == nullptr)
 	{
 		return;
 	}
@@ -69,7 +68,7 @@ void FNavTestSceneProxy::GetDynamicMeshElements(const TArray<const FSceneView*>&
 				const FColor ClosestWallColor = FColorList::Orange;
 				const FVector BoxExtent(20, 20, 20);
 
-				FMaterialRenderProxy* const ColoredMeshInstance = new(FMemStack::Get()) FColoredMaterialRenderProxy(GEngine->DebugMeshMaterial->GetRenderProxy(), ProjectedColor);
+				const FMaterialRenderProxy* const ColoredMeshInstance = new(FMemStack::Get()) FColoredMaterialRenderProxy(GEngine->DebugMeshMaterial->GetRenderProxy(), ProjectedColor);
 				//DrawBox(PDI, FTransform(ProjectedLocation).ToMatrixNoScale(),BoxExtent, ColoredMeshInstance, SDPG_World);
 				GetSphereMesh(ProjectedLocation, BoxExtent, 10, 7, ColoredMeshInstance, SDPG_World, false, ViewIndex, Collector);
 
@@ -132,7 +131,7 @@ void FNavTestSceneProxy::GetDynamicMeshElements(const TArray<const FSceneView*>&
 
 				if (bShowDiff)
 				{
-					LineColor.A = NodeData.bModified ? NavMeshRenderAlpha_Modifed : NavMeshRenderAlpha_NonModified;
+					LineColor.A = NodeData.bModified ? NavMeshRenderAlpha_Modified : NavMeshRenderAlpha_NonModified;
 				}
 
 				FVector ParentPos(NodeData.ParentId.IsValidId() ? NodeDebug[NodeData.ParentId].Position : NodeData.Position);
@@ -298,7 +297,7 @@ FPrimitiveViewRelevance FNavTestSceneProxy::GetViewRelevance(const FSceneView* V
 	return Result;
 }
 
-uint32 FNavTestSceneProxy::GetAllocatedSize(void) const
+uint32 FNavTestSceneProxy::GetAllocatedSizeInternal() const
 {
 	int32 InternalAllocSize = 0;
 	for (TSet<FNodeDebugData>::TConstIterator It(NodeDebug); It; ++It)
@@ -355,7 +354,7 @@ void FNavTestDebugDrawDelegateHelper::UnregisterDebugDrawDelegate()
 
 void FNavTestDebugDrawDelegateHelper::DrawDebugLabels(UCanvas* Canvas, APlayerController*)
 {
-	if (NavTestActor == NULL)
+	if (NavTestActor == nullptr)
 	{
 		return;
 	}
@@ -367,7 +366,7 @@ void FNavTestDebugDrawDelegateHelper::DrawDebugLabels(UCanvas* Canvas, APlayerCo
 #if WITH_EDITORONLY_DATA && WITH_RECAST
 	if (NodeDebug.Num())
 	{
-		UFont* RenderFont = GEngine->GetSmallFont();
+		const UFont* RenderFont = GEngine->GetSmallFont();
 		for (TSet<FNavTestSceneProxy::FNodeDebugData>::TConstIterator It(NodeDebug); It; ++It)
 		{
 			const FNavTestSceneProxy::FNodeDebugData& NodeData = *It;
@@ -381,7 +380,7 @@ void FNavTestDebugDrawDelegateHelper::DrawDebugLabels(UCanvas* Canvas, APlayerCo
 				}
 				if (bShowDiff)
 				{
-					MyColor.A = NodeData.bModified ? NavMeshRenderAlpha_Modifed : NavMeshRenderAlpha_NonModified;
+					MyColor.A = NodeData.bModified ? NavMeshRenderAlpha_Modified : NavMeshRenderAlpha_NonModified;
 				}
 
 				Canvas->SetDrawColor(MyColor);
@@ -399,7 +398,7 @@ void FNavTestDebugDrawDelegateHelper::DrawDebugLabels(UCanvas* Canvas, APlayerCo
 			if (FNavTestSceneProxy::LocationInView(PathPoints[PointIndex], View))
 			{
 				const FVector PathPointLoc = Canvas->Project(PathPoints[PointIndex]);
-				UFont* RenderFont = GEngine->GetSmallFont();
+				const UFont* RenderFont = GEngine->GetSmallFont();
 				Canvas->DrawText(RenderFont, PathPointFlags[PointIndex], PathPointLoc.X, PathPointLoc.Y);
 			}
 		}
@@ -420,7 +419,7 @@ FPrimitiveSceneProxy* UNavTestRenderingComponent::CreateSceneProxy()
 	FNavTestSceneProxy* SceneProxy2 = new FNavTestSceneProxy(this);
 #if WITH_RECAST && WITH_EDITOR
 	NavTestDebugDrawDelegateHelper.InitDelegateHelper(SceneProxy2);
-	NavTestDebugDrawDelegateHelper.ReregisterDebugDrawDelegate();
+	NavTestDebugDrawDelegateHelper.RegisterDebugDrawDelegate();
 #endif
 	return SceneProxy2;
 }
@@ -461,15 +460,6 @@ FBoxSphereBounds UNavTestRenderingComponent::CalcBounds(const FTransform& LocalT
 	return FBoxSphereBounds(BoundingBox);
 }
 
-void UNavTestRenderingComponent::CreateRenderState_Concurrent(FRegisterComponentContext* Context)
-{
-	Super::CreateRenderState_Concurrent(Context);
-
-#if WITH_RECAST && WITH_EDITOR
-	NavTestDebugDrawDelegateHelper.RegisterDebugDrawDelegate();
-#endif
-}
-
 void UNavTestRenderingComponent::DestroyRenderState_Concurrent()
 {
 #if WITH_RECAST && WITH_EDITOR
diff --git a/Engine/Source/Runtime/NavigationSystem/Public/NavMesh/NavMeshRenderingComponent.h b/Engine/Source/Runtime/NavigationSystem/Public/NavMesh/NavMeshRenderingComponent.h
--- a/Engine/Source/Runtime/NavigationSystem/Public/NavMesh/NavMeshRenderingComponent.h
+++ b/Engine/Source/Runtime/NavigationSystem/Public/NavMesh/NavMeshRenderingComponent.h
@@ -2,11 +2,8 @@
 
 #pragma once
 
-#include "CoreMinimal.h"
-#include "UObject/ObjectMacros.h"
 #include "Engine/EngineTypes.h"
 #include "PrimitiveViewRelevance.h"
-#include "RenderResource.h"
 #include "MaterialShared.h"
 #include "DynamicMeshBuilder.h"
 #include "DebugRenderSceneProxy.h"
@@ -59,11 +56,10 @@ struct NAVIGATIONSYSTEM_API FNavMeshSceneProxyData : public TSharedFromThis<FNav
 
 	struct FDebugPoint
 	{
-		FDebugPoint() {}
 		FDebugPoint(const FVector& InPosition, const FColor& InColor, const float InSize) : Position(InPosition), Color(InColor), Size(InSize) {}
 		FVector Position;
 		FColor Color;
-		float Size;
+		float Size = 0.f;
 	};
 
 	TArray<FDebugRenderSceneProxy::FDebugLine> ThickLineItems;
@@ -102,7 +98,6 @@ struct NAVIGATIONSYSTEM_API FNavMeshSceneProxyData : public TSharedFromThis<FNav
 	uint32 GetAllocatedSize() const;
 
 #if WITH_RECAST
-	int32 GetDetailFlags(const ARecastNavMesh* NavMesh) const;
 	void GatherData(const ARecastNavMesh* NavMesh, int32 InNavDetailFlags, const TArray<int32>& TileSet);
 
 #if RECAST_INTERNAL_DEBUG_DATA
@@ -120,16 +115,15 @@ public:
 	virtual SIZE_T GetTypeHash() const override;
 
 	FNavMeshSceneProxy(const UPrimitiveComponent* InComponent, FNavMeshSceneProxyData* InProxyData, bool ForceToRender = false);
-	virtual ~FNavMeshSceneProxy();
+	virtual ~FNavMeshSceneProxy() override;
 
 	virtual void GetDynamicMeshElements(const TArray<const FSceneView*>& Views, const FSceneViewFamily& ViewFamily, uint32 VisibilityMap, FMeshElementCollector& Collector) const override;
 
 protected:
-	void DrawDebugBox(FPrimitiveDrawInterface* PDI, FVector const& Center, FVector const& Box, FColor const& Color) const;
 	virtual FPrimitiveViewRelevance GetViewRelevance(const FSceneView* View) const override;
 
-	virtual uint32 GetMemoryFootprint(void) const override { return sizeof(*this) + GetAllocatedSize(); }
-	uint32 GetAllocatedSize(void) const;
+	virtual uint32 GetMemoryFootprint(void) const override { return sizeof(*this) + GetAllocatedSizeInternal(); }
+	uint32 GetAllocatedSizeInternal(void) const;
 
 private:			
 	FNavMeshSceneProxyData ProxyData;
@@ -203,7 +197,6 @@ public:
 	//~ End UPrimitiveComponent Interface
 
 	//~ Begin UActorComponent Interface
-	virtual void CreateRenderState_Concurrent(FRegisterComponentContext* Context) override;
 	virtual void DestroyRenderState_Concurrent() override;
 	//~ End UActorComponent Interface
 
diff --git a/Engine/Source/Runtime/NavigationSystem/Public/NavMesh/NavTestRenderingComponent.h b/Engine/Source/Runtime/NavigationSystem/Public/NavMesh/NavTestRenderingComponent.h
--- a/Engine/Source/Runtime/NavigationSystem/Public/NavMesh/NavTestRenderingComponent.h
+++ b/Engine/Source/Runtime/NavigationSystem/Public/NavMesh/NavTestRenderingComponent.h
@@ -2,8 +2,6 @@
 
 #pragma once
 
-#include "CoreMinimal.h"
-#include "UObject/ObjectMacros.h"
 #include "PrimitiveViewRelevance.h"
 #include "Components/PrimitiveComponent.h"
 #include "DynamicMeshBuilder.h"
@@ -21,7 +19,7 @@ class FNavTestSceneProxy final : public FDebugRenderSceneProxy
 	friend class FNavTestDebugDrawDelegateHelper;
 
 public:
-	SIZE_T GetTypeHash() const override;
+	virtual SIZE_T GetTypeHash() const override;
 
 	struct FNodeDebugData
 	{
@@ -46,8 +44,6 @@ public:
 
 	FNavTestSceneProxy(const UNavTestRenderingComponent* InComponent);
 
-	~FNavTestSceneProxy() {}
-
 	virtual void GetDynamicMeshElements(const TArray<const FSceneView*>& Views, const FSceneViewFamily& ViewFamily, uint32 VisibilityMap, FMeshElementCollector& Collector) const override;
 
 	void GatherPathPoints();
@@ -61,11 +57,11 @@ public:
 
 	virtual FPrimitiveViewRelevance GetViewRelevance(const FSceneView* View) const override;
 
-	virtual uint32 GetMemoryFootprint(void) const override { return sizeof(*this) + GetAllocatedSize(); }
-
-	uint32 GetAllocatedSize(void) const;
+	virtual uint32 GetMemoryFootprint() const override { return sizeof(*this) + GetAllocatedSizeInternal(); }
 
 private:
+	uint32 GetAllocatedSizeInternal() const;
+
 	FVector3f NavMeshDrawOffset;
 
 	uint32 bShowBestPath : 1;
@@ -94,6 +90,10 @@ class FNavTestDebugDrawDelegateHelper : public FDebugDrawDelegateHelper
 	typedef FDebugDrawDelegateHelper Super;
 
 public:
+	FNavTestDebugDrawDelegateHelper(): bShowBestPath(false), bShowDiff(false)
+	{
+	}
+
 	virtual void InitDelegateHelper(const FDebugRenderSceneProxy* InSceneProxy) override
 	{
 		check(0);
@@ -105,11 +105,11 @@ public:
 	virtual void UnregisterDebugDrawDelegate() override;
 
 protected:
-	void DrawDebugLabels(UCanvas* Canvas, APlayerController*) override;
+	virtual void DrawDebugLabels(UCanvas* Canvas, APlayerController*) override;
 
 private:
 	TSet<FNavTestSceneProxy::FNodeDebugData> NodeDebug;
-	ANavigationTestingActor* NavTestActor;
+	ANavigationTestingActor* NavTestActor = nullptr;
 	TArray<FVector> PathPoints;
 	TArray<FString> PathPointFlags;
 	FSetElementId BestNodeId;
@@ -129,7 +129,6 @@ class UNavTestRenderingComponent: public UPrimitiveComponent
 
 	virtual FPrimitiveSceneProxy* CreateSceneProxy() override;
 	virtual FBoxSphereBounds CalcBounds(const FTransform &LocalToWorld) const override;
-	virtual void CreateRenderState_Concurrent(FRegisterComponentContext* Context) override;
 	virtual void DestroyRenderState_Concurrent() override;
 
 private:
