Fix crashes during the shader compilation when showing NavMesh (UE-117842).

diff --git a/Engine/Source/Runtime/Engine/Private/Materials/MaterialShared.cpp b/Engine/Source/Runtime/Engine/Private/Materials/MaterialShared.cpp
--- a/Engine/Source/Runtime/Engine/Private/Materials/MaterialShared.cpp
+++ b/Engine/Source/Runtime/Engine/Private/Materials/MaterialShared.cpp
@@ -3303,10 +3303,20 @@ FMaterialRenderProxy::FMaterialRenderProxy()
 	, DeletedFlag(0)
 	, HasVirtualTextureCallbacks(0)
 {
+	// MaterialRenderProxyMap is only used by shader compiling
+	if (!FPlatformProperties::RequiresCookedData())
+	{
+		FMaterialRenderProxy::MaterialRenderProxyMap.Add(this);
+	}
 }
 
 FMaterialRenderProxy::~FMaterialRenderProxy()
 {
+	if (!FPlatformProperties::RequiresCookedData())
+	{
+		FMaterialRenderProxy::MaterialRenderProxyMap.Remove(this);
+	}
+
 	if(IsInitialized())
 	{
 		check(IsInRenderingThread());
@@ -3323,22 +3333,8 @@ FMaterialRenderProxy::~FMaterialRenderProxy()
 	DeletedFlag = 1;
 }
 
-void FMaterialRenderProxy::InitDynamicRHI()
-{
-	// MaterialRenderProxyMap is only used by shader compiling
-	if (!FPlatformProperties::RequiresCookedData())
-	{
-		FMaterialRenderProxy::MaterialRenderProxyMap.Add(this);
-	}
-}
-
 void FMaterialRenderProxy::ReleaseDynamicRHI()
 {
-	if (!FPlatformProperties::RequiresCookedData())
-	{
-		FMaterialRenderProxy::MaterialRenderProxyMap.Remove(this);
-	}
-
 	DeferredUniformExpressionCacheRequests.Remove(this);
 
 	InvalidateUniformExpressionCache(true);
diff --git a/Engine/Source/Runtime/Engine/Public/MaterialShared.h b/Engine/Source/Runtime/Engine/Public/MaterialShared.h
--- a/Engine/Source/Runtime/Engine/Public/MaterialShared.h
+++ b/Engine/Source/Runtime/Engine/Public/MaterialShared.h
@@ -27,6 +27,7 @@
 #include "Serialization/ArchiveProxy.h"
 #include "MaterialSceneTextureId.h"
 #include "VirtualTexturing.h"
+#include "Templates/UnrealTemplate.h"
 
 #include "Shader/Preshader.h"
 
@@ -2258,7 +2259,7 @@ class USubsurfaceProfile;
 /**
  * A material render proxy used by the renderer.
  */
-class FMaterialRenderProxy : public FRenderResource
+class FMaterialRenderProxy : public FRenderResource, public FNoncopyable
 {
 public:
 
@@ -2351,7 +2352,6 @@ public:
 	}
 
 	// FRenderResource interface.
-	ENGINE_API virtual void InitDynamicRHI() override;
 	ENGINE_API virtual void ReleaseDynamicRHI() override;
 	ENGINE_API virtual void ReleaseResource() override;
 
diff --git a/Engine/Source/Runtime/NavigationSystem/Private/NavMesh/NavMeshRenderingComponent.cpp b/Engine/Source/Runtime/NavigationSystem/Private/NavMesh/NavMeshRenderingComponent.cpp
--- a/Engine/Source/Runtime/NavigationSystem/Private/NavMesh/NavMeshRenderingComponent.cpp
+++ b/Engine/Source/Runtime/NavigationSystem/Private/NavMesh/NavMeshRenderingComponent.cpp
@@ -1092,7 +1092,7 @@ FNavMeshSceneProxy::FNavMeshSceneProxy(const UPrimitiveComponent* InComponent, F
 		return;
 	}
 
-	MeshColors.Reserve(NumberOfMeshes);
+	MeshColors.Reserve(NumberOfMeshes + 1);	// we add one more proxy after the loop
 	MeshBatchElements.Reserve(NumberOfMeshes);
 	const FMaterialRenderProxy* ParentMaterial = GEngine->DebugMeshMaterial->GetRenderProxy();
 
@@ -1109,7 +1109,7 @@ FNavMeshSceneProxy::FNavMeshSceneProxy(const UPrimitiveComponent* InComponent, F
 		Element.IndexBuffer = &IndexBuffer;
 		MeshBatchElements.Add(Element);
 
-		MeshColors.Add(FColoredMaterialRenderProxy(ParentMaterial, CurrentMeshBuilder.ClusterColor));
+		MeshColors.Add(MakeUnique<FColoredMaterialRenderProxy>(ParentMaterial, CurrentMeshBuilder.ClusterColor));
 
 		const int32 VertexIndexOffset = Vertices.Num();
 		Vertices.Append(CurrentMeshBuilder.Vertices);
@@ -1127,7 +1127,7 @@ FNavMeshSceneProxy::FNavMeshSceneProxy(const UPrimitiveComponent* InComponent, F
 		}
 	}
 
-	MeshColors.Add(FColoredMaterialRenderProxy(ParentMaterial, NavMeshRenderColor_PathCollidingGeom));
+	MeshColors.Add(MakeUnique<FColoredMaterialRenderProxy>(ParentMaterial, NavMeshRenderColor_PathCollidingGeom));
 
 	if (Vertices.Num())
 	{
@@ -1216,7 +1216,7 @@ void FNavMeshSceneProxy::GetDynamicMeshElements(const TArray<const FSceneView*>&
 
 					Mesh.bWireframe = false;
 					Mesh.VertexFactory = &VertexFactory;
-					Mesh.MaterialRenderProxy = &MeshColors[Index];
+					Mesh.MaterialRenderProxy = MeshColors[Index].Get();
 					Mesh.ReverseCulling = IsLocalToWorldDeterminantNegative();
 					Mesh.Type = PT_TriangleList;
 					Mesh.DepthPriorityGroup = SDPG_World;
@@ -1395,7 +1395,7 @@ uint32 FNavMeshSceneProxy::GetAllocatedSizeInternal() const
 		VertexBuffers.PositionVertexBuffer.GetNumVertices() * VertexBuffers.PositionVertexBuffer.GetStride() +
 		VertexBuffers.StaticMeshVertexBuffer.GetResourceSize() +
 		VertexBuffers.ColorVertexBuffer.GetNumVertices() * VertexBuffers.ColorVertexBuffer.GetStride() +
-		MeshColors.GetAllocatedSize() +
+		MeshColors.GetAllocatedSize() + MeshColors.Num() * sizeof(FColoredMaterialRenderProxy) +
 		MeshBatchElements.GetAllocatedSize();
 }
 
diff --git a/Engine/Source/Runtime/NavigationSystem/Public/NavMesh/NavMeshRenderingComponent.h b/Engine/Source/Runtime/NavigationSystem/Public/NavMesh/NavMeshRenderingComponent.h
--- a/Engine/Source/Runtime/NavigationSystem/Public/NavMesh/NavMeshRenderingComponent.h
+++ b/Engine/Source/Runtime/NavigationSystem/Public/NavMesh/NavMeshRenderingComponent.h
@@ -13,6 +13,7 @@
 #include "Math/GenericOctree.h"
 #include "StaticMeshResources.h"
 #include "NavigationSystemTypes.h"
+#include "Templates/UnrealTemplate.h"
 #include "NavMeshRenderingComponent.generated.h"
 
 class APlayerController;
@@ -108,7 +109,7 @@ struct NAVIGATIONSYSTEM_API FNavMeshSceneProxyData : public TSharedFromThis<FNav
 };
 
 // exported to API for GameplayDebugger module
-class NAVIGATIONSYSTEM_API FNavMeshSceneProxy final : public FDebugRenderSceneProxy
+class NAVIGATIONSYSTEM_API FNavMeshSceneProxy final : public FDebugRenderSceneProxy, public FNoncopyable
 {
 	friend class FNavMeshDebugDrawDelegateHelper;
 public:
@@ -132,7 +133,7 @@ private:
 	FStaticMeshVertexBuffers VertexBuffers;
 	FLocalVertexFactory VertexFactory;
 
-	TArray<FColoredMaterialRenderProxy> MeshColors;
+	TArray<TUniquePtr<FColoredMaterialRenderProxy>> MeshColors;
 	TArray<FMeshBatchElement> MeshBatchElements;
 
 	FDebugDrawDelegate DebugTextDrawingDelegate;
