[Navigation] Stop dirtying packages on any navigation generation finish.
That change was made to help with a user experience issue but it's causing more issues than benefit.

diff --git a/Engine/Source/Runtime/NavigationSystem/Private/NavigationSystem.cpp b/Engine/Source/Runtime/NavigationSystem/Private/NavigationSystem.cpp
--- a/Engine/Source/Runtime/NavigationSystem/Private/NavigationSystem.cpp
+++ b/Engine/Source/Runtime/NavigationSystem/Private/NavigationSystem.cpp
@@ -3784,9 +3784,9 @@ void UNavigationSystemV1::OnNavigationGenerationFinished(ANavigationData& NavDat
 	OnNavigationGenerationFinishedDelegate.Broadcast(&NavData);
 
 #if WITH_EDITOR
-	if (GetWorld()->IsGameWorld() == false && NavData.IsBuildingOnLoad() == false)
+	if (GetWorld()->IsGameWorld() == false)
 	{
-		NavData.MarkPackageDirty();
+		UE_LOG(LogNavigationDataBuild, Log, TEXT("Navigation data generation finished in %s."), *NavData.GetPackage()->GetName());
 	}
 
 	// Reset bIsBuildingOnLoad
