Remove references to UE4

diff --git a/Engine/Source/Runtime/NavigationSystem/Private/NavMesh/PImplRecastNavMesh.cpp b/Engine/Source/Runtime/NavigationSystem/Private/NavMesh/PImplRecastNavMesh.cpp
--- a/Engine/Source/Runtime/NavigationSystem/Private/NavMesh/PImplRecastNavMesh.cpp
+++ b/Engine/Source/Runtime/NavigationSystem/Private/NavMesh/PImplRecastNavMesh.cpp
@@ -391,15 +391,10 @@ void FPImplRecastNavMesh::Serialize( FArchive& Ar, int32 NavMeshVersion )
 					dtMeshHeader* const TileHeader = (dtMeshHeader*)TileData;
 					dtFree(TileHeader);
 
-					//
-					if (Ar.UEVer() >= VER_UE4_ADD_MODIFIERS_RUNTIME_GENERATION && 
-						(Ar.EngineVer().GetMajor() != 4 || Ar.EngineVer().GetMinor() != 7)) // Merged package from 4.7 branch
-					{
-						unsigned char* ComressedTileData = NULL;
-						int32 CompressedTileDataSize = 0;
-						SerializeCompressedTileCacheData(Ar, NavMeshVersion, ComressedTileData, CompressedTileDataSize);
-						dtFree(ComressedTileData);
-					}
+					unsigned char* ComressedTileData = NULL;
+					int32 CompressedTileDataSize = 0;
+					SerializeCompressedTileCacheData(Ar, NavMeshVersion, ComressedTileData, CompressedTileDataSize);
+					dtFree(ComressedTileData);
 				}
 			}
 		}
@@ -433,18 +428,14 @@ void FPImplRecastNavMesh::Serialize( FArchive& Ar, int32 NavMeshVersion )
 					DetourNavMesh->addTile(TileData, TileDataSize, DT_TILE_FREE_DATA, TileRef, NULL);
 
 					// Serialize compressed tile cache layer
-					if (Ar.UEVer() >= VER_UE4_ADD_MODIFIERS_RUNTIME_GENERATION &&
-						(Ar.EngineVer().GetMajor() != 4 || Ar.EngineVer().GetMinor() != 7)) // Merged package from 4.7 branch
+					uint8* ComressedTileData = nullptr;
+					int32 CompressedTileDataSize = 0;
+					SerializeCompressedTileCacheData(Ar, NavMeshVersion, ComressedTileData, CompressedTileDataSize);
+					
+					if (CompressedTileDataSize > 0)
 					{
-						uint8* ComressedTileData = nullptr;
-						int32 CompressedTileDataSize = 0;
-						SerializeCompressedTileCacheData(Ar, NavMeshVersion, ComressedTileData, CompressedTileDataSize);
-						
-						if (CompressedTileDataSize > 0)
-						{
-							AddTileCacheLayer(TileHeader->x, TileHeader->y, TileHeader->layer,
-								FNavMeshTileData(ComressedTileData, CompressedTileDataSize, TileHeader->layer, Recast2UnrealBox(TileHeader->bmin, TileHeader->bmax)));
-						}
+						AddTileCacheLayer(TileHeader->x, TileHeader->y, TileHeader->layer,
+							FNavMeshTileData(ComressedTileData, CompressedTileDataSize, TileHeader->layer, Recast2UnrealBox(TileHeader->bmin, TileHeader->bmax)));
 					}
 				}
 			}
@@ -1799,7 +1790,7 @@ bool FPImplRecastNavMesh::GetPolyCenter(NavNodeRef PolyID, FVector& OutCenter) c
 			Center[1] *= InvCount;
 			Center[2] *= InvCount;
 
-			// convert output to UE4 coords
+			// convert output to UE coords
 			OutCenter = Recast2UnrVector(Center);
 
 			return true;
@@ -1822,7 +1813,7 @@ bool FPImplRecastNavMesh::GetPolyVerts(NavNodeRef PolyID, TArray<FVector>& OutVe
 			// flush and pre-size output array
 			OutVerts.Reset(Poly->vertCount);
 
-			// convert to UE4 coords and copy verts into output array 
+			// convert to UE coords and copy verts into output array 
 			for (uint32 VertIdx=0; VertIdx < Poly->vertCount; ++VertIdx)
 			{
 				const float* V = &Tile->verts[Poly->verts[VertIdx]*3];
diff --git a/Engine/Source/Runtime/NavigationSystem/Private/NavMesh/RecastNavMeshDataChunk.cpp b/Engine/Source/Runtime/NavigationSystem/Private/NavMesh/RecastNavMeshDataChunk.cpp
--- a/Engine/Source/Runtime/NavigationSystem/Private/NavMesh/RecastNavMeshDataChunk.cpp
+++ b/Engine/Source/Runtime/NavigationSystem/Private/NavMesh/RecastNavMeshDataChunk.cpp
@@ -140,11 +140,7 @@ void URecastNavMeshDataChunk::SerializeRecastData(FArchive& Ar, int32 NavMeshVer
 				// Load compressed tile cache layer
 				int32 TileCacheDataSize = 0;
 				uint8* TileCacheRawData = nullptr;
-				if (Ar.UEVer() >= VER_UE4_ADD_MODIFIERS_RUNTIME_GENERATION && 
-					(Ar.EngineVer().GetMajor() != 4 || Ar.EngineVer().GetMinor() != 7)) // Merged package from 4.7 branch
-				{
-					FPImplRecastNavMesh::SerializeCompressedTileCacheData(Ar, NavMeshVersion, TileCacheRawData, TileCacheDataSize); //allocates TileCacheRawData on load
-				}
+				FPImplRecastNavMesh::SerializeCompressedTileCacheData(Ar, NavMeshVersion, TileCacheRawData, TileCacheDataSize); //allocates TileCacheRawData on load
 				
 				// We are owner of tile raw data
 				FRecastTileData TileData(TileDataSize, TileRawData, TileCacheDataSize, TileCacheRawData);
diff --git a/Engine/Source/Runtime/NavigationSystem/Private/NavigationData.cpp b/Engine/Source/Runtime/NavigationSystem/Private/NavigationData.cpp
--- a/Engine/Source/Runtime/NavigationSystem/Private/NavigationData.cpp
+++ b/Engine/Source/Runtime/NavigationSystem/Private/NavigationData.cpp
@@ -219,12 +219,6 @@ void ANavigationData::PostLoad()
 {
 	Super::PostLoad();
 
-	if ((GetLinkerUEVersion() < VER_UE4_ADD_MODIFIERS_RUNTIME_GENERATION) &&
-		(RuntimeGeneration == ERuntimeGenerationType::LegacyGeneration))
-	{
-		RuntimeGeneration = bRebuildAtRuntime_DEPRECATED ? ERuntimeGenerationType::Dynamic : ERuntimeGenerationType::Static;
-	}
-
 	InstantiateAndRegisterRenderingComponent();
 
 	bNetLoadOnClient = FNavigationSystem::ShouldLoadNavigationOnClient(*this);
