[Navigation] Fix crash occuring when building tiles with no geometry and MaskFillCollisionUnderneathforNavMesh enabled on a navmodifier.

diff --git a/Engine/Source/Runtime/NavigationSystem/Private/NavMesh/RecastNavMeshGenerator.cpp b/Engine/Source/Runtime/NavigationSystem/Private/NavMesh/RecastNavMeshGenerator.cpp
--- a/Engine/Source/Runtime/NavigationSystem/Private/NavMesh/RecastNavMeshGenerator.cpp
+++ b/Engine/Source/Runtime/NavigationSystem/Private/NavMesh/RecastNavMeshGenerator.cpp
@@ -3952,6 +3952,11 @@ void FRecastTileGenerator::ComputeRasterizationMasks(FNavMeshBuildContext& Build
 	{
 		if (ModifierElement.bMaskFillCollisionUnderneathForNavmesh)
 		{
+			if (RasterContext.SolidHF == nullptr)
+			{
+				return;
+			}
+
 			const int32 Mask = ~RC_PROJECT_TO_BOTTOM;
 			for (const FAreaNavModifier& ModifierArea : ModifierElement.Areas)
 			{
