Fixed a crash in FPImplRecastNavMesh::FindMoveAlongSurface due to not initializing node pools correctly as moveAlongSurface uses them.

diff --git a/Engine/Source/Runtime/NavigationSystem/Private/NavMesh/PImplRecastNavMesh.cpp b/Engine/Source/Runtime/NavigationSystem/Private/NavMesh/PImplRecastNavMesh.cpp
--- a/Engine/Source/Runtime/NavigationSystem/Private/NavMesh/PImplRecastNavMesh.cpp
+++ b/Engine/Source/Runtime/NavigationSystem/Private/NavMesh/PImplRecastNavMesh.cpp
@@ -1480,9 +1480,7 @@ bool FPImplRecastNavMesh::FindMoveAlongSurface(const FNavLocation& StartLocation
 	}
 
 	FRecastSpeciaLinkFilter LinkFilter(FNavigationSystem::GetCurrent<UNavigationSystemV1>(NavMeshOwner->GetWorld()), Owner);
-	// using 0 as NumNodes since findNearestPoly2D, being the only dtNavMeshQuery
-	// function we're using, is not utilizing m_nodePool
-	INITIALIZE_NAVQUERY(NavQuery, /*NumNodes=*/0, LinkFilter);
+	INITIALIZE_NAVQUERY(NavQuery, Filter.GetMaxSearchNodes(), LinkFilter);
 
 	const dtQueryFilter* QueryFilter = ((const FRecastQueryFilter*)(Filter.GetImplementation()))->GetAsDetourQueryFilter();
 	ensure(QueryFilter);
