Fix Navigation System overzealously re-building DynamicModifiersOnly-type navigation data during OnWorldInitDone when AutoGenerateNavigationData is disabled. Re-factored the check for dynamic navmesh generation type to its own function and replaced the code in RebuildAll, CanRebuildDirtyNavigation, and OnWorldInitDone.

diff --git a/Engine/Source/Runtime/NavigationSystem/Private/NavigationSystem.cpp b/Engine/Source/Runtime/NavigationSystem/Private/NavigationSystem.cpp
--- a/Engine/Source/Runtime/NavigationSystem/Private/NavigationSystem.cpp
+++ b/Engine/Source/Runtime/NavigationSystem/Private/NavigationSystem.cpp
@@ -897,12 +897,10 @@ void UNavigationSystemV1::OnWorldInitDone(FNavigationSystemRunMode Mode)
 		else
 		{
 			const bool bIsBuildLocked = IsNavigationBuildingLocked();
-			const bool bAllowRebuild = !bIsBuildLocked && GetIsAutoUpdateEnabled();
+			const bool bCanRebuild = !bIsBuildLocked && GetIsAutoUpdateEnabled();
 
 			if (GetDefaultNavDataInstance(FNavigationSystem::DontCreate) != NULL)
 			{
-				const bool bIsInGame = World->IsGameWorld();
-
 				// trigger navmesh update
 				for (TActorIterator<ANavigationData> It(World); It; ++It)
 				{
@@ -914,7 +912,9 @@ void UNavigationSystemV1::OnWorldInitDone(FNavigationSystemRunMode Mode)
 
 						if (Result == RegistrationSuccessful)
 						{
-							if (bAllowRebuild && (!bIsInGame || NavData->SupportsRuntimeGeneration()))
+							// allowing full rebuild of the entire navmesh only for the fully dynamic generation modes
+							// other modes partly rely on the serialized data and full rebuild would wipe it out
+							if (bCanRebuild && IsAllowedToRebuild())
 							{
 								NavData->RebuildAll();
 							}
@@ -4134,6 +4134,13 @@ void UNavigationSystemV1::LogNavDataRegistrationResult(ERegistrationResult InRes
 	}
 }
 
+bool UNavigationSystemV1::IsAllowedToRebuild() const
+{
+	const UWorld* World = GetWorld();
+	
+	return World && (!World->IsGameWorld() || GetRuntimeGenerationType() == ERuntimeGenerationType::Dynamic);
+}
+
 //----------------------------------------------------------------------//
 // Blueprint functions
 //----------------------------------------------------------------------//
diff --git a/Engine/Source/Runtime/NavigationSystem/Public/NavigationSystem.h b/Engine/Source/Runtime/NavigationSystem/Public/NavigationSystem.h
--- a/Engine/Source/Runtime/NavigationSystem/Public/NavigationSystem.h
+++ b/Engine/Source/Runtime/NavigationSystem/Public/NavigationSystem.h
@@ -1187,6 +1187,11 @@ protected:
 	ERuntimeGenerationType GetRuntimeGenerationType() const;
 
 	void LogNavDataRegistrationResult(ERegistrationResult);
+
+	/** Whether Navigation System is allowed to rebuild the navmesh
+	 *  Depends on runtime generation settings of each navigation data, always true in the editor
+	 */
+	bool IsAllowedToRebuild() const;
 	
 	//----------------------------------------------------------------------//
 	// new stuff
