[DebugDrawDelegateHelper] rollback of delegate unregistration issue previous fix and implemented a different approach for cases where proxy is not created:
In some code paths for loaded actors the call order might be different since primitive registration gets deferred (i.e. FRegisterComponentContext != nullptr).

Case 1 (normal flow): FRegisterComponentContext == nullptr
 > UPrimitiveComponent::CreateRenderState_Concurrent : calls CreateSceneProxy
 > RegisterDebugDrawDelegate
 > UnregisterDebugDrawDelegate  ==> works fine

Case 2 (deferred AddPrimitive): FRegisterComponentContext != nullptr
 > UPrimitiveComponent::CreateRenderState_Concurrent : defers CreateSceneProxy
 > RegisterDebugDrawDelegate (skip register since not init)
 > CreateSceneProxy
 > UnregisterDebugDrawDelegate  ==>  ensures in UnregisterDebugDrawDelegate

With this new version `CreateRenderState_Concurrent` calls `RequestRegisterDebugDrawDelegate` that will take care of registering the delegate immediately or mark it for deferred execution. An explicit call to `ProcessDeferredRegister` is then required from `CreateSceneProxy` to implement proper flow:

Case 2 (fixed) (deferred AddPrimitive + deferred Register): FRegisterComponentContext != nullptr
 > UPrimitiveComponent::CreateRenderState_Concurrent: defers CreateSceneProxy
 > RequestRegisterDebugDrawDelegate: defers RegisterDebugDrawDelegate
 > CreateSceneProxy: calls ProcessDeferredRegister & RegisterDebugDrawDelegate
 > UnregisterDebugDrawDelegate  ==> works fine

Bonus:
- New DebugDrawComponent abstract class that is now used to share common functionalities between existing debug related components. Takes care of DelegateHelper registration flow to render text on screen from different sources.
- FDebugDrawDelegateHelper::InitDelegateHelper is no longer virtual and derived classes requiring extra data from their associated scene proxy should set it up from their overriden CreateSceneProxy
- FDebugDrawDelegateHelper derived classes should not override RegisterDebugDrawDelegate/UnregisterDebugDrawDelegate to use their `DrawDebugLabels` method since this is already a virtual method that would get called from the base class.
- Fixed a few ViewFlag members in the SceneProxy so the DelegateHelper base class can behave as expected

diff --git a/Engine/Plugins/AI/MassAI/Source/MassAIMovementEditor/Private/MassMovementTestingActor.cpp b/Engine/Plugins/AI/MassAI/Source/MassAIMovementEditor/Private/MassMovementTestingActor.cpp
--- a/Engine/Plugins/AI/MassAI/Source/MassAIMovementEditor/Private/MassMovementTestingActor.cpp
+++ b/Engine/Plugins/AI/MassAI/Source/MassAIMovementEditor/Private/MassMovementTestingActor.cpp
@@ -189,14 +189,9 @@ void UMassMovementTestingComponent::ClearPinnedLane()
 	UpdateTests();
 }
 
-#if !UE_BUILD_SHIPPING
-void UMassMovementTestingComponent::DestroyRenderState_Concurrent()
-{
-	DebugDrawDelegateHelper.UnregisterDebugDrawDelegate();
-	Super::DestroyRenderState_Concurrent();
-}
+#if UE_ENABLE_DEBUG_DRAWING
 
-FPrimitiveSceneProxy* UMassMovementTestingComponent::CreateSceneProxy()
+FDebugRenderSceneProxy* UMassMovementTestingComponent::CreateDebugSceneProxy()
 {
 	FMassMovementTestingSceneProxy* DebugProxy = new FMassMovementTestingSceneProxy(*this);
 
@@ -289,12 +284,9 @@ FPrimitiveSceneProxy* UMassMovementTestingComponent::CreateSceneProxy()
 		}
 	}
 
-	DebugDrawDelegateHelper.InitDelegateHelper(DebugProxy);
-	DebugDrawDelegateHelper.RegisterDebugDrawDelegate();
-
 	return DebugProxy;
 }
-#endif // !UE_BUILD_SHIPPING
+#endif // UE_ENABLE_DEBUG_DRAWING
 
 //////////////////////////////////////////////////////////////////////////
 // UMassMovementTestingComponent
diff --git a/Engine/Plugins/AI/MassAI/Source/MassAIMovementEditor/Private/MassMovementTestingActor.h b/Engine/Plugins/AI/MassAI/Source/MassAIMovementEditor/Private/MassMovementTestingActor.h
--- a/Engine/Plugins/AI/MassAI/Source/MassAIMovementEditor/Private/MassMovementTestingActor.h
+++ b/Engine/Plugins/AI/MassAI/Source/MassAIMovementEditor/Private/MassMovementTestingActor.h
@@ -3,7 +3,7 @@
 #pragma once
 
 #include "GameFramework/Actor.h"
-#include "Components/PrimitiveComponent.h"
+#include "Debug/DebugDrawComponent.h"
 #include "DebugRenderSceneProxy.h"
 #include "MassAvoidanceProcessors.h"
 #include "ZoneGraphTypes.h"
@@ -28,8 +28,8 @@ public:
 #endif
 
 /** Component for testing MassMovement functionality. */
-UCLASS(ClassGroup = Custom, hidecategories = (Physics, Collision, Rendering, Cooking, Lighting, Navigation, Tags, HLOD, Mobile, AssetUserData, Activation))
-class MASSAIMOVEMENTEDITOR_API UMassMovementTestingComponent : public UPrimitiveComponent
+UCLASS(ClassGroup = Debug)
+class MASSAIMOVEMENTEDITOR_API UMassMovementTestingComponent : public UDebugDrawComponent
 {
 	GENERATED_BODY()
 public:
@@ -44,11 +44,8 @@ public:
 
 	virtual FBoxSphereBounds CalcBounds(const FTransform& LocalToWorld) const override;
 
-#if !UE_BUILD_SHIPPING
-	//~ Begin UPrimitiveComponent Interface.
-	virtual void DestroyRenderState_Concurrent() override;
-	virtual FPrimitiveSceneProxy* CreateSceneProxy() override;
-	//~ End UPrimitiveComponent Interface.
+#if UE_ENABLE_DEBUG_DRAWING
+	virtual FDebugRenderSceneProxy* CreateDebugSceneProxy() override;
 #endif
 
 	void UpdateTests();
@@ -62,10 +59,6 @@ protected:
 #endif
 	void OnZoneGraphDataChanged(const AZoneGraphData* ZoneGraphData);
 
-#if !UE_BUILD_SHIPPING && !UE_BUILD_TEST
-	FDebugDrawDelegateHelper DebugDrawDelegateHelper;
-#endif
-
 #if WITH_EDITOR
 	FDelegateHandle OnDataChangedHandle;
 #endif
diff --git a/Engine/Plugins/Runtime/ZoneGraphAnnotations/Source/ZoneGraphAnnotations/Private/ZoneGraphAnnotationComponent.cpp b/Engine/Plugins/Runtime/ZoneGraphAnnotations/Source/ZoneGraphAnnotations/Private/ZoneGraphAnnotationComponent.cpp
--- a/Engine/Plugins/Runtime/ZoneGraphAnnotations/Source/ZoneGraphAnnotations/Private/ZoneGraphAnnotationComponent.cpp
+++ b/Engine/Plugins/Runtime/ZoneGraphAnnotations/Source/ZoneGraphAnnotations/Private/ZoneGraphAnnotationComponent.cpp
@@ -13,7 +13,7 @@
 #endif
 
 
-#if !UE_BUILD_SHIPPING && !UE_BUILD_TEST
+#if UE_ENABLE_DEBUG_DRAWING
 //////////////////////////////////////////////////////////////////////////
 // FZoneGraphAnnotationSceneProxy
 
@@ -46,13 +46,13 @@ uint32 FZoneGraphAnnotationSceneProxy::GetMemoryFootprint(void) const
 	return sizeof(*this) + FDebugRenderSceneProxy::GetAllocatedSize();
 }
 
-#endif // !UE_BUILD_SHIPPING && !UE_BUILD_TEST
+#endif // UE_ENABLE_DEBUG_DRAWING
 
 
 //////////////////////////////////////////////////////////////////////////
 // UZoneGraphAnnotationComponent
 UZoneGraphAnnotationComponent::UZoneGraphAnnotationComponent(const FObjectInitializer& ObjectInitializer)
-	: UPrimitiveComponent(ObjectInitializer)
+	: Super(ObjectInitializer)
 {
 	SetCollisionEnabled(ECollisionEnabled::NoCollision);
 
@@ -108,7 +108,7 @@ void UZoneGraphAnnotationComponent::OnRegister()
 		OnPostWorldInitDelegateHandle = FWorldDelegates::OnPostWorldInitialization.AddUObject(this, &UZoneGraphAnnotationComponent::OnPostWorldInit);
 	}
 
-#if !UE_BUILD_SHIPPING && !UE_BUILD_TEST
+#if UE_ENABLE_DEBUG_DRAWING
 	CanvasDebugDrawDelegateHandle = UDebugDrawService::Register(TEXT("ZoneGraph"), FDebugDrawDelegate::CreateUObject(this, &UZoneGraphAnnotationComponent::DebugDrawCanvas));
 #endif
 }
@@ -178,7 +178,7 @@ void UZoneGraphAnnotationComponent::OnUnregister()
 	UE::ZoneGraphDelegates::OnPostZoneGraphDataAdded.Remove(OnPostZoneGraphDataAddedHandle);
 	UE::ZoneGraphDelegates::OnPreZoneGraphDataRemoved.Remove(OnPreZoneGraphDataRemovedHandle);
 
-#if !UE_BUILD_SHIPPING && !UE_BUILD_TEST
+#if UE_ENABLE_DEBUG_DRAWING
 	UDebugDrawService::Unregister(CanvasDebugDrawDelegateHandle);
 #endif
 	
@@ -207,35 +207,17 @@ void UZoneGraphAnnotationComponent::OnPreZoneGraphDataRemoved(const AZoneGraphDa
 	PreZoneGraphDataRemoved(*ZoneGraphData);
 }
 
-void UZoneGraphAnnotationComponent::DestroyRenderState_Concurrent()
+#if UE_ENABLE_DEBUG_DRAWING
+FDebugRenderSceneProxy* UZoneGraphAnnotationComponent::CreateDebugSceneProxy()
 {
-#if !UE_BUILD_SHIPPING && !UE_BUILD_TEST
-	DebugDrawDelegateHelper.UnregisterDebugDrawDelegate();
-#endif
-
-	Super::DestroyRenderState_Concurrent();
-}
-
-FPrimitiveSceneProxy* UZoneGraphAnnotationComponent::CreateSceneProxy()
-{
-#if !UE_BUILD_SHIPPING && !UE_BUILD_TEST
-	
 	FZoneGraphAnnotationSceneProxy* DebugProxy = new FZoneGraphAnnotationSceneProxy(*this);
-
-	if(bEnableDebugDrawing)
+	if (bEnableDebugDrawing)
 	{
 		DebugDraw(DebugProxy);
 	}
-
-	DebugDrawDelegateHelper.InitDelegateHelper(DebugProxy);
-	DebugDrawDelegateHelper.RegisterDebugDrawDelegate();
-	
 	return DebugProxy;
-
-#else
-	return nullptr;
-#endif //!UE_BUILD_SHIPPING && !UE_BUILD_TEST
 }
+#endif // UE_ENABLE_DEBUG_DRAWING
 
 FBoxSphereBounds UZoneGraphAnnotationComponent::CalcBounds(const FTransform& LocalToWorld) const
 {
diff --git a/Engine/Plugins/Runtime/ZoneGraphAnnotations/Source/ZoneGraphAnnotations/Private/ZoneGraphAnnotationTestingActor.cpp b/Engine/Plugins/Runtime/ZoneGraphAnnotations/Source/ZoneGraphAnnotations/Private/ZoneGraphAnnotationTestingActor.cpp
--- a/Engine/Plugins/Runtime/ZoneGraphAnnotations/Source/ZoneGraphAnnotations/Private/ZoneGraphAnnotationTestingActor.cpp
+++ b/Engine/Plugins/Runtime/ZoneGraphAnnotations/Source/ZoneGraphAnnotations/Private/ZoneGraphAnnotationTestingActor.cpp
@@ -22,7 +22,7 @@ FBoxSphereBounds UZoneGraphAnnotationTestingComponent::CalcBounds(const FTransfo
 {
 	FBox CombinedBounds(ForceInit);
 	
-#if !UE_BUILD_SHIPPING && !UE_BUILD_TEST
+#if UE_ENABLE_DEBUG_DRAWING
 	for (const UZoneGraphAnnotationTest* Test : Tests)
 	{
 		if (Test != nullptr)
@@ -68,7 +68,7 @@ void UZoneGraphAnnotationTestingComponent::OnRegister()
 			}
 		}
 	
-#if !UE_BUILD_SHIPPING && !UE_BUILD_TEST
+#if UE_ENABLE_DEBUG_DRAWING
 		CanvasDebugDrawDelegateHandle = UDebugDrawService::Register(TEXT("ZoneGraph"), FDebugDrawDelegate::CreateUObject(this, &UZoneGraphAnnotationTestingComponent::DebugDrawCanvas));
 #endif
 	}
@@ -78,7 +78,7 @@ void UZoneGraphAnnotationTestingComponent::OnUnregister()
 {
 	if (!HasAnyFlags(RF_ClassDefaultObject))
 	{
-#if !UE_BUILD_SHIPPING && !UE_BUILD_TEST
+#if UE_ENABLE_DEBUG_DRAWING
 		UDebugDrawService::Unregister(CanvasDebugDrawDelegateHandle);
 #endif
 	}
@@ -86,34 +86,14 @@ void UZoneGraphAnnotationTestingComponent::OnUnregister()
 	Super::OnUnregister();
 }
 
-void UZoneGraphAnnotationTestingComponent::DestroyRenderState_Concurrent()
+#if UE_ENABLE_DEBUG_DRAWING
+FDebugRenderSceneProxy* UZoneGraphAnnotationTestingComponent::CreateDebugSceneProxy()
 {
-#if !UE_BUILD_SHIPPING && !UE_BUILD_TEST
-	DebugDrawDelegateHelper.UnregisterDebugDrawDelegate();
-#endif
-
-	Super::DestroyRenderState_Concurrent();
-}
-
-FPrimitiveSceneProxy* UZoneGraphAnnotationTestingComponent::CreateSceneProxy()
-{
-#if !UE_BUILD_SHIPPING && !UE_BUILD_TEST
-	
 	FZoneGraphAnnotationSceneProxy* DebugProxy = new FZoneGraphAnnotationSceneProxy(*this, FZoneGraphAnnotationSceneProxy::WireMesh);
-
 	DebugDraw(DebugProxy);
-
-	DebugDrawDelegateHelper.InitDelegateHelper(DebugProxy);
-	DebugDrawDelegateHelper.RegisterDebugDrawDelegate();
-
 	return DebugProxy;
-
-#else
-	return nullptr;
-#endif //!UE_BUILD_SHIPPING && !UE_BUILD_TEST
 }
 
-#if !UE_BUILD_SHIPPING && !UE_BUILD_TEST
 void UZoneGraphAnnotationTestingComponent::DebugDraw(FDebugRenderSceneProxy* DebugProxy)
 {
 	for (UZoneGraphAnnotationTest* Test : Tests)
@@ -135,7 +115,7 @@ void UZoneGraphAnnotationTestingComponent::DebugDrawCanvas(UCanvas* Canvas, APla
 		}
 	}
 }
-#endif // !UE_BUILD_SHIPPING && !UE_BUILD_TEST
+#endif // UE_ENABLE_DEBUG_DRAWING
 
 /////////////////////////////////////////////////////
 // AZoneGraphAnnotationTestingActor
diff --git a/Engine/Plugins/Runtime/ZoneGraphAnnotations/Source/ZoneGraphAnnotations/Public/ZoneGraphAnnotationComponent.h b/Engine/Plugins/Runtime/ZoneGraphAnnotations/Source/ZoneGraphAnnotations/Public/ZoneGraphAnnotationComponent.h
--- a/Engine/Plugins/Runtime/ZoneGraphAnnotations/Source/ZoneGraphAnnotations/Public/ZoneGraphAnnotationComponent.h
+++ b/Engine/Plugins/Runtime/ZoneGraphAnnotations/Source/ZoneGraphAnnotations/Public/ZoneGraphAnnotationComponent.h
@@ -2,7 +2,7 @@
 
 #pragma once
 
-#include "Components/PrimitiveComponent.h"
+#include "Debug/DebugDrawComponent.h"
 #include "DebugRenderSceneProxy.h"
 #include "ZoneGraphTypes.h"
 #include "ZoneGraphAnnotationComponent.generated.h"
@@ -15,7 +15,7 @@ struct FInstancedStructStream;
 struct FZoneGraphAnnotationTagLookup;
 struct FZoneGraphAnnotationTagContainer;
 
-#if !UE_BUILD_SHIPPING && !UE_BUILD_TEST
+#if UE_ENABLE_DEBUG_DRAWING
 class ZONEGRAPHANNOTATIONS_API FZoneGraphAnnotationSceneProxy final : public FDebugRenderSceneProxy
 {
 public:
@@ -31,8 +31,8 @@ private:
 #endif
 
 
-UCLASS(Abstract, HideCategories = (Physics, Collision, Lighting, Rendering, Mobile, HLOD, AssetUserData, Cooking, Navigation, Activation, Tags))
-class ZONEGRAPHANNOTATIONS_API UZoneGraphAnnotationComponent : public UPrimitiveComponent
+UCLASS(Abstract)
+class ZONEGRAPHANNOTATIONS_API UZoneGraphAnnotationComponent : public UDebugDrawComponent
 {
 	GENERATED_BODY()
 
@@ -57,7 +57,7 @@ public:
 	/** Called when new ZoneGraph data is removed. */ 
 	virtual void PreZoneGraphDataRemoved(const AZoneGraphData& ZoneGraphData) {}
 	
-#if !UE_BUILD_SHIPPING && !UE_BUILD_TEST
+#if UE_ENABLE_DEBUG_DRAWING
 	/** Called when scene proxy is rebuilt. */
 	virtual void DebugDraw(FZoneGraphAnnotationSceneProxy* DebugProxy) {}
 
@@ -75,23 +75,17 @@ protected:
 	virtual void PostEditChangeChainProperty(FPropertyChangedChainEvent& PropertyChangedEvent) override;
 #endif // WITH_EDITOR
 
-	//~ Begin UPrimitiveComponent Interface
 	virtual void OnRegister()  override;
 	virtual void OnUnregister()  override;
-	virtual void DestroyRenderState_Concurrent() override;
-	virtual FPrimitiveSceneProxy* CreateSceneProxy() override;
-	//~ End UPrimitiveComponent Interface
 
-	//~ Begin USceneComponent Interface
 	virtual FBoxSphereBounds CalcBounds(const FTransform& LocalToWorld) const override;
-	//~ End USceneComponent Interface
 
 	FDelegateHandle OnPostZoneGraphDataAddedHandle;
 	FDelegateHandle OnPreZoneGraphDataRemovedHandle;
 	FDelegateHandle OnPostWorldInitDelegateHandle;
 
-#if !UE_BUILD_SHIPPING && !UE_BUILD_TEST
-	FDebugDrawDelegateHelper DebugDrawDelegateHelper;
+#if UE_ENABLE_DEBUG_DRAWING
+	virtual FDebugRenderSceneProxy* CreateDebugSceneProxy() override;
 	FDelegateHandle CanvasDebugDrawDelegateHandle;
 #endif
 	
diff --git a/Engine/Plugins/Runtime/ZoneGraphAnnotations/Source/ZoneGraphAnnotations/Public/ZoneGraphAnnotationTestingActor.h b/Engine/Plugins/Runtime/ZoneGraphAnnotations/Source/ZoneGraphAnnotations/Public/ZoneGraphAnnotationTestingActor.h
--- a/Engine/Plugins/Runtime/ZoneGraphAnnotations/Source/ZoneGraphAnnotations/Public/ZoneGraphAnnotationTestingActor.h
+++ b/Engine/Plugins/Runtime/ZoneGraphAnnotations/Source/ZoneGraphAnnotations/Public/ZoneGraphAnnotationTestingActor.h
@@ -3,7 +3,7 @@
 #pragma once
 
 #include "GameFramework/Actor.h"
-#include "Components/PrimitiveComponent.h"
+#include "Debug/DebugDrawComponent.h"
 #include "DebugRenderSceneProxy.h"
 #include "ZoneGraphAnnotationTestingActor.generated.h"
 
@@ -19,7 +19,7 @@ class ZONEGRAPHANNOTATIONS_API UZoneGraphAnnotationTest : public UObject
 public:
 	virtual void Trigger() {}
 
-#if !UE_BUILD_SHIPPING && !UE_BUILD_TEST
+#if UE_ENABLE_DEBUG_DRAWING
 	virtual FBox CalcBounds(const FTransform& LocalToWorld) const { return FBox(ForceInit); };
 	virtual void DebugDraw(FDebugRenderSceneProxy* DebugProxy) {}
 	virtual void DebugDrawCanvas(UCanvas* Canvas, APlayerController*) {}
@@ -37,8 +37,8 @@ protected:
 
 
 /** Debug component to test Mass ZoneGraph Annotations. Handles tests and rendering. */
-UCLASS(ClassGroup = Custom, HideCategories = (Physics, Collision, Lighting, Rendering, Mobile))
-class ZONEGRAPHANNOTATIONS_API UZoneGraphAnnotationTestingComponent : public UPrimitiveComponent
+UCLASS(ClassGroup = Debug)
+class ZONEGRAPHANNOTATIONS_API UZoneGraphAnnotationTestingComponent : public UDebugDrawComponent
 {
 	GENERATED_BODY()
 public:
@@ -54,25 +54,20 @@ protected:
 	virtual void PostEditChangeChainProperty(FPropertyChangedChainEvent& PropertyChangedEvent) override;
 #endif
 
-	//~ Begin UPrimitiveComponent Interface.
 	virtual void OnRegister() override;
 	virtual void OnUnregister() override;
-	virtual void DestroyRenderState_Concurrent() override;
-	virtual FPrimitiveSceneProxy* CreateSceneProxy() override;
-	//~ End UPrimitiveComponent Interface.
 
-#if !UE_BUILD_SHIPPING && !UE_BUILD_TEST
+
+#if UE_ENABLE_DEBUG_DRAWING
+	virtual FDebugRenderSceneProxy* CreateDebugSceneProxy() override;
 	virtual void DebugDraw(FDebugRenderSceneProxy* DebugProxy);
 	virtual void DebugDrawCanvas(UCanvas* Canvas, APlayerController*);
-#endif
-	
-	UPROPERTY(EditAnywhere, Category = "Test", Instanced)
-	TArray<UZoneGraphAnnotationTest*> Tests;
 
-#if !UE_BUILD_SHIPPING && !UE_BUILD_TEST
-	FDebugDrawDelegateHelper DebugDrawDelegateHelper;
 	FDelegateHandle CanvasDebugDrawDelegateHandle;
 #endif
+
+	UPROPERTY(EditAnywhere, Category = "Test", Instanced)
+	TArray<UZoneGraphAnnotationTest*> Tests;
 };
 
 
diff --git a/Engine/Source/Developer/GameplayDebugger/Private/GameplayDebuggerRenderingComponent.cpp b/Engine/Source/Developer/GameplayDebugger/Private/GameplayDebuggerRenderingComponent.cpp
--- a/Engine/Source/Developer/GameplayDebugger/Private/GameplayDebuggerRenderingComponent.cpp
+++ b/Engine/Source/Developer/GameplayDebugger/Private/GameplayDebuggerRenderingComponent.cpp
@@ -78,14 +78,14 @@ protected:
 	TArray<FDebugRenderSceneProxy*> ChildProxies;
 };
 
-void FGameplayDebuggerDebugDrawDelegateHelper::RegisterDebugDrawDelegate()
+void FGameplayDebuggerDebugDrawDelegateHelper::RegisterDebugDrawDelegateInternal()
 {
-	ensureMsgf(State != RegisteredState, TEXT("RegisterDebugDrawDelegate is already Registered!"));
+	ensureMsgf(State != RegisteredState, TEXT("DrawDelegate is already Registered!"));
 	if (State == InitializedState)
 	{
 		for (int32 Idx = 0; Idx < DebugDrawDelegateHelpers.Num(); Idx++)
 		{
-			DebugDrawDelegateHelpers[Idx]->RegisterDebugDrawDelegate();
+			DebugDrawDelegateHelpers[Idx]->RequestRegisterDebugDrawDelegate(nullptr /*Deferred Context*/);
 		}
 		State = RegisteredState;
 	}
@@ -93,7 +93,7 @@ void FGameplayDebuggerDebugDrawDelegateHelper::RegisterDebugDrawDelegate()
 
 void FGameplayDebuggerDebugDrawDelegateHelper::UnregisterDebugDrawDelegate()
 {
-	ensureMsgf(State != InitializedState, TEXT("UnregisterDebugDrawDelegate is in an invalid State: %i !"), State);
+	ensureMsgf(State != InitializedState, TEXT("DrawDelegate is in an invalid State: %i !"), State);
 	if (State == RegisteredState)
 	{
 		for (int32 Idx = 0; Idx < DebugDrawDelegateHelpers.Num(); Idx++)
@@ -126,7 +126,7 @@ UGameplayDebuggerRenderingComponent::UGameplayDebuggerRenderingComponent(const F
 {
 }
 
-FPrimitiveSceneProxy* UGameplayDebuggerRenderingComponent::CreateSceneProxy()
+FDebugRenderSceneProxy* UGameplayDebuggerRenderingComponent::CreateDebugSceneProxy()
 {
 	GameplayDebuggerDebugDrawDelegateHelper.Reset();
 
@@ -141,16 +141,16 @@ FPrimitiveSceneProxy* UGameplayDebuggerRenderingComponent::CreateSceneProxy()
 			const TSharedRef<FGameplayDebuggerCategory> Category = OwnerReplicator->GetCategory(Idx);
 			if (Category->IsCategoryEnabled())
 			{
-				FDebugDrawDelegateHelper* DebugDrawDelegateHelper = nullptr;
-				FDebugRenderSceneProxy* CategorySceneProxy = Category->CreateDebugSceneProxy(this, DebugDrawDelegateHelper);
+				FDebugDrawDelegateHelper* CategoryDelegateHelper = nullptr;
+				FDebugRenderSceneProxy* CategorySceneProxy = Category->CreateDebugSceneProxy(this, CategoryDelegateHelper);
 				if (CategorySceneProxy)
 				{
 					SceneProxies.Add(CategorySceneProxy);
 				}
 
-				if (DebugDrawDelegateHelper)
+				if (CategoryDelegateHelper)
 				{
-					GameplayDebuggerDebugDrawDelegateHelper.AddDelegateHelper(DebugDrawDelegateHelper);
+					GameplayDebuggerDebugDrawDelegateHelper.AddDelegateHelper(CategoryDelegateHelper);
 				}
 			}
 		}
@@ -162,11 +162,6 @@ FPrimitiveSceneProxy* UGameplayDebuggerRenderingComponent::CreateSceneProxy()
 		}
 	}
 
-	if (CompositeProxy)
-	{
-		GameplayDebuggerDebugDrawDelegateHelper.InitDelegateHelper(CompositeProxy);
-		GameplayDebuggerDebugDrawDelegateHelper.RegisterDebugDrawDelegate();
-	}
 	return CompositeProxy;
 }
 
@@ -174,10 +169,3 @@ FBoxSphereBounds UGameplayDebuggerRenderingComponent::CalcBounds(const FTransfor
 {
 	return FBoxSphereBounds(FBox::BuildAABB(FVector::ZeroVector, FVector(1000000.0f, 1000000.0f, 1000000.0f)));
 }
-
-void UGameplayDebuggerRenderingComponent::DestroyRenderState_Concurrent()
-{
-	GameplayDebuggerDebugDrawDelegateHelper.UnregisterDebugDrawDelegate();
-
-	Super::DestroyRenderState_Concurrent();
-}
diff --git a/Engine/Source/Developer/GameplayDebugger/Public/GameplayDebuggerRenderingComponent.h b/Engine/Source/Developer/GameplayDebugger/Public/GameplayDebuggerRenderingComponent.h
--- a/Engine/Source/Developer/GameplayDebugger/Public/GameplayDebuggerRenderingComponent.h
+++ b/Engine/Source/Developer/GameplayDebugger/Public/GameplayDebuggerRenderingComponent.h
@@ -2,7 +2,7 @@
 
 #pragma once
 
-#include "Components/PrimitiveComponent.h"
+#include "Debug/DebugDrawComponent.h"
 #include "DebugRenderSceneProxy.h"
 #include "GameplayDebuggerRenderingComponent.generated.h"
 
@@ -22,7 +22,8 @@ public:
 
 	void AddDelegateHelper(FDebugDrawDelegateHelper* InDebugDrawDelegateHelper);
 
-	virtual void RegisterDebugDrawDelegate() override;
+protected:
+	virtual void RegisterDebugDrawDelegateInternal() override;
 	virtual void UnregisterDebugDrawDelegate() override;
 
 private:
@@ -31,14 +32,16 @@ private:
 
 
 
-UCLASS(NotBlueprintable, NotBlueprintType, noteditinlinenew, hidedropdown, Transient)
-class UGameplayDebuggerRenderingComponent : public UPrimitiveComponent
+UCLASS(ClassGroup = Debug, NotBlueprintable, NotBlueprintType, noteditinlinenew, hidedropdown, Transient)
+class UGameplayDebuggerRenderingComponent : public UDebugDrawComponent
 {
 	GENERATED_UCLASS_BODY()
 
-	virtual FPrimitiveSceneProxy* CreateSceneProxy() override;
+protected:
+	virtual FDebugRenderSceneProxy* CreateDebugSceneProxy() override;
+	virtual FDebugDrawDelegateHelper& GetDebugDrawDelegateHelper() override { return GameplayDebuggerDebugDrawDelegateHelper; }
 	virtual FBoxSphereBounds CalcBounds(const FTransform &LocalToWorld) const override;
-	virtual void DestroyRenderState_Concurrent() override;
 
+private:
 	FGameplayDebuggerDebugDrawDelegateHelper GameplayDebuggerDebugDrawDelegateHelper;
 };
diff --git a/Engine/Source/Developer/LogVisualizer/Private/VisualLoggerRenderingActorBase.cpp b/Engine/Source/Developer/LogVisualizer/Private/VisualLoggerRenderingActorBase.cpp
--- a/Engine/Source/Developer/LogVisualizer/Private/VisualLoggerRenderingActorBase.cpp
+++ b/Engine/Source/Developer/LogVisualizer/Private/VisualLoggerRenderingActorBase.cpp
@@ -49,7 +49,7 @@ UVisualLoggerRenderingComponent::UVisualLoggerRenderingComponent(const FObjectIn
 {
 }
 
-FPrimitiveSceneProxy* UVisualLoggerRenderingComponent::CreateSceneProxy()
+FDebugRenderSceneProxy* UVisualLoggerRenderingComponent::CreateDebugSceneProxy()
 {
 	AVisualLoggerRenderingActorBase* RenderingActor = Cast<AVisualLoggerRenderingActorBase>(GetOuter());
 	if (RenderingActor == nullptr)
@@ -74,13 +74,6 @@ FPrimitiveSceneProxy* UVisualLoggerRenderingComponent::CreateSceneProxy()
 			VLogSceneProxy->Capsules.Append(Shapes.Capsules);
 		});
 
-#if WITH_EDITOR
-	if (VLogSceneProxy)
-	{
-		DebugDrawDelegateHelper.InitDelegateHelper(VLogSceneProxy);
-		DebugDrawDelegateHelper.RegisterDebugDrawDelegate();
-	}
-#endif
 	return VLogSceneProxy;
 }
 
@@ -94,15 +87,6 @@ FBoxSphereBounds UVisualLoggerRenderingComponent::CalcBounds(const FTransform& L
 	return MyBounds;
 }
 
-void UVisualLoggerRenderingComponent::DestroyRenderState_Concurrent()
-{
-#if WITH_EDITOR
-	DebugDrawDelegateHelper.UnregisterDebugDrawDelegate();
-#endif
-
-	Super::DestroyRenderState_Concurrent();
-}
-
 AVisualLoggerRenderingActorBase::AVisualLoggerRenderingActorBase(const FObjectInitializer& ObjectInitializer)
 	: Super(ObjectInitializer)
 {
diff --git a/Engine/Source/Developer/LogVisualizer/Private/VisualLoggerRenderingComponent.h b/Engine/Source/Developer/LogVisualizer/Private/VisualLoggerRenderingComponent.h
--- a/Engine/Source/Developer/LogVisualizer/Private/VisualLoggerRenderingComponent.h
+++ b/Engine/Source/Developer/LogVisualizer/Private/VisualLoggerRenderingComponent.h
@@ -2,27 +2,18 @@
 
 #pragma once
 
-#include "DebugRenderSceneProxy.h"
-#include "Components/PrimitiveComponent.h"
+#include "Debug/DebugDrawComponent.h"
 #include "VisualLoggerRenderingComponent.generated.h"
 
 /**
-*	Transient actor used to draw visual logger data on level
-*/
-
-UCLASS()
-class UVisualLoggerRenderingComponent : public UPrimitiveComponent
+ *	Primitive component used to draw visual logger data on level
+ */
+UCLASS(ClassGroup = Debug)
+class UVisualLoggerRenderingComponent : public UDebugDrawComponent
 {
 public:
 	GENERATED_UCLASS_BODY()
 
-	virtual FPrimitiveSceneProxy* CreateSceneProxy() override;
+	virtual FDebugRenderSceneProxy* CreateDebugSceneProxy() override;
 	virtual FBoxSphereBounds CalcBounds(const FTransform &LocalToWorld) const override;
-	virtual void DestroyRenderState_Concurrent() override;
-
-private:
-#if WITH_EDITOR
-	FDebugDrawDelegateHelper DebugDrawDelegateHelper;
-#endif
 };
-
diff --git a/Engine/Source/Runtime/AIModule/Classes/EnvironmentQuery/EQSRenderingComponent.h b/Engine/Source/Runtime/AIModule/Classes/EnvironmentQuery/EQSRenderingComponent.h
--- a/Engine/Source/Runtime/AIModule/Classes/EnvironmentQuery/EQSRenderingComponent.h
+++ b/Engine/Source/Runtime/AIModule/Classes/EnvironmentQuery/EQSRenderingComponent.h
@@ -7,7 +7,7 @@
 #include "PrimitiveViewRelevance.h"
 #include "DebugRenderSceneProxy.h"
 #include "EnvironmentQuery/EnvQueryDebugHelpers.h"
-#include "Components/PrimitiveComponent.h"
+#include "Debug/DebugDrawComponent.h"
 #include "EQSRenderingComponent.generated.h"
 
 class APlayerController;
@@ -53,15 +53,8 @@ public:
 	{
 	}
 
-	virtual void InitDelegateHelper(const FDebugRenderSceneProxy* InSceneProxy) override
+	void SetupFromProxy(const FEQSSceneProxy* InSceneProxy)
 	{
-		check(0);
-	}
-
-	void InitDelegateHelper(const FEQSSceneProxy* InSceneProxy)
-	{
-		Super::InitDelegateHelper(InSceneProxy);
-
 		ActorOwner = InSceneProxy->ActorOwner;
 		QueryDataSource = InSceneProxy->QueryDataSource;
 		bDrawOnlyWhenSelected = InSceneProxy->bDrawOnlyWhenSelected;
@@ -83,27 +76,28 @@ private:
 };
 #endif
 
-UCLASS(hidecategories=Object)
-class AIMODULE_API UEQSRenderingComponent : public UPrimitiveComponent
+UCLASS(ClassGroup = Debug)
+class AIMODULE_API UEQSRenderingComponent : public UDebugDrawComponent
 {
 	GENERATED_UCLASS_BODY()
 
 	FString DrawFlagName;
 	uint32 bDrawOnlyWhenSelected : 1;
 
-	virtual FPrimitiveSceneProxy* CreateSceneProxy() override;
-	virtual FBoxSphereBounds CalcBounds(const FTransform &LocalToWorld) const override;
-	virtual void DestroyRenderState_Concurrent() override;
-
 	void ClearStoredDebugData();
 #if  USE_EQS_DEBUGGER || ENABLE_VISUAL_LOG
 	void StoreDebugData(const EQSDebug::FQueryData& DebugData);
 #endif
-#if  USE_EQS_DEBUGGER
+
+protected:
+	virtual FBoxSphereBounds CalcBounds(const FTransform &LocalToWorld) const override;
+
+#if UE_ENABLE_DEBUG_DRAWING
+	virtual FDebugRenderSceneProxy* CreateDebugSceneProxy() override;
+	virtual FDebugDrawDelegateHelper& GetDebugDrawDelegateHelper() override { return EQSRenderingDebugDrawDelegateHelper; }
 	FEQSRenderingDebugDrawDelegateHelper EQSRenderingDebugDrawDelegateHelper;
 #endif
 
-protected:
 	//EQSDebug::FQueryData DebugData;
 	TArray<FDebugRenderSceneProxy::FSphere> DebugDataSolidSpheres;
 	TArray<FDebugRenderSceneProxy::FText3d> DebugDataTexts;
diff --git a/Engine/Source/Runtime/AIModule/Private/EnvironmentQuery/EQSRenderingComponent.cpp b/Engine/Source/Runtime/AIModule/Private/EnvironmentQuery/EQSRenderingComponent.cpp
--- a/Engine/Source/Runtime/AIModule/Private/EnvironmentQuery/EQSRenderingComponent.cpp
+++ b/Engine/Source/Runtime/AIModule/Private/EnvironmentQuery/EQSRenderingComponent.cpp
@@ -281,20 +281,14 @@ UEQSRenderingComponent::UEQSRenderingComponent(const FObjectInitializer& ObjectI
 {
 }
 
-FPrimitiveSceneProxy* UEQSRenderingComponent::CreateSceneProxy()
+#if UE_ENABLE_DEBUG_DRAWING
+FDebugRenderSceneProxy* UEQSRenderingComponent::CreateDebugSceneProxy()
 {
 	FEQSSceneProxy* NewSceneProxy = new FEQSSceneProxy(*this, DrawFlagName, DebugDataSolidSpheres, DebugDataTexts);
-
-#if  USE_EQS_DEBUGGER
-	if (NewSceneProxy)
-	{
-		EQSRenderingDebugDrawDelegateHelper.InitDelegateHelper(NewSceneProxy);
-		EQSRenderingDebugDrawDelegateHelper.RegisterDebugDrawDelegate();
-	}
-#endif
-
+  	EQSRenderingDebugDrawDelegateHelper.SetupFromProxy(NewSceneProxy);
 	return NewSceneProxy;
 }
+#endif
 
 void UEQSRenderingComponent::ClearStoredDebugData()
 {
@@ -350,15 +344,6 @@ FBoxSphereBounds UEQSRenderingComponent::CalcBounds(const FTransform& LocalToWor
 	return FBoxSphereBounds(BoundingSphere).TransformBy(LocalToWorld);
 }
 
-void UEQSRenderingComponent::DestroyRenderState_Concurrent()
-{
-#if USE_EQS_DEBUGGER
-	EQSRenderingDebugDrawDelegateHelper.UnregisterDebugDrawDelegate();
-#endif
-
-	Super::DestroyRenderState_Concurrent();
-}
-
 //----------------------------------------------------------------------//
 // UEQSQueryResultSourceInterface
 //----------------------------------------------------------------------//
diff --git a/Engine/Source/Runtime/Engine/Classes/Debug/DebugDrawComponent.h b/Engine/Source/Runtime/Engine/Classes/Debug/DebugDrawComponent.h
new file mode 100644
--- /dev/null
+++ b/Engine/Source/Runtime/Engine/Classes/Debug/DebugDrawComponent.h
@@ -0,0 +1,69 @@
+﻿// Copyright Epic Games, Inc. All Rights Reserved.
+#pragma once
+
+#include "Components/PrimitiveComponent.h"
+#include "DebugRenderSceneProxy.h"
+#include "DebugDrawComponent.generated.h"
+
+/**
+ * Helper class to derive from to use debug draw delegate functionalities (i.e. DrawDebugLabels)
+ * The class will take care of registering a delegate to the UDebugDrawService and draw all FText3d provided by the scene proxy.
+ * This functionality only requires the derived classes to override `CreateDebugSceneProxy`.
+ *
+ * It is also possible to add text from other sources of data from the scene proxy but that requires a few extra steps:
+ *  + create a class that inherits from `FDebugDrawDelegateHelper`
+ *		- override `DrawDebugLabels`
+ *		- add a method to be able to populate the source data from the scene proxy
+ *
+ *		ex: class FMyDelegateHelper : public FDebugDrawDelegateHelper
+ *			{
+ *				public:
+ *					void SetupFromProxy(const FMySceneProxy* InSceneProxy) { <collect data> };
+ *				protected:
+ *					virtual void DrawDebugLabels(UCanvas* Canvas, APlayerController*) override { <draw data> };
+ *				private:
+ *					TArray<FVector> SomeData;
+ *			};
+ *
+ *  + Within the component inheriting from `UDebugDrawComponent`
+ *		- add a property of that new helper type in the component inheriting from `UDebugDrawComponent`
+ *		- override `GetDebugDrawDelegateHelper` to return that new member as the new delegate helper to use
+ *		- and then use it in `CreateDebugSceneProxy` before returning the created proxy.
+ *
+ *		ex: class MyDebugDrawComponent : public UDebugDrawComponent
+ *			{
+ *				protected:
+ *					virtual FDebugRenderSceneProxy* CreateDebugSceneProxy() override
+ *					{
+ *						FMySceneProxy* Proxy = new FMySceneProxy(this);
+ *						MyDelegateHelper.SetupFromProxy(Proxy);
+ *						return Proxy;
+ *					}
+ *
+ *					virtual FDebugDrawDelegateHelper& GetDebugDrawDelegateHelper() override { return MyDelegateHelper; }
+ *				private:
+ *					FMyDelegateHelper MyDelegateHelper;
+ *			}
+ */
+UCLASS(Abstract, HideCategories = (Activation, AssetUserData, Collision, Cooking, HLOD, Lighting, LOD, Mobile, Navigation, Physics, RayTracing, Rendering, Tags, TextureStreaming))
+class ENGINE_API UDebugDrawComponent : public UPrimitiveComponent
+{
+	GENERATED_BODY()
+
+protected:
+	/* Method overriden and marked as final since derived class should override `CreateDebugSceneProxy` */
+	virtual FPrimitiveSceneProxy* CreateSceneProxy() override final;
+
+	/* Method that derived class should override to create the scene proxy and customize a custom delegate helper (if any) */
+	virtual FDebugRenderSceneProxy* CreateDebugSceneProxy() PURE_VIRTUAL(UDebugDrawComponent::CreateSceneProxyInternal, return nullptr;);
+
+#if UE_ENABLE_DEBUG_DRAWING
+	virtual void CreateRenderState_Concurrent(FRegisterComponentContext* Context) override;
+	virtual void DestroyRenderState_Concurrent() override;
+
+	/** Method that should be overriden when subclass uses a custom delegate helper. */
+	virtual FDebugDrawDelegateHelper& GetDebugDrawDelegateHelper() { return DebugDrawDelegateHelper; }
+
+	FDebugDrawDelegateHelper DebugDrawDelegateHelper;
+#endif
+};
\ No newline at end of file
diff --git a/Engine/Source/Runtime/Engine/Private/Debug/DebugDrawComponent.cpp b/Engine/Source/Runtime/Engine/Private/Debug/DebugDrawComponent.cpp
new file mode 100644
--- /dev/null
+++ b/Engine/Source/Runtime/Engine/Private/Debug/DebugDrawComponent.cpp
@@ -0,0 +1,31 @@
+﻿// Copyright Epic Games, Inc. All Rights Reserved.
+
+#include "Debug/DebugDrawComponent.h"
+
+FPrimitiveSceneProxy* UDebugDrawComponent::CreateSceneProxy()
+{
+	FDebugRenderSceneProxy* Proxy = CreateDebugSceneProxy();
+#if UE_ENABLE_DEBUG_DRAWING
+  	if (Proxy != nullptr)
+	{
+		GetDebugDrawDelegateHelper().InitDelegateHelper(Proxy);
+	}
+
+	GetDebugDrawDelegateHelper().ProcessDeferredRegister();
+#endif
+	return Proxy;
+}
+
+#if UE_ENABLE_DEBUG_DRAWING
+  void UDebugDrawComponent::CreateRenderState_Concurrent(FRegisterComponentContext* Context)
+{
+	Super::CreateRenderState_Concurrent(Context);
+	GetDebugDrawDelegateHelper().RequestRegisterDebugDrawDelegate(Context);
+}
+
+void UDebugDrawComponent::DestroyRenderState_Concurrent()
+{
+	GetDebugDrawDelegateHelper().UnregisterDebugDrawDelegate();
+	Super::DestroyRenderState_Concurrent();
+}
+#endif
diff --git a/Engine/Source/Runtime/Engine/Private/DebugRenderSceneProxy.cpp b/Engine/Source/Runtime/Engine/Private/DebugRenderSceneProxy.cpp
--- a/Engine/Source/Runtime/Engine/Private/DebugRenderSceneProxy.cpp
+++ b/Engine/Source/Runtime/Engine/Private/DebugRenderSceneProxy.cpp
@@ -25,9 +25,9 @@ FDebugRenderSceneProxy::FDebugRenderSceneProxy(const UPrimitiveComponent* InComp
 {
 }
 
-void FDebugDrawDelegateHelper::RegisterDebugDrawDelegate()
+void FDebugDrawDelegateHelper::RegisterDebugDrawDelegateInternal()
 {
-	ensureMsgf(State != RegisteredState, TEXT("RegisterDebugDrawDelegate is already Registered!"));
+	ensureMsgf(State != RegisteredState, TEXT("DrawDelegate is already Registered!"));
 	if (State == InitializedState)
 	{
 		DebugTextDrawingDelegate = FDebugDrawDelegate::CreateRaw(this, &FDebugDrawDelegateHelper::DrawDebugLabels);
@@ -36,9 +36,28 @@ void FDebugDrawDelegateHelper::RegisterDebugDrawDelegate()
 	}
 }
 
+void FDebugDrawDelegateHelper::RequestRegisterDebugDrawDelegate(FRegisterComponentContext* Context)
+{
+	bDeferredRegister = Context != nullptr;
+
+	if (!bDeferredRegister)
+	{
+		RegisterDebugDrawDelegateInternal();
+	}
+}
+
+void FDebugDrawDelegateHelper::ProcessDeferredRegister()
+{
+	if (bDeferredRegister)
+	{
+		RegisterDebugDrawDelegateInternal();
+		bDeferredRegister = false;
+	}
+}
+
 void FDebugDrawDelegateHelper::UnregisterDebugDrawDelegate()
 {
-	ensureMsgf(State != InitializedState, TEXT("UnregisterDebugDrawDelegate is in an invalid State: %i !"), State);
+	ensureMsgf(State != InitializedState, TEXT("DrawDelegate is in an invalid State: %i !"), State);
 	if (State == RegisteredState)
 	{
 		check(DebugTextDrawingDelegate.IsBound());
@@ -49,11 +68,11 @@ void FDebugDrawDelegateHelper::UnregisterDebugDrawDelegate()
 
 void  FDebugDrawDelegateHelper::ReregisterDebugDrawDelegate()
 {
-	ensureMsgf(State != UndefinedState, TEXT("ReregisterDebugDrawDelegate is in an invalid State: %i !"), State);
+	ensureMsgf(State != UndefinedState, TEXT("DrawDelegate is in an invalid State: %i !"), State);
 	if (State == RegisteredState)
 	{
 		UnregisterDebugDrawDelegate();
-		RegisterDebugDrawDelegate();
+		RegisterDebugDrawDelegateInternal();
 	}
 }
 
diff --git a/Engine/Source/Runtime/Engine/Public/DebugRenderSceneProxy.h b/Engine/Source/Runtime/Engine/Public/DebugRenderSceneProxy.h
--- a/Engine/Source/Runtime/Engine/Public/DebugRenderSceneProxy.h
+++ b/Engine/Source/Runtime/Engine/Public/DebugRenderSceneProxy.h
@@ -261,7 +261,7 @@ protected:
 	typedef TArray<FDebugRenderSceneProxy::FText3d> TextArray;
 
 public:
-	virtual void InitDelegateHelper(const FDebugRenderSceneProxy* InSceneProxy)
+	void InitDelegateHelper(const FDebugRenderSceneProxy* InSceneProxy)
 	{
 		check(IsInParallelGameThread() || IsInGameThread());
 
@@ -272,21 +272,39 @@ public:
 		State = (State == UndefinedState) ? InitializedState : State;
 	}
 
-	UE_DEPRECATED(5.0, "This method is deprecated. Use RegisterDebugDrawDelegate instead.")
-	ENGINE_API virtual void RegisterDebugDrawDelgate() { RegisterDebugDrawDelegate(); }
+	UE_DEPRECATED(5.0, "This method is deprecated. Call RequestRegisterDebugDrawDelegate or override RegisterDebugDrawDelegateInternal instead.")
+	ENGINE_API virtual void RegisterDebugDrawDelgate() final { RegisterDebugDrawDelegateInternal(); }
 	UE_DEPRECATED(5.0, "This method is deprecated. Use UnregisterDebugDrawDelegate instead.")
 	ENGINE_API virtual void UnregisterDebugDrawDelgate() { UnregisterDebugDrawDelegate(); }
 	UE_DEPRECATED(5.0, "This method is deprecated. Use ReregisterDebugDrawDelegate instead.")
 	ENGINE_API void ReregisterDebugDrawDelgate() { ReregisterDebugDrawDelegate(); }
+	UE_DEPRECATED(5.0, "This method is deprecated. Call RequestRegisterDebugDrawDelegate or override RegisterDebugDrawDelegateInternal instead.")
+	ENGINE_API virtual void RegisterDebugDrawDelegate() final { RegisterDebugDrawDelegateInternal(); }
+
+	/**
+	 * Method that should be called at render state creation (i.e. CreateRenderState_Concurrent).
+	 * It will either call `RegisterDebugDrawDelegate` when deferring context is not provided
+	 * or mark for deferred registration that should be flushed by calling `ProcessDeferredRegister`
+	 * on scene proxy creation.
+	 * @param Context valid context is provided when primitives are batched for deferred 'add'
+	 */
+	ENGINE_API void RequestRegisterDebugDrawDelegate(FRegisterComponentContext* Context);
+
+	/**
+	 * Method that should be called when creating scene proxy (i.e. CreateSceneProxy) to process any pending registration that might have
+	 * been requested from deferred primitive batching (i.e. CreateRenderState_Concurrent(FRegisterComponentContext != nullptr)).
+	 */
+	ENGINE_API void ProcessDeferredRegister();
 
-	/** called to set up debug drawing delegate in UDebugDrawService if you want to draw labels */
-	ENGINE_API virtual void RegisterDebugDrawDelegate();
 	/** called to clean up debug drawing delegate in UDebugDrawService */
 	ENGINE_API virtual void UnregisterDebugDrawDelegate();
 
 	ENGINE_API void ReregisterDebugDrawDelegate();
 
 protected:
+	/** called to set up debug drawing delegate in UDebugDrawService if you want to draw labels */
+	ENGINE_API virtual void RegisterDebugDrawDelegateInternal();
+
 	ENGINE_API virtual void DrawDebugLabels(UCanvas* Canvas, APlayerController*);
 	ENGINE_API void ResetTexts() { Texts.Reset(); }
 
@@ -300,6 +318,8 @@ protected:
 		RegisteredState,
 	} State;
 
+	bool bDeferredRegister = false;
+
 private:
 	TextArray Texts;
 	FString ViewFlagName;
diff --git a/Engine/Source/Runtime/NavigationSystem/Private/NavMesh/NavMeshRenderingComponent.cpp b/Engine/Source/Runtime/NavigationSystem/Private/NavMesh/NavMeshRenderingComponent.cpp
--- a/Engine/Source/Runtime/NavigationSystem/Private/NavMesh/NavMeshRenderingComponent.cpp
+++ b/Engine/Source/Runtime/NavigationSystem/Private/NavMesh/NavMeshRenderingComponent.cpp
@@ -1075,6 +1075,7 @@ FNavMeshSceneProxy::FNavMeshSceneProxy(const UPrimitiveComponent* InComponent, F
 	, bForceRendering(ForceToRender)
 {
 	DrawType = EDrawType::SolidAndWireMeshes;
+	ViewFlagName = TEXT("Navigation");
 
 	if (InProxyData)
 	{
@@ -1148,30 +1149,6 @@ FNavMeshSceneProxy::~FNavMeshSceneProxy()
 	VertexFactory.ReleaseResource();
 }
 
-#if WITH_RECAST && UE_ENABLE_DEBUG_DRAWING
-void FNavMeshDebugDrawDelegateHelper::RegisterDebugDrawDelegate()
-{
-	ensureMsgf(State != RegisteredState, TEXT("RegisterDebugDrawDelegate is already Registered!"));
-	if (State == InitializedState)
-	{
-		DebugTextDrawingDelegate = FDebugDrawDelegate::CreateRaw(this, &FNavMeshDebugDrawDelegateHelper::DrawDebugLabels);
-		DebugTextDrawingDelegateHandle = UDebugDrawService::Register(TEXT("Navigation"), DebugTextDrawingDelegate);
-		State = RegisteredState;
-	}
-}
-
-void FNavMeshDebugDrawDelegateHelper::UnregisterDebugDrawDelegate()
-{
-	ensureMsgf(State != InitializedState, TEXT("UnegisterDebugDrawDelegate is in an invalid State: %i !"), State);
-	if (State == RegisteredState)
-	{
-		check(DebugTextDrawingDelegate.IsBound());
-		UDebugDrawService::Unregister(DebugTextDrawingDelegateHandle);
-		State = InitializedState;
-	}
-}
-#endif
-
 void FNavMeshSceneProxy::GetDynamicMeshElements(const TArray<const FSceneView*>& Views, const FSceneViewFamily& ViewFamily, uint32 VisibilityMap, FMeshElementCollector& Collector) const
 {
 	QUICK_SCOPE_CYCLE_COUNTER(STAT_RecastRenderingSceneProxy_GetDynamicMeshElements);
@@ -1493,7 +1470,7 @@ void UNavMeshRenderingComponent::OnRegister()
 {
 	Super::OnRegister();
 
-#if WITH_RECAST && !UE_BUILD_SHIPPING && !UE_BUILD_TEST
+#if WITH_RECAST && UE_ENABLE_DEBUG_DRAWING
 	// it's a kind of HACK but there is no event or other information that show flag was changed by user => we have to check it periodically
 #if WITH_EDITOR
 	if (GEditor)
@@ -1505,12 +1482,12 @@ void UNavMeshRenderingComponent::OnRegister()
 	{
 		GetWorld()->GetTimerManager().SetTimer(TimerHandle, FTimerDelegate::CreateUObject(this, &UNavMeshRenderingComponent::TimerFunction), 1, true);
 	}
-#endif //WITH_RECAST && !UE_BUILD_SHIPPING && !UE_BUILD_TEST
+#endif //WITH_RECAST && UE_ENABLE_DEBUG_DRAWING
 }
 
 void UNavMeshRenderingComponent::OnUnregister()
 {
-#if WITH_RECAST && !UE_BUILD_SHIPPING && !UE_BUILD_TEST
+#if WITH_RECAST && UE_ENABLE_DEBUG_DRAWING
 	// it's a kind of HACK but there is no event or other information that show flag was changed by user => we have to check it periodically
 #if WITH_EDITOR
 	if (GEditor)
@@ -1522,7 +1499,7 @@ void UNavMeshRenderingComponent::OnUnregister()
 	{
 		GetWorld()->GetTimerManager().ClearTimer(TimerHandle);
 	}
-#endif //WITH_RECAST && !UE_BUILD_SHIPPING && !UE_BUILD_TEST
+#endif //WITH_RECAST && UE_ENABLE_DEBUG_DRAWING
 	Super::OnUnregister();
 }
 
@@ -1533,9 +1510,10 @@ void UNavMeshRenderingComponent::GatherData(const ARecastNavMesh& NavMesh, FNavM
 	OutProxyData.GatherData(&NavMesh, DetailFlags, EmptyTileSet);
 }
 
-FPrimitiveSceneProxy* UNavMeshRenderingComponent::CreateSceneProxy()
+#if UE_ENABLE_DEBUG_DRAWING
+  FDebugRenderSceneProxy* UNavMeshRenderingComponent::CreateDebugSceneProxy()
 {
-#if WITH_RECAST && !UE_BUILD_SHIPPING && !UE_BUILD_TEST
+#if WITH_RECAST
 	FNavMeshSceneProxy* NavMeshSceneProxy = nullptr;
 
 	const bool bShowNavigation = IsNavigationShowFlagSet(GetWorld());
@@ -1551,29 +1529,17 @@ FPrimitiveSceneProxy* UNavMeshRenderingComponent::CreateSceneProxy()
 			GatherData(*NavMesh, ProxyData);
 
 			NavMeshSceneProxy = new FNavMeshSceneProxy(this, &ProxyData);
+			NavMeshDebugDrawDelegateManager.SetupFromProxy(NavMeshSceneProxy);
 		}
 	}
 
-	if (NavMeshSceneProxy)
-	{
-		NavMeshDebugDrawDelegateManager.InitDelegateHelper(NavMeshSceneProxy);
-		NavMeshDebugDrawDelegateManager.RegisterDebugDrawDelegate();
-	}
 	return NavMeshSceneProxy;
 #else
 	return nullptr;
-#endif //WITH_RECAST && !UE_BUILD_SHIPPING && !UE_BUILD_TEST
+#endif // WITH_RECAST
 }
-
-void UNavMeshRenderingComponent::DestroyRenderState_Concurrent()
-{
-#if WITH_RECAST && !UE_BUILD_SHIPPING && !UE_BUILD_TEST
-	NavMeshDebugDrawDelegateManager.UnregisterDebugDrawDelegate();
 #endif
 
-	Super::DestroyRenderState_Concurrent();
-}
-
 FBoxSphereBounds UNavMeshRenderingComponent::CalcBounds(const FTransform& LocalToWorld) const
 {
 	FBox BoundingBox(ForceInit);
diff --git a/Engine/Source/Runtime/NavigationSystem/Private/NavMesh/NavTestRenderingComponent.cpp b/Engine/Source/Runtime/NavigationSystem/Private/NavMesh/NavTestRenderingComponent.cpp
--- a/Engine/Source/Runtime/NavigationSystem/Private/NavMesh/NavTestRenderingComponent.cpp
+++ b/Engine/Source/Runtime/NavigationSystem/Private/NavMesh/NavTestRenderingComponent.cpp
@@ -24,6 +24,8 @@ FNavTestSceneProxy::FNavTestSceneProxy(const UNavTestRenderingComponent* InCompo
 	, NavMeshDrawOffset(0,0,10)
 	, NavTestActor(nullptr)
 {
+	ViewFlagName = TEXT("Navigation");
+
 	if (InComponent == nullptr)
 	{
 		return;
@@ -314,10 +316,8 @@ uint32 FNavTestSceneProxy::GetAllocatedSizeInternal() const
 }
 
 #if WITH_RECAST && WITH_EDITOR
-void FNavTestDebugDrawDelegateHelper::InitDelegateHelper(const FNavTestSceneProxy* InSceneProxy)
+void FNavTestDebugDrawDelegateHelper::SetupFromProxy(const FNavTestSceneProxy* InSceneProxy)
 {
-	Super::InitDelegateHelper(InSceneProxy);
-
 	PathPoints.Reset();
 	PathPoints.Append(InSceneProxy->PathPoints);
 	PathPointFlags.Reset();
@@ -330,28 +330,6 @@ void FNavTestDebugDrawDelegateHelper::InitDelegateHelper(const FNavTestSceneProx
 	bShowDiff = InSceneProxy->bShowDiff;
 }
 
-void FNavTestDebugDrawDelegateHelper::RegisterDebugDrawDelegate()
-{
-	ensureMsgf(State != RegisteredState, TEXT("RegisterDebugDrawDelegate is already Registered!"));
-	if (State == InitializedState)
-	{
-		DebugTextDrawingDelegate = FDebugDrawDelegate::CreateRaw(this, &FNavTestDebugDrawDelegateHelper::DrawDebugLabels);
-		DebugTextDrawingDelegateHandle = UDebugDrawService::Register(TEXT("Navigation"), DebugTextDrawingDelegate);
-		State = RegisteredState;
-	}
-}
-
-void FNavTestDebugDrawDelegateHelper::UnregisterDebugDrawDelegate()
-{
-	ensureMsgf(State != InitializedState, TEXT("UnegisterDebugDrawDelegate is in an invalid State: %i !"), State);
-	if (State == RegisteredState)
-	{
-		check(DebugTextDrawingDelegate.IsBound());
-		UDebugDrawService::Unregister(DebugTextDrawingDelegateHandle);
-		State = InitializedState;
-	}
-}
-
 void FNavTestDebugDrawDelegateHelper::DrawDebugLabels(UCanvas* Canvas, APlayerController*)
 {
 	if (NavTestActor == nullptr)
@@ -414,15 +392,16 @@ UNavTestRenderingComponent::UNavTestRenderingComponent(const FObjectInitializer&
 {
 }
 
-FPrimitiveSceneProxy* UNavTestRenderingComponent::CreateSceneProxy()
+#if UE_ENABLE_DEBUG_DRAWING
+  FDebugRenderSceneProxy* UNavTestRenderingComponent::CreateDebugSceneProxy()
 {
-	FNavTestSceneProxy* SceneProxy2 = new FNavTestSceneProxy(this);
+	FNavTestSceneProxy* NewSceneProxy = new FNavTestSceneProxy(this);
 #if WITH_RECAST && WITH_EDITOR
-	NavTestDebugDrawDelegateHelper.InitDelegateHelper(SceneProxy2);
-	NavTestDebugDrawDelegateHelper.RegisterDebugDrawDelegate();
+	NavTestDebugDrawDelegateHelper.SetupFromProxy(NewSceneProxy);
 #endif
-	return SceneProxy2;
+	return NewSceneProxy;
 }
+#endif
 
 FBoxSphereBounds UNavTestRenderingComponent::CalcBounds(const FTransform& LocalToWorld) const
 {
@@ -459,12 +438,3 @@ FBoxSphereBounds UNavTestRenderingComponent::CalcBounds(const FTransform& LocalT
 
 	return FBoxSphereBounds(BoundingBox);
 }
-
-void UNavTestRenderingComponent::DestroyRenderState_Concurrent()
-{
-#if WITH_RECAST && WITH_EDITOR
-	NavTestDebugDrawDelegateHelper.UnregisterDebugDrawDelegate();
-#endif
-
-	Super::DestroyRenderState_Concurrent();
-}
diff --git a/Engine/Source/Runtime/NavigationSystem/Public/NavMesh/NavMeshRenderingComponent.h b/Engine/Source/Runtime/NavigationSystem/Public/NavMesh/NavMeshRenderingComponent.h
--- a/Engine/Source/Runtime/NavigationSystem/Public/NavMesh/NavMeshRenderingComponent.h
+++ b/Engine/Source/Runtime/NavigationSystem/Public/NavMesh/NavMeshRenderingComponent.h
@@ -7,7 +7,7 @@
 #include "MaterialShared.h"
 #include "DynamicMeshBuilder.h"
 #include "DebugRenderSceneProxy.h"
-#include "Components/PrimitiveComponent.h"
+#include "Debug/DebugDrawComponent.h"
 #include "MeshBatch.h"
 #include "LocalVertexFactory.h"
 #include "Math/GenericOctree.h"
@@ -156,24 +156,14 @@ public:
 	{
 	}
 
-	virtual void InitDelegateHelper(const FDebugRenderSceneProxy* InSceneProxy) override
+	void SetupFromProxy(const FNavMeshSceneProxy* InSceneProxy)
 	{
-		check(0);
-	}
-
-	void InitDelegateHelper(const FNavMeshSceneProxy* InSceneProxy)
-	{
-		Super::InitDelegateHelper(InSceneProxy);
-
 		DebugLabels.Reset();
 		DebugLabels.Append(InSceneProxy->ProxyData.DebugLabels);
 		bForceRendering = InSceneProxy->bForceRendering;
 		bNeedsNewData = InSceneProxy->ProxyData.bNeedsNewData;
 	}
 
-	NAVIGATIONSYSTEM_API virtual void RegisterDebugDrawDelegate() override;
-	NAVIGATIONSYSTEM_API virtual void UnregisterDebugDrawDelegate() override;
-
 protected:
 	NAVIGATIONSYSTEM_API virtual void DrawDebugLabels(UCanvas* Canvas, APlayerController*) override;
 
@@ -184,33 +174,28 @@ private:
 };
 #endif
 
-UCLASS(hidecategories=Object, editinlinenew)
-class NAVIGATIONSYSTEM_API UNavMeshRenderingComponent : public UPrimitiveComponent
+UCLASS(editinlinenew, ClassGroup = Debug)
+class NAVIGATIONSYSTEM_API UNavMeshRenderingComponent : public UDebugDrawComponent
 {
 	GENERATED_UCLASS_BODY()
 
 public:
-	
-	//~ Begin UPrimitiveComponent Interface
-	virtual FPrimitiveSceneProxy* CreateSceneProxy() override;
-	virtual void OnRegister()  override;
-	virtual void OnUnregister()  override;
-	//~ End UPrimitiveComponent Interface
-
-	//~ Begin UActorComponent Interface
-	virtual void DestroyRenderState_Concurrent() override;
-	//~ End UActorComponent Interface
-
-	//~ Begin USceneComponent Interface
-	virtual FBoxSphereBounds CalcBounds(const FTransform &LocalToWorld) const override;
-	//~ End USceneComponent Interface
-
 	void ForceUpdate() { bForceUpdate = true; }
 	bool IsForcingUpdate() const { return bForceUpdate; }
 
 	static bool IsNavigationShowFlagSet(const UWorld* World);
 
 protected:
+	virtual void OnRegister()  override;
+	virtual void OnUnregister()  override;
+
+#if UE_ENABLE_DEBUG_DRAWING
+  	virtual FDebugRenderSceneProxy* CreateDebugSceneProxy() override;
+	virtual FDebugDrawDelegateHelper& GetDebugDrawDelegateHelper() override { return NavMeshDebugDrawDelegateManager; }
+#endif
+
+	virtual FBoxSphereBounds CalcBounds(const FTransform &LocalToWorld) const override;
+
 	/** Gathers drawable information from NavMesh and puts it in OutProxyData. 
 	 *	Override to add additional information to OutProxyData.*/
 	virtual void GatherData(const ARecastNavMesh& NavMesh, FNavMeshSceneProxyData& OutProxyData) const;
@@ -223,7 +208,7 @@ protected:
 	FTimerHandle TimerHandle;
 
 protected:
-#if WITH_RECAST && UE_ENABLE_DEBUG_DRAWING
+#if UE_ENABLE_DEBUG_DRAWING
 	FNavMeshDebugDrawDelegateHelper NavMeshDebugDrawDelegateManager;
 #endif
 };
diff --git a/Engine/Source/Runtime/NavigationSystem/Public/NavMesh/NavTestRenderingComponent.h b/Engine/Source/Runtime/NavigationSystem/Public/NavMesh/NavTestRenderingComponent.h
--- a/Engine/Source/Runtime/NavigationSystem/Public/NavMesh/NavTestRenderingComponent.h
+++ b/Engine/Source/Runtime/NavigationSystem/Public/NavMesh/NavTestRenderingComponent.h
@@ -3,7 +3,7 @@
 #pragma once
 
 #include "PrimitiveViewRelevance.h"
-#include "Components/PrimitiveComponent.h"
+#include "Debug/DebugDrawComponent.h"
 #include "DynamicMeshBuilder.h"
 #include "DebugRenderSceneProxy.h"
 
@@ -94,15 +94,7 @@ public:
 	{
 	}
 
-	virtual void InitDelegateHelper(const FDebugRenderSceneProxy* InSceneProxy) override
-	{
-		check(0);
-	}
-
-	void InitDelegateHelper(const FNavTestSceneProxy* InSceneProxy);
-
-	virtual void RegisterDebugDrawDelegate() override;
-	virtual void UnregisterDebugDrawDelegate() override;
+	void SetupFromProxy(const FNavTestSceneProxy* InSceneProxy);
 
 protected:
 	virtual void DrawDebugLabels(UCanvas* Canvas, APlayerController*) override;
@@ -122,17 +114,22 @@ private:
 
 class FPrimitiveSceneProxy;
 
-UCLASS(hidecategories=Object)
-class UNavTestRenderingComponent: public UPrimitiveComponent
+UCLASS(ClassGroup = Debug)
+class UNavTestRenderingComponent: public UDebugDrawComponent
 {
 	GENERATED_UCLASS_BODY()
 
-	virtual FPrimitiveSceneProxy* CreateSceneProxy() override;
+protected:
+
 	virtual FBoxSphereBounds CalcBounds(const FTransform &LocalToWorld) const override;
-	virtual void DestroyRenderState_Concurrent() override;
 
-private:
+#if UE_ENABLE_DEBUG_DRAWING
+	virtual FDebugRenderSceneProxy* CreateDebugSceneProxy() override;
+
 #if WITH_RECAST && WITH_EDITOR
+	virtual FDebugDrawDelegateHelper& GetDebugDrawDelegateHelper() override { return NavTestDebugDrawDelegateHelper; }
+private:
 	FNavTestDebugDrawDelegateHelper NavTestDebugDrawDelegateHelper;
-#endif
+#endif // WITH_RECAST && WITH_EDITOR
+#endif // UE_ENABLE_DEBUG_DRAWING
 };
