Fix navmesh rendering not displaying some triangles

diff --git a/Engine/Source/Runtime/NavigationSystem/Private/NavMesh/NavMeshRenderingComponent.cpp b/Engine/Source/Runtime/NavigationSystem/Private/NavMesh/NavMeshRenderingComponent.cpp
--- a/Engine/Source/Runtime/NavigationSystem/Private/NavMesh/NavMeshRenderingComponent.cpp
+++ b/Engine/Source/Runtime/NavigationSystem/Private/NavMesh/NavMeshRenderingComponent.cpp
@@ -994,7 +994,7 @@ void FNavMeshSceneProxyData::AddMeshForInternalData(const FRecastInternalDebugDa
 				const uint32 Color = Colors[Indices[i]];	// Use the first color as representative of the triangle color.
 				if (Color != PrevColor)
 				{
-					ColorIdx = UniqueColors.Add(FNavMeshRenderingHelpers::FUniqueColor(PrevColor, 1));
+					ColorIdx = UniqueColors.Add(FNavMeshRenderingHelpers::FUniqueColor(Color, 1));
 					PrevColor = Color;
 				}
 				else
