Fix navmesh rendering colors originating from recast

diff --git a/Engine/Source/Runtime/NavigationSystem/Private/NavMesh/NavMeshRenderingComponent.cpp b/Engine/Source/Runtime/NavigationSystem/Private/NavMesh/NavMeshRenderingComponent.cpp
--- a/Engine/Source/Runtime/NavigationSystem/Private/NavMesh/NavMeshRenderingComponent.cpp
+++ b/Engine/Source/Runtime/NavigationSystem/Private/NavMesh/NavMeshRenderingComponent.cpp
@@ -960,8 +960,8 @@ namespace FNavMeshRenderingHelpers
 	struct FUniqueColor
 	{
 		FUniqueColor() : Color(0), Count(0) {}
-		FUniqueColor(const uint32 InColor, const uint32 InCount) : Color(InColor), Count(InCount) {}
-		uint32 Color;
+		FUniqueColor(const FColor InColor, const uint32 InCount) : Color(InColor), Count(InCount) {}
+		FColor Color;
 		uint32 Count;
 		bool operator==(const FUniqueColor& Other) const
 		{
@@ -971,7 +971,7 @@ namespace FNavMeshRenderingHelpers
 	};
 	FORCEINLINE uint32 GetTypeHash(const FUniqueColor& UniqueColor)
 	{
-		return UniqueColor.Color;
+		return UniqueColor.Color.DWColor();
 	}
 }
 
@@ -981,17 +981,17 @@ void FNavMeshSceneProxyData::AddMeshForInternalData(const FRecastInternalDebugDa
 	{
 		const TArray<uint32>& Indices = InInternalData.TriangleIndices;
 		const TArray<FVector>& Vertices = InInternalData.TriangleVertices;
-		const TArray<uint32>& Colors = InInternalData.TriangleColors;
+		const TArray<FColor>& Colors = InInternalData.TriangleColors;
 
 		if (ensure(Vertices.Num() == Colors.Num()))
 		{
 			// Split the mesh into different colored pieces (because we cannot use vertex colors).
 			TSet<FNavMeshRenderingHelpers::FUniqueColor> UniqueColors;
-			uint32 PrevColor = Colors[Indices[0]];
+			FColor PrevColor = Colors[Indices[0]];
 			FSetElementId ColorIdx = UniqueColors.Add(FNavMeshRenderingHelpers::FUniqueColor(PrevColor, 1));
 			for (int32 i = 3; i < Indices.Num(); i += 3)
 			{
-				const uint32 Color = Colors[Indices[i]];	// Use the first color as representative of the triangle color.
+				const FColor Color = Colors[Indices[i]];	// Use the first color as representative of the triangle color.
 				if (Color != PrevColor)
 				{
 					ColorIdx = UniqueColors.Add(FNavMeshRenderingHelpers::FUniqueColor(Color, 1));
@@ -1013,19 +1013,19 @@ void FNavMeshSceneProxyData::AddMeshForInternalData(const FRecastInternalDebugDa
 				MeshData.Vertices.Reserve(VertexCount);
 				for (int32 i = 0; i < Indices.Num(); i += 3)
 				{
-					const uint32 Color = Colors[Indices[i]];
+					const FColor Color = Colors[Indices[i]];
 					if (Color == CurrentColor.Color)
 					{
 						const int32 VertexBase = MeshData.Vertices.Num();
-						FNavMeshRenderingHelpers::AddVertex(MeshData, Vertices[Indices[i]] + NavMeshDrawOffset, FColor(Color));
-						FNavMeshRenderingHelpers::AddVertex(MeshData, Vertices[Indices[i + 1]] + NavMeshDrawOffset, FColor(Color));
-						FNavMeshRenderingHelpers::AddVertex(MeshData, Vertices[Indices[i + 2]] + NavMeshDrawOffset, FColor(Color));
+						FNavMeshRenderingHelpers::AddVertex(MeshData, Vertices[Indices[i]] + NavMeshDrawOffset, Color);
+						FNavMeshRenderingHelpers::AddVertex(MeshData, Vertices[Indices[i + 1]] + NavMeshDrawOffset, Color);
+						FNavMeshRenderingHelpers::AddVertex(MeshData, Vertices[Indices[i + 2]] + NavMeshDrawOffset, Color);
 						MeshData.Indices.Add(VertexBase);
 						MeshData.Indices.Add(VertexBase + 1);
 						MeshData.Indices.Add(VertexBase + 2);
 					}
 				}
-				MeshData.ClusterColor = FColor(CurrentColor.Color);
+				MeshData.ClusterColor = CurrentColor.Color;
 			}
 		}
 	}
@@ -1033,12 +1033,12 @@ void FNavMeshSceneProxyData::AddMeshForInternalData(const FRecastInternalDebugDa
 	if (InInternalData.LineVertices.Num() > 0)
 	{
 		const TArray<FVector>& Vertices = InInternalData.LineVertices;
-		const TArray<uint32>& Colors = InInternalData.LineColors;
+		const TArray<FColor>& Colors = InInternalData.LineColors;
 		if (ensure(Vertices.Num() == Colors.Num()))
 		{
 			for (int32 i = 0; i < Vertices.Num(); i += 2)
 			{
-				AuxLines.Emplace(Vertices[i] + NavMeshDrawOffset, Vertices[i + 1] + NavMeshDrawOffset, FColor(Colors[i]), 0.0f);
+				AuxLines.Emplace(Vertices[i] + NavMeshDrawOffset, Vertices[i + 1] + NavMeshDrawOffset, Colors[i], 0.0f);
 			}
 		}
 	}
@@ -1046,12 +1046,12 @@ void FNavMeshSceneProxyData::AddMeshForInternalData(const FRecastInternalDebugDa
 	if (InInternalData.PointVertices.Num() > 0)
 	{
 		const TArray<FVector>& Vertices = InInternalData.PointVertices;
-		const TArray<uint32>& Colors = InInternalData.PointColors;
+		const TArray<FColor>& Colors = InInternalData.PointColors;
 		if (ensure(Vertices.Num() == Colors.Num()))
 		{
 			for (int32 i = 0; i < Vertices.Num(); i++)
 			{
-				AuxPoints.Emplace(Vertices[i] + NavMeshDrawOffset, FColor(Colors[i]), 5.0f);
+				AuxPoints.Emplace(Vertices[i] + NavMeshDrawOffset, Colors[i], 5.0f);
 			}
 		}
 	}
diff --git a/Engine/Source/Runtime/NavigationSystem/Private/NavMesh/RecastHelpers.cpp b/Engine/Source/Runtime/NavigationSystem/Private/NavMesh/RecastHelpers.cpp
--- a/Engine/Source/Runtime/NavigationSystem/Private/NavMesh/RecastHelpers.cpp
+++ b/Engine/Source/Runtime/NavigationSystem/Private/NavMesh/RecastHelpers.cpp
@@ -50,4 +50,13 @@ FMatrix Recast2UnrealMatrix()
 {
 	static FMatrix TM(FVector(-1.f, 0, 0), FVector(0, 0, -1.f), FVector(0, 1.f, 0), FVector::ZeroVector);
 	return TM;
+}
+
+FColor Recast2UnrealColor(const unsigned int RecastColor)
+{
+	const uint8 R = (RecastColor & 0xFF);
+	const uint8 G = (RecastColor & 0xFF00) >> 8;
+	const uint8 B = (RecastColor & 0xFF0000) >> 16;
+	const uint8 A = (RecastColor & 0xFF000000) >> 24;
+	return FColor(R,G,B,A);	// can't be direct assignation since internally FColor is not RGBA
 }
\ No newline at end of file
diff --git a/Engine/Source/Runtime/NavigationSystem/Public/NavMesh/RecastHelpers.h b/Engine/Source/Runtime/NavigationSystem/Public/NavMesh/RecastHelpers.h
--- a/Engine/Source/Runtime/NavigationSystem/Public/NavMesh/RecastHelpers.h
+++ b/Engine/Source/Runtime/NavigationSystem/Public/NavMesh/RecastHelpers.h
@@ -11,6 +11,7 @@
 #pragma once
 
 #include "CoreMinimal.h"
+#include "Math/Color.h"
 
 extern NAVIGATIONSYSTEM_API FVector Unreal2RecastPoint(const float* UnrealPoint);
 extern NAVIGATIONSYSTEM_API FVector Unreal2RecastPoint(const FVector& UnrealPoint);
@@ -22,3 +23,4 @@ extern NAVIGATIONSYSTEM_API FVector Recast2UnrealPoint(const FVector& RecastPoin
 extern NAVIGATIONSYSTEM_API FBox Recast2UnrealBox(const float* RecastMin, const float* RecastMax);
 extern NAVIGATIONSYSTEM_API FBox Recast2UnrealBox(const FBox& RecastBox);
 extern NAVIGATIONSYSTEM_API FMatrix Recast2UnrealMatrix();
+extern NAVIGATIONSYSTEM_API FColor Recast2UnrealColor(const unsigned int RecastColor);
diff --git a/Engine/Source/Runtime/NavigationSystem/Public/NavMesh/RecastInternalDebugData.cpp b/Engine/Source/Runtime/NavigationSystem/Public/NavMesh/RecastInternalDebugData.cpp
--- a/Engine/Source/Runtime/NavigationSystem/Public/NavMesh/RecastInternalDebugData.cpp
+++ b/Engine/Source/Runtime/NavigationSystem/Public/NavMesh/RecastInternalDebugData.cpp
@@ -7,21 +7,22 @@ void FRecastInternalDebugData::vertex(const float x, const float y, const float
 {
 	float RecastPos[3] = { x,y,z };
 	const FVector Pos = Recast2UnrealPoint(RecastPos);
+	const FColor Color = Recast2UnrealColor(color);
 	switch(CurrentPrim)
 	{
 	case DU_DRAW_POINTS:
 		PointVertices.Push(Pos);
-		PointColors.Push(color);
+		PointColors.Push(Color);
 		break;
 	case DU_DRAW_LINES:
 		LineVertices.Push(Pos);
-		LineColors.Push(color);
+		LineColors.Push(Color);
 		break;
 	case DU_DRAW_TRIS:
 		// Fallthrough
 	case DU_DRAW_QUADS:
 		TriangleVertices.Push(Pos);
-		TriangleColors.Push(color);
+		TriangleColors.Push(Color);
 		break;
 	}
 }
diff --git a/Engine/Source/Runtime/NavigationSystem/Public/NavMesh/RecastInternalDebugData.h b/Engine/Source/Runtime/NavigationSystem/Public/NavMesh/RecastInternalDebugData.h
--- a/Engine/Source/Runtime/NavigationSystem/Public/NavMesh/RecastInternalDebugData.h
+++ b/Engine/Source/Runtime/NavigationSystem/Public/NavMesh/RecastInternalDebugData.h
@@ -4,6 +4,7 @@
 
 #include "NavMesh/RecastHelpers.h"
 #include "DebugUtils/DebugDraw.h"
+#include "Math/Color.h"
 
 struct FRecastInternalDebugData : public duDebugDraw
 {
@@ -12,13 +13,13 @@ struct FRecastInternalDebugData : public duDebugDraw
 
 	TArray<uint32> TriangleIndices;
 	TArray<FVector> TriangleVertices;
-	TArray<uint32> TriangleColors;
+	TArray<FColor> TriangleColors;
 
 	TArray<FVector> LineVertices;
-	TArray<uint32>  LineColors;
+	TArray<FColor>  LineColors;
 
 	TArray<FVector> PointVertices;
-	TArray<uint32>  PointColors;
+	TArray<FColor>  PointColors;
 
 	FRecastInternalDebugData() {}
 	virtual ~FRecastInternalDebugData() override {}
