[Navmesh] It's now possible to select and display internal navmesh tile generation steps by setting properties on a navmesh actor.

diff --git a/Engine/Source/Runtime/NavigationSystem/Private/NavMesh/RecastNavMesh.cpp b/Engine/Source/Runtime/NavigationSystem/Private/NavMesh/RecastNavMesh.cpp
--- a/Engine/Source/Runtime/NavigationSystem/Private/NavMesh/RecastNavMesh.cpp
+++ b/Engine/Source/Runtime/NavigationSystem/Private/NavMesh/RecastNavMesh.cpp
@@ -353,6 +353,24 @@ FRecastNavMeshGenerationProperties::FRecastNavMeshGenerationProperties(const ARe
 	bFixedTilePoolSize = RecastNavMesh.bFixedTilePoolSize;
 }
 
+FRecastNavMeshTileGenerationDebug::FRecastNavMeshTileGenerationDebug()
+{
+	bEnabled = false;
+	TileCoordinate = FIntVector::ZeroValue;
+	bHeightfieldSolidFromRasterization = false;
+	bHeightfieldSolidPostRadiusFiltering = false;
+	bHeightfieldSolidPostHeightFiltering = false;
+	bCompactHeightfield = false;
+	bCompactHeightfieldEroded = false;
+	bCompactHeightfieldRegions = false;
+	bCompactHeightfieldDistances = false;
+	bTileCacheLayerAreas = false;
+	bTileCacheLayerRegions = false;
+	bTileCacheContours = false;
+	bTileCachePolyMesh = false;
+	bTileCacheDetailMesh = false;
+}
+
 ARecastNavMesh::FNavPolyFlags ARecastNavMesh::NavLinkFlag = ARecastNavMesh::FNavPolyFlags(0);
 
 ARecastNavMesh::ARecastNavMesh(const FObjectInitializer& ObjectInitializer)
diff --git a/Engine/Source/Runtime/NavigationSystem/Private/NavMesh/RecastNavMeshGenerator.cpp b/Engine/Source/Runtime/NavigationSystem/Private/NavMesh/RecastNavMeshGenerator.cpp
--- a/Engine/Source/Runtime/NavigationSystem/Private/NavMesh/RecastNavMeshGenerator.cpp
+++ b/Engine/Source/Runtime/NavigationSystem/Private/NavMesh/RecastNavMeshGenerator.cpp
@@ -62,10 +62,8 @@ static int32 GNavmeshSynchronousTileGeneration = 0;
 static FAutoConsoleVariableRef NavmeshVarSynchronous(TEXT("n.GNavmeshSynchronousTileGeneration"), GNavmeshSynchronousTileGeneration, TEXT(""), ECVF_Default);
 
 #if RECAST_INTERNAL_DEBUG_DATA
-static int32 GNavmeshDisplayStep = 0;
 static int32 GNavmeshDebugTileX = MAX_int32;
 static int32 GNavmeshDebugTileY = MAX_int32;
-static FAutoConsoleVariableRef NavmeshVarDisplayStep(TEXT("n.GNavmeshDisplayStep"), GNavmeshDisplayStep, TEXT(""), ECVF_Default);
 static FAutoConsoleVariableRef NavmeshVarDebugTileX(TEXT("n.GNavmeshDebugTileX"), GNavmeshDebugTileX, TEXT(""), ECVF_Default);
 static FAutoConsoleVariableRef NavmeshVarDebugTileY(TEXT("n.GNavmeshDebugTileY"), GNavmeshDebugTileY, TEXT(""), ECVF_Default);
 #endif //RECAST_INTERNAL_DEBUG_DATA
@@ -1936,6 +1934,7 @@ FRecastTileGenerator::FRecastTileGenerator(FRecastNavMeshGenerator& ParentGenera
 	TileY = Location.Y;
 
 	TileConfig = ParentGenerator.GetConfig();
+	TileDebugSettings = ParentGenerator.GetTileDebugSettings();
 	Version = ParentGenerator.GetVersion();
 	AdditionalCachedData = ParentGenerator.GetAdditionalCachedData();
 
@@ -2253,7 +2252,7 @@ void FRecastTileGenerator::GatherNavigationDataGeometry(const TSharedRef<FNaviga
 	bool bDumpGeometryData = false;
 
 #if RECAST_INTERNAL_DEBUG_DATA
-	if (IsTileToDebug())
+	if (IsTileDebugActive())
 	{
 		UE_LOG(LogNavigation, Log, TEXT("Gathering geometry for tile (%i,%i): %s."), TileX, TileY, *GetFullNameSafe(ElementData->GetOwner()));
 		UE_LOG(LogNavigation, Log, TEXT("                       bounds: %s"), *ElementData->Bounds.ToString());
@@ -3510,7 +3509,7 @@ bool FRecastTileGenerator::GenerateCompressedLayers(FNavMeshBuildContext& BuildC
 	}
 
 #if RECAST_INTERNAL_DEBUG_DATA
-	if (GNavmeshDisplayStep == 10 && IsTileToDebug())
+	if (IsTileDebugActive() && TileDebugSettings.bHeightfieldSolidFromRasterization)
 	{
 		duDebugDrawHeightfieldSolid(&BuildContext.InternalDebugData, *RasterContext.SolidHF);
 	}
@@ -3523,7 +3522,7 @@ bool FRecastTileGenerator::GenerateCompressedLayers(FNavMeshBuildContext& BuildC
 	}
 
 #if RECAST_INTERNAL_DEBUG_DATA
-	if (GNavmeshDisplayStep == 20 && IsTileToDebug())
+	if (IsTileDebugActive() && TileDebugSettings.bHeightfieldSolidPostRadiusFiltering)
 	{
 		duDebugDrawHeightfieldSolid(&BuildContext.InternalDebugData, *RasterContext.SolidHF);
 	}
@@ -3532,7 +3531,7 @@ bool FRecastTileGenerator::GenerateCompressedLayers(FNavMeshBuildContext& BuildC
 	GenerateRecastFilter(BuildContext, RasterContext);
 
 #if RECAST_INTERNAL_DEBUG_DATA
-	if (GNavmeshDisplayStep == 30 && IsTileToDebug())
+	if (IsTileDebugActive() && TileDebugSettings.bHeightfieldSolidPostHeightFiltering)
 	{
 		duDebugDrawHeightfieldSolid(&BuildContext.InternalDebugData, *RasterContext.SolidHF);
 	}
@@ -3544,7 +3543,7 @@ bool FRecastTileGenerator::GenerateCompressedLayers(FNavMeshBuildContext& BuildC
 	}
 
 #if RECAST_INTERNAL_DEBUG_DATA
-	if (GNavmeshDisplayStep == 40 && IsTileToDebug())
+	if (IsTileDebugActive() && TileDebugSettings.bCompactHeightfield)
 	{
 		duDebugDrawCompactHeightfieldSolid(&BuildContext.InternalDebugData, *RasterContext.CompactHF);
 	}
@@ -3556,7 +3555,7 @@ bool FRecastTileGenerator::GenerateCompressedLayers(FNavMeshBuildContext& BuildC
 	}
 
 #if RECAST_INTERNAL_DEBUG_DATA
-	if (GNavmeshDisplayStep == 50 && IsTileToDebug())
+	if (IsTileDebugActive() && TileDebugSettings.bCompactHeightfieldEroded)
 	{
 		duDebugDrawCompactHeightfieldSolid(&BuildContext.InternalDebugData, *RasterContext.CompactHF);
 	}
@@ -3567,6 +3566,21 @@ bool FRecastTileGenerator::GenerateCompressedLayers(FNavMeshBuildContext& BuildC
 		return false;
 	}
 
+#if RECAST_INTERNAL_DEBUG_DATA
+	if (IsTileDebugActive())
+	{
+		if (TileDebugSettings.bCompactHeightfieldRegions)
+		{
+			duDebugDrawCompactHeightfieldRegions(&BuildContext.InternalDebugData, *RasterContext.CompactHF);	
+		}
+
+		if (TileDebugSettings.bCompactHeightfieldDistances)
+		{
+			duDebugDrawCompactHeightfieldDistance(&BuildContext.InternalDebugData, *RasterContext.CompactHF);	
+		}
+	}
+#endif
+	
 	return RecastBuildTileCache(BuildContext, RasterContext);
 }
 
@@ -3662,6 +3676,13 @@ bool FRecastTileGenerator::GenerateNavigationDataLayer(FNavMeshBuildContext& Bui
 		}
 		else
 		{
+#if RECAST_INTERNAL_DEBUG_DATA
+			if (IsTileDebugActive() && TileDebugSettings.bTileCacheLayerAreas)
+			{
+				duDebugDrawTileCacheLayerAreas(&BuildContext.InternalDebugData, *GenerationContext.Layer, TileConfig.cs, TileConfig.ch);
+			}
+#endif
+			
 			status = dtBuildTileCacheRegionsChunky(&GenNavAllocator, TileConfig.minRegionArea, TileConfig.mergeRegionArea, *GenerationContext.Layer, TileConfig.TileCacheChunkSize);
 		}
 
@@ -3679,7 +3700,7 @@ bool FRecastTileGenerator::GenerateNavigationDataLayer(FNavMeshBuildContext& Bui
 	}
 
 #if RECAST_INTERNAL_DEBUG_DATA
-	if (GNavmeshDisplayStep == 100 && IsTileToDebug())
+	if (IsTileDebugActive() && TileDebugSettings.bTileCacheLayerRegions)
 	{
 		duDebugDrawTileCacheLayerRegions(&BuildContext.InternalDebugData, *GenerationContext.Layer, TileConfig.cs, TileConfig.ch);
 	}
@@ -3726,7 +3747,7 @@ bool FRecastTileGenerator::GenerateNavigationDataLayer(FNavMeshBuildContext& Bui
 	}
 
 #if RECAST_INTERNAL_DEBUG_DATA
-	if (GNavmeshDisplayStep == 110 && IsTileToDebug())
+	if (IsTileDebugActive() && TileDebugSettings.bTileCacheContours)
 	{
 		duDebugDrawTileCacheContours(&BuildContext.InternalDebugData, *GenerationContext.ContourSet, GenerationContext.Layer->header->bmin, TileConfig.cs, TileConfig.ch);
 	}
@@ -3761,7 +3782,7 @@ bool FRecastTileGenerator::GenerateNavigationDataLayer(FNavMeshBuildContext& Bui
 	}
 
 #if RECAST_INTERNAL_DEBUG_DATA
-	if (GNavmeshDisplayStep == 120 && IsTileToDebug())
+	if (IsTileDebugActive() && TileDebugSettings.bTileCachePolyMesh)
 	{
 		duDebugDrawTileCachePolyMesh(&BuildContext.InternalDebugData, *GenerationContext.PolyMesh, GenerationContext.Layer->header->bmin, TileConfig.cs, TileConfig.ch);
 	}
@@ -3789,7 +3810,7 @@ bool FRecastTileGenerator::GenerateNavigationDataLayer(FNavMeshBuildContext& Bui
 		}
 
 #if RECAST_INTERNAL_DEBUG_DATA
-		if (GNavmeshDisplayStep == 130 && IsTileToDebug())
+		if (IsTileDebugActive() && TileDebugSettings.bTileCacheDetailMesh)
 		{
 			duDebugDrawTileCacheDetailMesh(&BuildContext.InternalDebugData, *GenerationContext.DetailMesh);
 		}
@@ -6860,8 +6881,9 @@ public:
 #endif // WITH_RECAST
 
 #if RECAST_INTERNAL_DEBUG_DATA
-bool FRecastTileGenerator::IsTileToDebug()
+bool FRecastTileGenerator::IsTileDebugActive()
 {
-	return TileX == GNavmeshDebugTileX && TileY == GNavmeshDebugTileY;
+	const FIntVector& Coord = TileDebugSettings.TileCoordinate; 
+	return (TileDebugSettings.bEnabled && TileX == Coord.X && TileY == Coord.Y) || (TileX == GNavmeshDebugTileX && TileY == GNavmeshDebugTileY);
 }
-#endif //RECAST_INTERNAL_DEBUG_DATA
\ No newline at end of file
+#endif //RECAST_INTERNAL_DEBUG_DATA
diff --git a/Engine/Source/Runtime/NavigationSystem/Public/NavMesh/RecastNavMesh.h b/Engine/Source/Runtime/NavigationSystem/Public/NavMesh/RecastNavMesh.h
--- a/Engine/Source/Runtime/NavigationSystem/Public/NavMesh/RecastNavMesh.h
+++ b/Engine/Source/Runtime/NavigationSystem/Public/NavMesh/RecastNavMesh.h
@@ -377,6 +377,58 @@ struct NAVIGATIONSYSTEM_API FRecastNavMeshGenerationProperties
 	FRecastNavMeshGenerationProperties(const ARecastNavMesh& RecastNavMesh);
 };
 
+USTRUCT()
+struct NAVIGATIONSYSTEM_API FRecastNavMeshTileGenerationDebug
+{
+	GENERATED_BODY()
+	
+	FRecastNavMeshTileGenerationDebug();
+
+	/** If set, the selected internal debug data will be kept during tile generation to be displayed with the navmesh. */
+	UPROPERTY(Transient, EditAnywhere, Category = Debug)
+	uint32 bEnabled : 1;
+
+	/** Selected tile coordinate, only this tile will have it's internal data kept.
+	 *  Tip: displaying the navmesh using 'Draw Tile Labels' show tile coordinates. */ 
+	UPROPERTY(EditAnywhere, Category = Debug)
+	FIntVector TileCoordinate = FIntVector::ZeroValue;
+
+	UPROPERTY(EditAnywhere, Category = Debug)
+	uint32 bHeightfieldSolidFromRasterization : 1;
+
+	UPROPERTY(EditAnywhere, Category = Debug)
+	uint32 bHeightfieldSolidPostRadiusFiltering : 1;
+
+	UPROPERTY(EditAnywhere, Category = Debug)
+	uint32 bHeightfieldSolidPostHeightFiltering : 1;
+
+	UPROPERTY(EditAnywhere, Category = Debug)
+	uint32 bCompactHeightfield : 1;
+
+	UPROPERTY(EditAnywhere, Category = Debug)
+	uint32 bCompactHeightfieldEroded : 1;
+
+	UPROPERTY(EditAnywhere, Category = Debug)
+	uint32 bCompactHeightfieldRegions : 1;
+
+	UPROPERTY(EditAnywhere, Category = Debug)
+	uint32 bCompactHeightfieldDistances : 1;
+
+	UPROPERTY(EditAnywhere, Category = Debug)
+	uint32 bTileCacheLayerAreas : 1;
+
+	UPROPERTY(EditAnywhere, Category = Debug)
+	uint32 bTileCacheLayerRegions : 1;
+
+	UPROPERTY(EditAnywhere, Category = Debug)
+	uint32 bTileCacheContours : 1;
+
+	UPROPERTY(EditAnywhere, Category = Debug)
+	uint32 bTileCachePolyMesh : 1;
+
+	UPROPERTY(EditAnywhere, Category = Debug)
+	uint32 bTileCacheDetailMesh : 1;
+};
 
 /**
  *	Structure to handle nav mesh tile's raw data persistence and releasing
@@ -528,6 +580,9 @@ class NAVIGATIONSYSTEM_API ARecastNavMesh : public ANavigationData
 	/** vertical offset added to navmesh's debug representation for better readability */
 	UPROPERTY(EditAnywhere, Category=Display, config)
 	float DrawOffset;
+
+	UPROPERTY(EditAnywhere, Category = Display)
+	FRecastNavMeshTileGenerationDebug TileGenerationDebug;
 	
 	//----------------------------------------------------------------------//
 	// NavMesh generation parameters
diff --git a/Engine/Source/Runtime/NavigationSystem/Public/NavMesh/RecastNavMeshGenerator.h b/Engine/Source/Runtime/NavigationSystem/Public/NavMesh/RecastNavMeshGenerator.h
--- a/Engine/Source/Runtime/NavigationSystem/Public/NavMesh/RecastNavMeshGenerator.h
+++ b/Engine/Source/Runtime/NavigationSystem/Public/NavMesh/RecastNavMeshGenerator.h
@@ -421,7 +421,7 @@ protected:
 	void DumpAsyncData();
 
 #if RECAST_INTERNAL_DEBUG_DATA
-	bool IsTileToDebug();
+	bool IsTileDebugActive();
 #endif
 
 protected:
@@ -490,6 +490,8 @@ protected:
 	TWeakObjectPtr<UNavigationSystemV1> NavSystem; 
 	FNavDataConfig NavDataConfig;
 
+	FRecastNavMeshTileGenerationDebug TileDebugSettings;
+
 #if RECAST_INTERNAL_DEBUG_DATA
 	FRecastInternalDebugData DebugData;
 #endif
@@ -665,6 +667,8 @@ public:
 
 	const FRecastBuildConfig& GetConfig() const { return Config; }
 
+	const FRecastNavMeshTileGenerationDebug& GetTileDebugSettings() const { return DestNavMesh->TileGenerationDebug; } 
+
 	const TNavStatArray<FBox>& GetInclusionBounds() const { return InclusionBounds; }
 	
 	FVector GetRcNavMeshOrigin() const { return RcNavMeshOrigin; }
