Modified logging to include the mapname for better debugging of Cooking warnings.

diff --git a/Engine/Source/Runtime/NavigationSystem/Private/NavMesh/RecastNavMesh.cpp b/Engine/Source/Runtime/NavigationSystem/Private/NavMesh/RecastNavMesh.cpp
--- a/Engine/Source/Runtime/NavigationSystem/Private/NavMesh/RecastNavMesh.cpp
+++ b/Engine/Source/Runtime/NavigationSystem/Private/NavMesh/RecastNavMesh.cpp
@@ -796,7 +796,7 @@ void ARecastNavMesh::Serialize( FArchive& Ar )
 	{
 		if (NavMeshVersion < NAVMESHVER_MIN_COMPATIBLE)
 		{
-			UE_LOG(LogNavigation, Warning, TEXT("%s: ARecastNavMesh: Nav mesh version %d < Min compatible %d. Nav mesh must be rebuilt in the editor \n"), *GetName(), NavMeshVersion, NAVMESHVER_MIN_COMPATIBLE);
+			UE_LOG(LogNavigation, Warning, TEXT("%s: ARecastNavMesh: Nav mesh version %d < Min compatible %d. Nav mesh needs to be rebuilt. \n"), *GetFullName(), NavMeshVersion, NAVMESHVER_MIN_COMPATIBLE);
 
 			// incompatible, just skip over this data.  navmesh needs rebuilt.
 			Ar.Seek( RecastNavMeshSizePos + RecastNavMeshSizeBytes );
diff --git a/Engine/Source/Runtime/NavigationSystem/Private/NavMesh/RecastNavMeshDataChunk.cpp b/Engine/Source/Runtime/NavigationSystem/Private/NavMesh/RecastNavMeshDataChunk.cpp
--- a/Engine/Source/Runtime/NavigationSystem/Private/NavMesh/RecastNavMeshDataChunk.cpp
+++ b/Engine/Source/Runtime/NavigationSystem/Private/NavMesh/RecastNavMeshDataChunk.cpp
@@ -88,7 +88,7 @@ void URecastNavMeshDataChunk::Serialize(FArchive& Ar)
 	{
 		if (NavMeshVersion < NAVMESHVER_MIN_COMPATIBLE)
 		{
-			UE_LOG(LogNavigation, Warning, TEXT("%s: URecastNavMeshDataChunk: Nav mesh version %d < Min compatLWCible %d. Nav mesh must be rebuilt in the editor \n"), *GetName(), NavMeshVersion, NAVMESHVER_MIN_COMPATIBLE);
+			UE_LOG(LogNavigation, Warning, TEXT("%s: URecastNavMeshDataChunk: Nav mesh version %d < Min compatible %d. Nav mesh needs to be rebuilt. \n"), *GetFullName(), NavMeshVersion, NAVMESHVER_MIN_COMPATIBLE);
 
 			// incompatible, just skip over this data. Navmesh needs rebuilt.
 			Ar.Seek(RecastNavMeshSizePos + RecastNavMeshSizeBytes);
