(per navmesh) avoided dirtying nav areas that have been requested by a source object that is not supposed to generate geometry onto the navmesh (via FNavigationOctreeElement::ShouldUseGeometry): optional via bUseVirtualGeometryFilteringAndDirtying (disabled by default)
Added a new weakptr "OptionalSourceObject" in the FNavigationDirtyArea, needed to support the previous change

diff --git a/Engine/Source/Runtime/Engine/Classes/AI/Navigation/NavRelevantInterface.h b/Engine/Source/Runtime/Engine/Classes/AI/Navigation/NavRelevantInterface.h
--- a/Engine/Source/Runtime/Engine/Classes/AI/Navigation/NavRelevantInterface.h
+++ b/Engine/Source/Runtime/Engine/Classes/AI/Navigation/NavRelevantInterface.h
@@ -56,7 +56,10 @@ struct ENGINE_API FNavigationRelevantData : public TSharedFromThis<FNavigationRe
 	/** Gathers per instance data for navigation geometry in a specified area box */
 	FNavDataPerInstanceTransformDelegate NavDataPerInstanceTransformDelegate;
 
-	/** called to check if hosted geometry should be used for given FNavDataConfig. If not set then "true" is assumed. */
+	/** called to check if hosted geometry should be used for given FNavDataConfig. If not set then "true" is assumed.
+	 *  Might want to set bUseVirtualGeometryFilteringAndDirtying to true in the Navmesh class you are excluding geometry from.
+	 *  This will improve cpu performance by stopping the navmesh from dirtying tiles requested by actors being excluded by this delegate.
+	 */
 	FFilterNavDataDelegate ShouldUseGeometryDelegate;
 
 	/** additional modifiers: areas and external links */
diff --git a/Engine/Source/Runtime/Engine/Classes/AI/Navigation/NavigationTypes.h b/Engine/Source/Runtime/Engine/Classes/AI/Navigation/NavigationTypes.h
--- a/Engine/Source/Runtime/Engine/Classes/AI/Navigation/NavigationTypes.h
+++ b/Engine/Source/Runtime/Engine/Classes/AI/Navigation/NavigationTypes.h
@@ -91,9 +91,10 @@ struct FNavigationDirtyArea
 {
 	FBox Bounds;
 	int32 Flags;
+	TWeakObjectPtr<UObject> OptionalSourceObject;
 	
 	FNavigationDirtyArea() : Flags(0) {}
-	FNavigationDirtyArea(const FBox& InBounds, int32 InFlags) : Bounds(InBounds), Flags(InFlags) {}
+	FNavigationDirtyArea(const FBox& InBounds, int32 InFlags, UObject* const InOptionalSourceObject = nullptr) : Bounds(InBounds), Flags(InFlags), OptionalSourceObject(InOptionalSourceObject) {}
 	FORCEINLINE bool HasFlag(ENavigationDirtyFlag::Type Flag) const { return (Flags & Flag) != 0; }
 };
 
diff --git a/Engine/Source/Runtime/NavigationSystem/Private/NavMesh/RecastNavMesh.cpp b/Engine/Source/Runtime/NavigationSystem/Private/NavMesh/RecastNavMesh.cpp
--- a/Engine/Source/Runtime/NavigationSystem/Private/NavMesh/RecastNavMesh.cpp
+++ b/Engine/Source/Runtime/NavigationSystem/Private/NavMesh/RecastNavMesh.cpp
@@ -394,6 +394,7 @@ ARecastNavMesh::ARecastNavMesh(const FObjectInitializer& ObjectInitializer)
 	, bFilterLowSpanFromTileCache(false)
 	, bStoreEmptyTileLayers(false)
 	, bUseVirtualFilters(true)
+	, bUseVirtualGeometryFilteringAndDirtying(false)
 	, bAllowNavLinkAsPathEnd(false)
 	, TileSetUpdateInterval(1.0f)
 	, NavMeshVersion(NAVMESHVER_LATEST)	
diff --git a/Engine/Source/Runtime/NavigationSystem/Private/NavMesh/RecastNavMeshGenerator.cpp b/Engine/Source/Runtime/NavigationSystem/Private/NavMesh/RecastNavMeshGenerator.cpp
--- a/Engine/Source/Runtime/NavigationSystem/Private/NavMesh/RecastNavMeshGenerator.cpp
+++ b/Engine/Source/Runtime/NavigationSystem/Private/NavMesh/RecastNavMeshGenerator.cpp
@@ -2269,10 +2269,16 @@ void FRecastTileGenerator::PrepareGeometrySources(const FRecastNavMeshGenerator&
 	NavigationRelevantData.Reset();
 	NavSystem = NavSys;
 	bUpdateGeometry = bGeometryChanged;
+	
+	const ARecastNavMesh* const OwnerNav = ParentGenerator.GetOwner();
+	const bool bUseVirtualGeometryFilteringAndDirtying = OwnerNav != nullptr && OwnerNav->bUseVirtualGeometryFilteringAndDirtying;
 
-	NavOctreeInstance->FindElementsWithBoundsTest(ParentGenerator.GrowBoundingBox(TileBB, /*bIncludeAgentHeight*/ false), [this, bGeometryChanged](const FNavigationOctreeElement& Element)
+	NavOctreeInstance->FindElementsWithBoundsTest(ParentGenerator.GrowBoundingBox(TileBB, /*bIncludeAgentHeight*/ false),
+		[&ParentGenerator,this, bGeometryChanged, bUseVirtualGeometryFilteringAndDirtying](const FNavigationOctreeElement& Element)
 	{
-		const bool bShouldUse = Element.ShouldUseGeometry(NavDataConfig);
+		const bool bShouldUse = bUseVirtualGeometryFilteringAndDirtying ?
+			ParentGenerator.ShouldGenerateGeometryForOctreeElement(Element, NavDataConfig) :
+			Element.ShouldUseGeometry(NavDataConfig);
 		if (bShouldUse)
 		{
 			const bool bExportGeometry = bGeometryChanged && (Element.Data->HasGeometry() || Element.Data->IsPendingLazyGeometryGathering());
@@ -2297,11 +2303,16 @@ void FRecastTileGenerator::GatherGeometry(const FRecastNavMeshGenerator& ParentG
 	{
 		return;
 	}
-	const FNavDataConfig& OwnerNavDataConfig = ParentGenerator.GetOwner()->GetConfig();
+	const ARecastNavMesh* const OwnerNav = ParentGenerator.GetOwner();
+	const bool bUseVirtualGeometryFilteringAndDirtying = OwnerNav->bUseVirtualGeometryFilteringAndDirtying;
+	const FNavDataConfig& OwnerNavDataConfig = OwnerNav->GetConfig();
 
-	NavigationOctree->FindElementsWithBoundsTest(ParentGenerator.GrowBoundingBox(TileBB, /*bIncludeAgentHeight*/ false), [&OwnerNavDataConfig, &NavigationOctree, this, NavSys, bGeometryChanged](const FNavigationOctreeElement& Element)
+	NavigationOctree->FindElementsWithBoundsTest(ParentGenerator.GrowBoundingBox(TileBB, /*bIncludeAgentHeight*/ false),
+		[&OwnerNavDataConfig, &ParentGenerator, this, NavSys, bGeometryChanged, bUseVirtualGeometryFilteringAndDirtying](const FNavigationOctreeElement& Element)
 	{
-		const bool bShouldUse = Element.ShouldUseGeometry(OwnerNavDataConfig);
+		const bool bShouldUse = bUseVirtualGeometryFilteringAndDirtying ?
+			ParentGenerator.ShouldGenerateGeometryForOctreeElement(Element, OwnerNavDataConfig) :
+			Element.ShouldUseGeometry(OwnerNavDataConfig);
 		if (bShouldUse)
 		{
 			GatherNavigationDataGeometry(Element.Data, *NavSys, OwnerNavDataConfig, bGeometryChanged);
@@ -5594,6 +5605,11 @@ FBox FRecastNavMeshGenerator::GrowBoundingBox(const FBox& BBox, bool bIncludeAge
 	return FBox(BBox.Min - BBoxGrowth - BBoxGrowOffsetMin, BBox.Max + BBoxGrowth);
 }
 
+bool FRecastNavMeshGenerator::ShouldGenerateGeometryForOctreeElement(const FNavigationOctreeElement& Element, const FNavDataConfig& NavDataConfig) const
+{
+	return Element.ShouldUseGeometry(NavDataConfig);
+}
+
 static bool IntersectBounds(const FBox& TestBox, const TNavStatArray<FBox>& Bounds)
 {
 	for (const FBox& Box : Bounds)
@@ -5693,6 +5709,14 @@ void FRecastNavMeshGenerator::MarkDirtyTiles(const TArray<FNavigationDirtyArea>&
 	check(TileSizeInWorldUnits > 0);
 
 	const bool bGameStaticNavMesh = IsGameStaticNavMesh(DestNavMesh);
+
+	const ARecastNavMesh* const OwnerNav = GetOwner();
+	const bool bUseVirtualGeometryFilteringAndDirtying = OwnerNav != nullptr && OwnerNav->bUseVirtualGeometryFilteringAndDirtying;
+	// Those are set only if bUseVirtualGeometryFilteringAndDirtying is enabled since we do not use them for anything else
+	const UNavigationSystemV1* const NavSys = bUseVirtualGeometryFilteringAndDirtying ? FNavigationSystem::GetCurrent<UNavigationSystemV1>(GetWorld()) : nullptr;
+	const FNavigationOctree* const NavOctreeInstance = (bUseVirtualGeometryFilteringAndDirtying && NavSys) ? NavSys->GetNavOctree() : nullptr;
+	const FNavDataConfig* const NavDataConfig = (bUseVirtualGeometryFilteringAndDirtying && DestNavMesh) ? &DestNavMesh->GetConfig() : nullptr;
+	// ~
 		
 	// find all tiles that need regeneration
 	TSet<FPendingTileElement> DirtyTiles;
@@ -5704,6 +5728,18 @@ void FRecastNavMeshGenerator::MarkDirtyTiles(const TArray<FNavigationDirtyArea>&
 			continue;
 		}
 		
+		// (if bUseVirtualGeometryFilteringAndDirtying is true) Ignore dirty areas flagged by a source object that is not supposed to apply to this navmesh
+		if (bUseVirtualGeometryFilteringAndDirtying && NavSys && NavOctreeInstance && NavDataConfig)
+		{
+			if (const UObject* const SourceObject = DirtyArea.OptionalSourceObject.Get())
+			{
+				if (!ShouldDirtyTilesRequestedByObject(*NavSys, *NavOctreeInstance, *SourceObject, *NavDataConfig))
+				{
+					continue;
+				}
+			}
+		}
+		
 		bool bDoTileInclusionTest = false;
 		FBox AdjustedAreaBounds = DirtyArea.Bounds;
 		
@@ -5823,6 +5859,14 @@ void FRecastNavMeshGenerator::MarkDirtyTiles(const TArray<FNavigationDirtyArea>&
 	}
 }
 
+bool FRecastNavMeshGenerator::ShouldDirtyTilesRequestedByObject(const UNavigationSystemV1& NavSys,
+	const FNavigationOctree& NavOctreeInstance, const UObject& SourceObject, const FNavDataConfig& NavDataConfig) const
+{
+	const FOctreeElementId2* const OctreeElementId = NavSys.GetObjectsNavOctreeId(SourceObject);
+
+	return (OctreeElementId == nullptr) || ShouldGenerateGeometryForOctreeElement(NavOctreeInstance.GetElementById(*OctreeElementId), NavDataConfig);
+}
+
 void FRecastNavMeshGenerator::SortPendingBuildTiles()
 {
 	if (bSortTilesWithSeedLocations == false)
@@ -6545,9 +6589,16 @@ void FRecastNavMeshGenerator::GrabDebugSnapshot(struct FVisualLogEntry* Snapshot
 		const ARecastNavMesh* NavData = Cast<const ARecastNavMesh>(NavSys->NavDataSet[Index]);
 		if (NavData)
 		{
-			NavOctree->FindElementsWithBoundsTest(BoundingBox, [this, NavData, &Indices, &CoordBuffer, Snapshot, &CategoryName, LogVerbosity, NavAreaVerbosity](const FNavigationOctreeElement& Element)
+			const bool bUseVirtualGeometryFilteringAndDirtying = NavData->bUseVirtualGeometryFilteringAndDirtying;
+			
+			NavOctree->FindElementsWithBoundsTest(BoundingBox,
+				[this, NavData, &Indices, &CoordBuffer, Snapshot, &CategoryName, LogVerbosity, NavAreaVerbosity, bUseVirtualGeometryFilteringAndDirtying](const FNavigationOctreeElement& Element)
 			{
-				const bool bExportGeometry = Element.Data->HasGeometry() && Element.ShouldUseGeometry(DestNavMesh->GetConfig());
+				const bool bExportGeometry = Element.Data->HasGeometry() && (
+					bUseVirtualGeometryFilteringAndDirtying ?
+						ShouldGenerateGeometryForOctreeElement(Element, NavData->GetConfig()) :
+						Element.ShouldUseGeometry(NavData->GetConfig())
+					);
 
 				TArray<FTransform> InstanceTransforms;
 				Element.Data->NavDataPerInstanceTransformDelegate.ExecuteIfBound(Element.Bounds.GetBox(), InstanceTransforms);
@@ -6680,6 +6731,8 @@ void FRecastNavMeshGenerator::ExportNavigationData(const FString& FileName) cons
 		const ARecastNavMesh* NavData = Cast<const ARecastNavMesh>(NavSys->NavDataSet[Index]);
 		if (NavData)
 		{
+			const bool bUseVirtualGeometryFilteringAndDirtying = NavData->bUseVirtualGeometryFilteringAndDirtying;
+			
 			struct FAreaExportData
 			{
 				FConvexNavAreaData Convex;
@@ -6687,9 +6740,14 @@ void FRecastNavMeshGenerator::ExportNavigationData(const FString& FileName) cons
 			};
 			TArray<FAreaExportData> AreaExport;
 
-			NavOctree->FindElementsWithBoundsTest(TotalNavBounds, [this, NavData, &IndexBuffer, &CoordBuffer, &AreaExport](const FNavigationOctreeElement& Element)
+			NavOctree->FindElementsWithBoundsTest(TotalNavBounds,
+				[this, NavData, &IndexBuffer, &CoordBuffer, &AreaExport, bUseVirtualGeometryFilteringAndDirtying](const FNavigationOctreeElement& Element)
 			{
-				const bool bExportGeometry = Element.Data->HasGeometry() && Element.ShouldUseGeometry(DestNavMesh->GetConfig());
+				const bool bExportGeometry = Element.Data->HasGeometry() && (
+					bUseVirtualGeometryFilteringAndDirtying ?
+						ShouldGenerateGeometryForOctreeElement(Element, NavData->GetConfig()) :
+						Element.ShouldUseGeometry(NavData->GetConfig())
+				);
 
 				TArray<FTransform> InstanceTransforms;
 				Element.Data->NavDataPerInstanceTransformDelegate.ExecuteIfBound(Element.Bounds.GetBox(), InstanceTransforms);
diff --git a/Engine/Source/Runtime/NavigationSystem/Private/NavigationDirtyAreasController.cpp b/Engine/Source/Runtime/NavigationSystem/Private/NavigationDirtyAreasController.cpp
--- a/Engine/Source/Runtime/NavigationSystem/Private/NavigationDirtyAreasController.cpp
+++ b/Engine/Source/Runtime/NavigationSystem/Private/NavigationDirtyAreasController.cpp
@@ -90,7 +90,7 @@ void FNavigationDirtyAreasController::AddArea(const FBox& NewArea, const int32 F
 
 	if (Flags > 0 && bCanAccumulateDirtyAreas)
 	{
-		DirtyAreas.Add(FNavigationDirtyArea(NewArea, Flags));
+		DirtyAreas.Add(FNavigationDirtyArea(NewArea, Flags, ObjectProviderFunc ? ObjectProviderFunc() : nullptr));
 	}
 }
 
diff --git a/Engine/Source/Runtime/NavigationSystem/Public/NavMesh/RecastNavMesh.h b/Engine/Source/Runtime/NavigationSystem/Public/NavMesh/RecastNavMesh.h
--- a/Engine/Source/Runtime/NavigationSystem/Public/NavMesh/RecastNavMesh.h
+++ b/Engine/Source/Runtime/NavigationSystem/Public/NavMesh/RecastNavMesh.h
@@ -735,6 +735,12 @@ class NAVIGATIONSYSTEM_API ARecastNavMesh : public ANavigationData
 	UPROPERTY(config)
 	uint32 bUseVirtualFilters : 1;
 
+	/** Indicates whether use the virtual methods to check if an object should generate geometry or if we should call the normal method directly (i.e. FNavigationOctreeElement::ShouldUseGeometry).
+	 *  If enabled, will also check if an object requesting an update on the navmesh is excluded to avoid dirtying the areas unnecessarily.
+	 *  Defaults to false. */
+	UPROPERTY(config)
+	uint32 bUseVirtualGeometryFilteringAndDirtying : 1;
+
 	/** If set, paths can end at navlink poly (not the ground one!) */
 	UPROPERTY(config)
 	uint32 bAllowNavLinkAsPathEnd : 1;
diff --git a/Engine/Source/Runtime/NavigationSystem/Public/NavMesh/RecastNavMeshGenerator.h b/Engine/Source/Runtime/NavigationSystem/Public/NavMesh/RecastNavMeshGenerator.h
--- a/Engine/Source/Runtime/NavigationSystem/Public/NavMesh/RecastNavMeshGenerator.h
+++ b/Engine/Source/Runtime/NavigationSystem/Public/NavMesh/RecastNavMeshGenerator.h
@@ -691,6 +691,9 @@ public:
 
 	FBox GrowBoundingBox(const FBox& BBox, bool bIncludeAgentHeight) const;
 	
+	/** Returns if the provided Octree Element should generate geometry on the provided NavDataConfig. Can be used to extend the logic to decide what geometry is generated on what Navmesh */
+	virtual bool ShouldGenerateGeometryForOctreeElement(const FNavigationOctreeElement& Element, const FNavDataConfig& NavDataConfig) const;
+	
 #if !(UE_BUILD_SHIPPING || UE_BUILD_TEST) && ENABLE_VISUAL_LOG
 	virtual void ExportNavigationData(const FString& FileName) const override;
 	virtual void GrabDebugSnapshot(struct FVisualLogEntry* Snapshot, const FBox& BoundingBox, const FName& CategoryName, ELogVerbosity::Type Verbosity) const override;
@@ -741,6 +744,9 @@ protected:
 	/** Marks grid tiles affected by specified areas as dirty */
 	virtual void MarkDirtyTiles(const TArray<FNavigationDirtyArea>& DirtyAreas);
 
+	/** Returns if the provided UObject that requested a navmesh dirtying should dirty this Navmesh. Useful to avoid tiles regeneration from objects that are excluded from the provided NavDataConfig */
+	virtual bool ShouldDirtyTilesRequestedByObject(const UNavigationSystemV1& NavSys, const FNavigationOctree& NavOctreeInstance, const UObject& SourceObject, const FNavDataConfig& NavDataConfig) const;
+
 	/** Marks all tiles overlapping with InclusionBounds dirty (via MarkDirtyTiles). */
 	bool MarkNavBoundsDirty();
 
diff --git a/Engine/Source/Runtime/NavigationSystem/Public/NavigationDirtyAreasController.h b/Engine/Source/Runtime/NavigationSystem/Public/NavigationDirtyAreasController.h
--- a/Engine/Source/Runtime/NavigationSystem/Public/NavigationDirtyAreasController.h
+++ b/Engine/Source/Runtime/NavigationSystem/Public/NavigationDirtyAreasController.h
@@ -46,7 +46,7 @@ public:
 	 * Accumulation must be allowed and flags valid otherwise the add is ignored.
 	 *	@param NewArea Bounding box of the affected area
 	 *	@param Flags Indicates the type of modification applied to the area
-	 *	@param ObjectProviderFunc Optional function to retrieve source object that can be use for error reporting
+	 *	@param ObjectProviderFunc Optional function to retrieve source object that can be use for error reporting and navmesh exclusion
 	 */
 	void AddArea(const FBox& NewArea, const int32 Flags, const TFunction<UObject*()>& ObjectProviderFunc = nullptr);
 	
