Removed unnecessary nullptr check on DestNavMesh which is triggering PVS for some old code

diff --git a/Engine/Source/Runtime/NavigationSystem/Private/NavMesh/RecastNavMeshGenerator.cpp b/Engine/Source/Runtime/NavigationSystem/Private/NavMesh/RecastNavMeshGenerator.cpp
--- a/Engine/Source/Runtime/NavigationSystem/Private/NavMesh/RecastNavMeshGenerator.cpp
+++ b/Engine/Source/Runtime/NavigationSystem/Private/NavMesh/RecastNavMeshGenerator.cpp
@@ -5715,7 +5715,7 @@ void FRecastNavMeshGenerator::MarkDirtyTiles(const TArray<FNavigationDirtyArea>&
 	// Those are set only if bUseVirtualGeometryFilteringAndDirtying is enabled since we do not use them for anything else
 	const UNavigationSystemV1* const NavSys = bUseVirtualGeometryFilteringAndDirtying ? FNavigationSystem::GetCurrent<UNavigationSystemV1>(GetWorld()) : nullptr;
 	const FNavigationOctree* const NavOctreeInstance = (bUseVirtualGeometryFilteringAndDirtying && NavSys) ? NavSys->GetNavOctree() : nullptr;
-	const FNavDataConfig* const NavDataConfig = (bUseVirtualGeometryFilteringAndDirtying && DestNavMesh) ? &DestNavMesh->GetConfig() : nullptr;
+	const FNavDataConfig* const NavDataConfig = bUseVirtualGeometryFilteringAndDirtying ? &DestNavMesh->GetConfig() : nullptr;
 	// ~
 		
 	// find all tiles that need regeneration
