Fix for unnecessarily rebuilding of nav mesh based on the number of bits requried to store maxTiles.

diff --git a/Engine/Source/Runtime/NavigationSystem/Private/NavMesh/RecastNavMeshGenerator.cpp b/Engine/Source/Runtime/NavigationSystem/Private/NavMesh/RecastNavMeshGenerator.cpp
--- a/Engine/Source/Runtime/NavigationSystem/Private/NavMesh/RecastNavMeshGenerator.cpp
+++ b/Engine/Source/Runtime/NavigationSystem/Private/NavMesh/RecastNavMeshGenerator.cpp
@@ -4778,7 +4778,7 @@ void FRecastNavMeshGenerator::Init()
 				if (!bRecreateNavmesh)
 				{
 					CalcNavMeshProperties(MaxTiles, MaxPolysPerTile);
-					if (FMath::Log2(static_cast<float>(MaxTiles)) != FMath::Log2(static_cast<float>(SavedNavParams->maxTiles)))
+					if (FMath::CeilToInt(FMath::Log2(static_cast<float>(MaxTiles))) != FMath::CeilToInt(FMath::Log2(static_cast<float>(SavedNavParams->maxTiles))))
 					{
 						bRecreateNavmesh = true;
 						UE_LOG(LogNavigation, Warning, TEXT("Recreating dtNavMesh instance due mismatch in number of bytes required to store serialized maxTiles (%d, %d bits) vs calculated maxtiles (%d, %d bits)")
