Added a log that shows when a dirty area is dirtying one or more tiles:
- usable to understand (per navmesh) how many tiles we are dirtying and what the biggest offenders are

diff --git a/Engine/Source/Runtime/NavigationSystem/Private/NavMesh/RecastNavMeshGenerator.cpp b/Engine/Source/Runtime/NavigationSystem/Private/NavMesh/RecastNavMeshGenerator.cpp
--- a/Engine/Source/Runtime/NavigationSystem/Private/NavMesh/RecastNavMeshGenerator.cpp
+++ b/Engine/Source/Runtime/NavigationSystem/Private/NavMesh/RecastNavMeshGenerator.cpp
@@ -10,6 +10,7 @@
 #include "EngineGlobals.h"
 #include "GameFramework/PlayerController.h"
 #include "Engine/Engine.h"
+#include "Logging/LogScopedCategoryAndVerbosityOverride.h"
 #include "NavigationSystem.h"
 #include "FramePro/FrameProProfiler.h"
 #include "UObject/GarbageCollection.h"
@@ -5768,6 +5769,11 @@ void FRecastNavMeshGenerator::MarkDirtyTiles(const TArray<FNavigationDirtyArea>&
 			bDoTileInclusionTest = (FindInclusionBoundEncapsulatingBox(CutDownArea) == INDEX_NONE);
 		}
 		
+#if !UE_BUILD_SHIPPING
+		// Used for debug purposes to track the number of tiles relevant to this area
+		uint32 RelevantTiles = 0;
+#endif
+		
 		const FRcTileBox TileBox(AdjustedAreaBounds, RcNavMeshOrigin, TileSizeInWorldUnits);
 
 		for (int32 TileY = TileBox.YMin; TileY <= TileBox.YMax; ++TileY)
@@ -5817,8 +5823,31 @@ void FRecastNavMeshGenerator::MarkDirtyTiles(const TArray<FNavigationDirtyArea>&
 				{
 					DirtyTiles.Add(Element);
 				}
+				
+#if !UE_BUILD_SHIPPING
+				RelevantTiles++;
+#endif
 			}
 		}
+		
+#if !UE_BUILD_SHIPPING
+		UE_SUPPRESS(LogNavigationDirtyArea, VeryVerbose, 
+		{
+			if (RelevantTiles > 0)
+			{
+				if (const UObject* const SourceObject = DirtyArea.OptionalSourceObject.Get())
+				{
+					const UActorComponent* const ObjectAsComponent = Cast<UActorComponent>(SourceObject);
+					const AActor* const ComponentOwner = ObjectAsComponent ? ObjectAsComponent->GetOwner() : nullptr;
+					
+					UE_LOG(LogNavigationDirtyArea, VeryVerbose,
+						TEXT("(navmesh: %s) Dirty area trying to dirt %d tiles | Source Object = %s | Potential component's owner = %s | Bounds size = %s)"),
+						*GetNameSafe(GetOwner()), RelevantTiles, *GetFullNameSafe(SourceObject),
+						*GetFullNameSafe(ComponentOwner), *DirtyArea.Bounds.ToString());
+				}
+			}
+		});
+#endif
 	}
 	
 	int32 NumTilesMarked = DirtyTiles.Num();
diff --git a/Engine/Source/Runtime/NavigationSystem/Private/NavigationDirtyAreasController.cpp b/Engine/Source/Runtime/NavigationSystem/Private/NavigationDirtyAreasController.cpp
--- a/Engine/Source/Runtime/NavigationSystem/Private/NavigationDirtyAreasController.cpp
+++ b/Engine/Source/Runtime/NavigationSystem/Private/NavigationDirtyAreasController.cpp
@@ -3,7 +3,7 @@
 #include "NavigationDirtyAreasController.h"
 #include "NavigationData.h"
 
-DEFINE_LOG_CATEGORY_STATIC(LogNavigationDirtyArea, Warning, All);
+DEFINE_LOG_CATEGORY(LogNavigationDirtyArea);
 
 //----------------------------------------------------------------------//
 // FNavigationDirtyAreasController
diff --git a/Engine/Source/Runtime/NavigationSystem/Public/NavigationDirtyAreasController.h b/Engine/Source/Runtime/NavigationSystem/Public/NavigationDirtyAreasController.h
--- a/Engine/Source/Runtime/NavigationSystem/Public/NavigationDirtyAreasController.h
+++ b/Engine/Source/Runtime/NavigationSystem/Public/NavigationDirtyAreasController.h
@@ -4,6 +4,8 @@
 #include "AI/Navigation/NavigationTypes.h"
 
 
+NAVIGATIONSYSTEM_API DECLARE_LOG_CATEGORY_EXTERN(LogNavigationDirtyArea, Warning, All);
+
 class ANavigationData;
 
 struct NAVIGATIONSYSTEM_API FNavigationDirtyAreasController
