[Navmesh] Validate bounds passed to FNavigationOctree::AddNode

diff --git a/Engine/Source/Runtime/NavigationSystem/Private/NavigationOctree.cpp b/Engine/Source/Runtime/NavigationSystem/Private/NavigationOctree.cpp
--- a/Engine/Source/Runtime/NavigationSystem/Private/NavigationOctree.cpp
+++ b/Engine/Source/Runtime/NavigationSystem/Private/NavigationOctree.cpp
@@ -132,6 +132,12 @@ bool FNavigationOctree::IsLazyGathering(const INavRelevantInterface& ChildNavInt
 
 void FNavigationOctree::AddNode(UObject* ElementOb, INavRelevantInterface* NavElement, const FBox& Bounds, FNavigationOctreeElement& Element)
 {
+	if (!Bounds.IsValid)
+	{
+		UE_LOG(LogNavigation, Warning, TEXT("%s: Ignoring %s because bounds are invalid."), ANSI_TO_TCHAR(__FUNCTION__), *GetNameSafe(ElementOb));
+		return;
+	}
+	
 	// we assume NavElement is ElementOb already cast
 	Element.Bounds = Bounds;	
 
