Added safety check to nav mesh loading in case code with older nav mesh version attempts to load newer versioned nav mesh.

diff --git a/Engine/Source/Runtime/NavigationSystem/Private/NavMesh/RecastNavMesh.cpp b/Engine/Source/Runtime/NavigationSystem/Private/NavMesh/RecastNavMesh.cpp
--- a/Engine/Source/Runtime/NavigationSystem/Private/NavMesh/RecastNavMesh.cpp
+++ b/Engine/Source/Runtime/NavigationSystem/Private/NavMesh/RecastNavMesh.cpp
@@ -790,17 +790,27 @@ void ARecastNavMesh::Serialize( FArchive& Ar )
 	{
 		Ar << RecastNavMeshSizeBytes;
 	}
+
 	if (Ar.IsLoading())
 	{
-		if (NavMeshVersion < NAVMESHVER_MIN_COMPATIBLE)
+		auto CleanUpBadVersion = [&Ar, RecastNavMeshSizePos, RecastNavMeshSizeBytes, this]()
 		{
-			UE_LOG(LogNavigation, Warning, TEXT("%s: ARecastNavMesh: Nav mesh version %d < Min compatible %d. Nav mesh needs to be rebuilt. \n"), *GetFullName(), NavMeshVersion, NAVMESHVER_MIN_COMPATIBLE);
-
 			// incompatible, just skip over this data.  navmesh needs rebuilt.
-			Ar.Seek( RecastNavMeshSizePos + RecastNavMeshSizeBytes );
+			Ar.Seek(RecastNavMeshSizePos + RecastNavMeshSizeBytes);
 
 			// Mark self for delete
 			CleanUpAndMarkPendingKill();
+		};
+
+		if (NavMeshVersion < NAVMESHVER_MIN_COMPATIBLE)
+		{
+			UE_LOG(LogNavigation, Warning, TEXT("%s: ARecastNavMesh: Nav mesh version %d < Min compatible %d. Nav mesh needs to be rebuilt. \n"), *GetFullName(), NavMeshVersion, NAVMESHVER_MIN_COMPATIBLE);
+			CleanUpBadVersion();
+		}
+		else if (NavMeshVersion > NAVMESHVER_LATEST)
+		{
+			UE_LOG(LogNavigation, Warning, TEXT("%s: ARecastNavMesh: Nav mesh version %d > NAVMESHVER_LATEST %d. Newer nav mesh should not be loaded by older code. At a minimum the nav mesh needs to be rebuilt. \n"), *GetFullName(), NavMeshVersion, NAVMESHVER_LATEST);
+			CleanUpBadVersion();
 		}
 		else if (RecastNavMeshSizeBytes > 4)
 		{
diff --git a/Engine/Source/Runtime/NavigationSystem/Private/NavMesh/RecastNavMeshDataChunk.cpp b/Engine/Source/Runtime/NavigationSystem/Private/NavMesh/RecastNavMeshDataChunk.cpp
--- a/Engine/Source/Runtime/NavigationSystem/Private/NavMesh/RecastNavMeshDataChunk.cpp
+++ b/Engine/Source/Runtime/NavigationSystem/Private/NavMesh/RecastNavMeshDataChunk.cpp
@@ -86,12 +86,23 @@ void URecastNavMeshDataChunk::Serialize(FArchive& Ar)
 
 	if (Ar.IsLoading())
 	{
+		auto CleanUpBadVersion = [&Ar, RecastNavMeshSizePos, RecastNavMeshSizeBytes]()
+		{
+			// incompatible, just skip over this data. Navmesh needs rebuilt.
+			Ar.Seek(RecastNavMeshSizePos + RecastNavMeshSizeBytes);
+		};
+
 		if (NavMeshVersion < NAVMESHVER_MIN_COMPATIBLE)
 		{
 			UE_LOG(LogNavigation, Warning, TEXT("%s: URecastNavMeshDataChunk: Nav mesh version %d < Min compatible %d. Nav mesh needs to be rebuilt. \n"), *GetFullName(), NavMeshVersion, NAVMESHVER_MIN_COMPATIBLE);
 
-			// incompatible, just skip over this data. Navmesh needs rebuilt.
-			Ar.Seek(RecastNavMeshSizePos + RecastNavMeshSizeBytes);
+			CleanUpBadVersion();
+		}
+		else if (NavMeshVersion > NAVMESHVER_LATEST)
+		{
+			UE_LOG(LogNavigation, Warning, TEXT("%s: URecastNavMeshDataChunk: Nav mesh version %d > NAVMESHVER_LATEST %d. Newer nav mesh should not be loaded by older versioned code. At a minimum the nav mesh needs to be rebuilt. \n"), *GetFullName(), NavMeshVersion, NAVMESHVER_LATEST);
+
+			CleanUpBadVersion();
 		}
 #if WITH_RECAST
 		else if (RecastNavMeshSizeBytes > 4)
