PR #8709: Make ApplyVoxelFilter virtual (Contributed by [REDACTED_PERSON])

diff --git a/Engine/Source/Runtime/NavigationSystem/Public/NavMesh/RecastNavMeshGenerator.h b/Engine/Source/Runtime/NavigationSystem/Public/NavMesh/RecastNavMeshGenerator.h
--- a/Engine/Source/Runtime/NavigationSystem/Public/NavMesh/RecastNavMeshGenerator.h
+++ b/Engine/Source/Runtime/NavigationSystem/Public/NavMesh/RecastNavMeshGenerator.h
@@ -394,7 +394,7 @@ protected:
 	/** builds NavigationData array (layers + obstacles) */
 	bool GenerateNavigationData(FNavMeshBuildContext& BuildContext);
 
-	void ApplyVoxelFilter(struct rcHeightfield* SolidHF, FVector::FReal WalkableRadius);
+	virtual void ApplyVoxelFilter(struct rcHeightfield* SolidHF, FVector::FReal WalkableRadius);
 
 	/** Compute rasterization mask */
 	void InitRasterizationMaskArray(const rcHeightfield* SolidHF, TInlineMaskArray& OutRasterizationMasks);
