[Navmesh] Fix CalculateMaxTilesCount() changing per platform and target

diff --git a/Engine/Source/Runtime/NavigationSystem/Private/NavMesh/RecastNavMeshGenerator.cpp b/Engine/Source/Runtime/NavigationSystem/Private/NavMesh/RecastNavMeshGenerator.cpp
--- a/Engine/Source/Runtime/NavigationSystem/Private/NavMesh/RecastNavMeshGenerator.cpp
+++ b/Engine/Source/Runtime/NavigationSystem/Private/NavMesh/RecastNavMeshGenerator.cpp
@@ -13,6 +13,7 @@
 #include "Logging/LogScopedCategoryAndVerbosityOverride.h"
 #include "NavigationSystem.h"
 #include "FramePro/FrameProProfiler.h"
+#include "NavMesh/RecastVersion.h"
 #include "UObject/GarbageCollection.h"
 
 #if WITH_RECAST
@@ -4588,16 +4589,29 @@ FString FRecastTileGenerator::GetReferencerName() const
 	return TEXT("FRecastTileGenerator");
 }
 
-static int32 CaclulateMaxTilesCount(const TNavStatArray<FBox>& NavigableAreas, FVector::FReal TileSizeinWorldUnits, FVector::FReal AvgLayersPerGridCell)
+static int32 CalculateMaxTilesCount(const TNavStatArray<FBox>& NavigableAreas, FVector::FReal TileSizeInWorldUnits, FVector::FReal AvgLayersPerGridCell, const uint32 NavMeshVersion)
 {
 	int32 GridCellsCount = 0;
-	for (FBox AreaBounds : NavigableAreas)
+	for (const FBox& AreaBounds : NavigableAreas)
 	{
 		// TODO: need more precise calculation, currently we don't take into account that volumes can be overlapped
-		FBox RCBox = Unreal2RecastBox(AreaBounds);
-		int32 XSize = FMath::CeilToInt(RCBox.GetSize().X/TileSizeinWorldUnits) + 1;
-		int32 YSize = FMath::CeilToInt(RCBox.GetSize().Z/TileSizeinWorldUnits) + 1;
-		GridCellsCount+= (XSize*YSize);
+		const FBox RCBox = Unreal2RecastBox(AreaBounds);
+
+		if (NavMeshVersion >= NAVMESHVER_MAXTILES_COUNT_CHANGE)
+		{
+			// Keep this as an integer division to avoid imprecision between platforms and targets (since MaxTilesCount is compared with stored data).
+			const int32 TileSizeUU = (int32)TileSizeInWorldUnits;
+			const int32 XSize = (FMath::CeilToInt(RCBox.GetSize().X) / TileSizeUU) + 1;
+			const int32 YSize = (FMath::CeilToInt(RCBox.GetSize().Z) / TileSizeUU) + 1;
+			GridCellsCount += (XSize*YSize);
+		}
+		else
+		{
+			// Support old navmesh versions
+			int32 XSize = FMath::CeilToInt(RCBox.GetSize().X/TileSizeInWorldUnits) + 1;
+			int32 YSize = FMath::CeilToInt(RCBox.GetSize().Z/TileSizeInWorldUnits) + 1;
+			GridCellsCount+= (XSize*YSize);
+		}
 	}
 	
 	return FMath::CeilToInt(GridCellsCount * AvgLayersPerGridCell);
@@ -4933,7 +4947,7 @@ void FRecastNavMeshGenerator::CalcNavMeshProperties(int32& MaxTiles, int32& MaxP
 	int32 MaxRequestedTiles = 0;
 	if (DestNavMesh->IsResizable())
 	{
-		MaxRequestedTiles = CaclulateMaxTilesCount(InclusionBounds, Config.tileSize * Config.cs, AvgLayersPerTile);
+		MaxRequestedTiles = CalculateMaxTilesCount(InclusionBounds, Config.tileSize * Config.cs, AvgLayersPerTile, DestNavMesh->NavMeshVersion);
 	}
 	else
 	{
@@ -5084,7 +5098,7 @@ void FRecastNavMeshGenerator::OnNavigationBoundsChanged()
 	if (!IsGameStaticNavMesh(DestNavMesh) && DestNavMesh->IsResizable() && DetourMesh)
 	{
 		// Check whether Navmesh size needs to be changed
-		int32 MaxRequestedTiles = CaclulateMaxTilesCount(InclusionBounds, Config.tileSize * Config.cs, AvgLayersPerTile);
+		int32 MaxRequestedTiles = CalculateMaxTilesCount(InclusionBounds, Config.tileSize * Config.cs, AvgLayersPerTile, DestNavMesh->NavMeshVersion);
 		if (DetourMesh->getMaxTiles() != MaxRequestedTiles)
 		{
 			// Destroy current NavMesh
diff --git a/Engine/Source/Runtime/NavigationSystem/Public/NavMesh/RecastVersion.h b/Engine/Source/Runtime/NavigationSystem/Public/NavMesh/RecastVersion.h
--- a/Engine/Source/Runtime/NavigationSystem/Public/NavMesh/RecastVersion.h
+++ b/Engine/Source/Runtime/NavigationSystem/Public/NavMesh/RecastVersion.h
@@ -17,7 +17,8 @@
 #define NAVMESHVER_LWCOORDS					14
 #define NAVMESHVER_OODLE_COMPRESSION		15
 #define NAVMESHVER_LWCOORDS_SEREALIZATION 	17 // Allows for nav meshes to be serialized agnostic of LWCoords being float or double.
+#define NAVMESHVER_MAXTILES_COUNT_CHANGE 	19
 
 
-#define NAVMESHVER_LATEST				NAVMESHVER_LWCOORDS_SEREALIZATION
+#define NAVMESHVER_LATEST				NAVMESHVER_MAXTILES_COUNT_CHANGE
 #define NAVMESHVER_MIN_COMPATIBLE		NAVMESHVER_LWCOORDS_SEREALIZATION
