Fix navmesh not generating on the default map floor when NavMeshBoundsVolume is placed perfectly flat with the surface (which used to occur by default)

diff --git a/Engine/Source/Editor/UnrealEd/Classes/ActorFactories/ActorFactory.h b/Engine/Source/Editor/UnrealEd/Classes/ActorFactories/ActorFactory.h
--- a/Engine/Source/Editor/UnrealEd/Classes/ActorFactories/ActorFactory.h
+++ b/Engine/Source/Editor/UnrealEd/Classes/ActorFactories/ActorFactory.h
@@ -49,6 +49,9 @@ class UNREALED_API UActorFactory : public UObject, public IAssetFactoryInterface
 	UPROPERTY()
 	uint32 bUseSurfaceOrientation:1;
 
+	UPROPERTY()
+	uint32 bUsePlacementExtent:1;
+
 	/** Translation applied to the spawn position. */
 	UPROPERTY()
 	FVector SpawnPositionOffset;
diff --git a/Engine/Source/Editor/UnrealEd/Private/Editor/ActorPositioning.cpp b/Engine/Source/Editor/UnrealEd/Private/Editor/ActorPositioning.cpp
--- a/Engine/Source/Editor/UnrealEd/Private/Editor/ActorPositioning.cpp
+++ b/Engine/Source/Editor/UnrealEd/Private/Editor/ActorPositioning.cpp
@@ -233,7 +233,8 @@ FTransform FActorPositioning::GetSurfaceAlignedTransform(const FPositioningData&
 	// Choose the largest location offset of the various options (global viewport settings, collision, factory offset)
 	const ULevelEditorViewportSettings* ViewportSettings = GetDefault<ULevelEditorViewportSettings>();
 	const float SnapOffsetExtent = (ViewportSettings->SnapToSurface.bEnabled) ? (ViewportSettings->SnapToSurface.SnapOffsetExtent) : (0.0f);
-	const float CollisionOffsetExtent = FVector::BoxPushOut(Data.SurfaceNormal, Data.PlacementExtent);
+	const FVector PlacementExtent = (Data.ActorFactory && !Data.ActorFactory->bUsePlacementExtent) ? FVector::ZeroVector : Data.PlacementExtent;
+	const float CollisionOffsetExtent = FVector::BoxPushOut(Data.SurfaceNormal, PlacementExtent);
 
 	FVector LocationOffset = Data.SurfaceNormal * FMath::Max(SnapOffsetExtent, CollisionOffsetExtent);
 	if (Data.ActorFactory && LocationOffset.SizeSquared() < Data.ActorFactory->SpawnPositionOffset.SizeSquared())
diff --git a/Engine/Source/Editor/UnrealEd/Private/Factories/ActorFactory.cpp b/Engine/Source/Editor/UnrealEd/Private/Factories/ActorFactory.cpp
--- a/Engine/Source/Editor/UnrealEd/Private/Factories/ActorFactory.cpp
+++ b/Engine/Source/Editor/UnrealEd/Private/Factories/ActorFactory.cpp
@@ -215,6 +215,7 @@ UActorFactory::UActorFactory(const FObjectInitializer& ObjectInitializer)
 {
 	DisplayName = LOCTEXT("DefaultName","Actor");
 	bShowInEditorQuickMenu = false;
+	bUsePlacementExtent = true;
 }
 
 bool UActorFactory::CanCreateActorFrom( const FAssetData& AssetData, FText& OutErrorMsg )
diff --git a/Engine/Source/Runtime/NavigationSystem/Private/NavMesh/NavMeshBoundsVolume.cpp b/Engine/Source/Runtime/NavigationSystem/Private/NavMesh/NavMeshBoundsVolume.cpp
--- a/Engine/Source/Runtime/NavigationSystem/Private/NavMesh/NavMeshBoundsVolume.cpp
+++ b/Engine/Source/Runtime/NavigationSystem/Private/NavMesh/NavMeshBoundsVolume.cpp
@@ -4,6 +4,9 @@
 #include "NavigationSystem.h"
 #include "Engine/CollisionProfile.h"
 #include "Components/BrushComponent.h"
+#if WITH_EDITOR
+#include "ActorFactories/ActorFactory.h"
+#endif // WITH_EDITOR
 
 ANavMeshBoundsVolume::ANavMeshBoundsVolume(const FObjectInitializer& ObjectInitializer)
 	: Super(ObjectInitializer)
@@ -50,6 +53,23 @@ void ANavMeshBoundsVolume::PostEditUndo()
 	}
 }
 
+void ANavMeshBoundsVolume::OnPostEngineInit()
+{
+	if (GEditor)
+	{
+		const TArray<UActorFactory*>& ActorFactories = GEditor->ActorFactories;
+		for (UActorFactory* Factory : ActorFactories)
+		{
+			// For ANavMeshBoundsVolume, do not use placement extent so that the volume embeds into the surface.
+			// When flush with the surface, the collision might not be in the volume and the navmesh might not generate.
+			if (Factory->NewActorClass->IsChildOf(ANavMeshBoundsVolume::StaticClass()))
+			{
+				Factory->bUsePlacementExtent = false;
+			}
+		}
+	}
+}
+
 #endif // WITH_EDITOR
 
 void ANavMeshBoundsVolume::PostRegisterAllComponents() 
diff --git a/Engine/Source/Runtime/NavigationSystem/Private/NavigationSystemModule.cpp b/Engine/Source/Runtime/NavigationSystem/Private/NavigationSystemModule.cpp
--- a/Engine/Source/Runtime/NavigationSystem/Private/NavigationSystemModule.cpp
+++ b/Engine/Source/Runtime/NavigationSystem/Private/NavigationSystemModule.cpp
@@ -6,6 +6,7 @@
 #include "Templates/SubclassOf.h"
 #include "Engine/World.h"
 #include "NavLinkCustomInterface.h"
+#include "NavMesh/NavMeshBoundsVolume.h"
 
 #define LOCTEXT_NAMESPACE "NavigationSystem"
 
@@ -23,6 +24,7 @@ class FNavigationSystemModule : public INavSysModule
 private:
 #if WITH_EDITOR 
 	static FDelegateHandle OnPreWorldInitializationHandle;
+	static FDelegateHandle OnPostEngineInitHandle;
 #endif
 };
 
@@ -30,6 +32,7 @@ IMPLEMENT_MODULE(FNavigationSystemModule, NavigationSystem)
 
 #if WITH_EDITOR 
 FDelegateHandle FNavigationSystemModule::OnPreWorldInitializationHandle;
+FDelegateHandle FNavigationSystemModule::OnPostEngineInitHandle;
 #endif
 
 void FNavigationSystemModule::StartupModule()
@@ -37,6 +40,7 @@ void FNavigationSystemModule::StartupModule()
 	// mz@todo bind to all the delegates in FNavigationSystem
 #if WITH_EDITOR 
 	OnPreWorldInitializationHandle = FWorldDelegates::OnPreWorldInitialization.AddStatic(&INavLinkCustomInterface::OnPreWorldInitialization);
+	OnPostEngineInitHandle = FCoreDelegates::OnPostEngineInit.AddStatic(&ANavMeshBoundsVolume::OnPostEngineInit);
 #endif
 }
 
@@ -45,6 +49,7 @@ void FNavigationSystemModule::ShutdownModule()
 	// This function may be called during shutdown to clean up your module.  For modules that support dynamic reloading,
 	// we call this function before unloading the module.
 #if WITH_EDITOR
+	FCoreDelegates::OnPostEngineInit.Remove(OnPostEngineInitHandle);
 	FWorldDelegates::OnPreWorldInitialization.Remove(OnPreWorldInitializationHandle);
 #endif
 }
diff --git a/Engine/Source/Runtime/NavigationSystem/Public/NavMesh/NavMeshBoundsVolume.h b/Engine/Source/Runtime/NavigationSystem/Public/NavMesh/NavMeshBoundsVolume.h
--- a/Engine/Source/Runtime/NavigationSystem/Public/NavMesh/NavMeshBoundsVolume.h
+++ b/Engine/Source/Runtime/NavigationSystem/Public/NavMesh/NavMeshBoundsVolume.h
@@ -31,6 +31,8 @@ class NAVIGATIONSYSTEM_API ANavMeshBoundsVolume : public AVolume
 	virtual void PostEditChangeProperty( struct FPropertyChangedEvent& PropertyChangedEvent) override;
 	virtual void PostEditUndo() override;
 	//~ End UObject Interface
+
+	static void OnPostEngineInit();
 #endif // WITH_EDITOR
 };
 
