Added error case vlogging to FPImplRecastNavMesh::ProjectPointToNavMesh

diff --git a/Engine/Source/Runtime/NavigationSystem/Private/NavMesh/PImplRecastNavMesh.cpp b/Engine/Source/Runtime/NavigationSystem/Private/NavMesh/PImplRecastNavMesh.cpp
--- a/Engine/Source/Runtime/NavigationSystem/Private/NavMesh/PImplRecastNavMesh.cpp
+++ b/Engine/Source/Runtime/NavigationSystem/Private/NavMesh/PImplRecastNavMesh.cpp
@@ -1876,6 +1876,14 @@ bool FPImplRecastNavMesh::ProjectPointToNavMesh(const FVector& Point, FNavLocati
 				bSuccess = true;
 				Result = FNavLocation(UnrealClosestPoint, PolyRef);
 			}
+			else
+			{
+				const UObject* LogOwner = Owner ? Owner : NavMeshOwner;
+				UE_VLOG(LogOwner, LogNavigation, Error, TEXT("ProjectPointToNavMesh failed due to ClosestPoint being too far away from projected point."));
+				UE_VLOG_LOCATION(LogOwner, LogNavigation, Error, Point, 30.f, FColor::Blue, TEXT("Requested point"));
+				UE_VLOG_LOCATION(LogOwner, LogNavigation, Error, UnrealClosestPoint, 30.f, FColor::Red, TEXT("Projection result"));
+				UE_VLOG_SEGMENT(LogOwner, LogNavigation, Error, Point, UnrealClosestPoint, FColor::Red, TEXT(""));
+			}
 		}
 	}
 
