Fix for navmesh queries (related to CL 18595413)
In the ANavigationData::RebuildAll() flow, NavDataGenerator->RebuildAll() is called and it releases the detour navmesh to make a new one. In that flow, the new dtNavMesh params where never sets, causing m_bvQuantFactor to be 0 and navmesh queries to fail.
(not that RebuildAll() flow can be called automatically when bNavigationAutoUpdateEnabled is enabled in editor)

diff --git a/Engine/Source/Runtime/NavigationSystem/Private/NavMesh/RecastNavMeshGenerator.cpp b/Engine/Source/Runtime/NavigationSystem/Private/NavMesh/RecastNavMeshGenerator.cpp
--- a/Engine/Source/Runtime/NavigationSystem/Private/NavMesh/RecastNavMeshGenerator.cpp
+++ b/Engine/Source/Runtime/NavigationSystem/Private/NavMesh/RecastNavMeshGenerator.cpp
@@ -4896,6 +4896,11 @@ bool FRecastNavMeshGenerator::ConstructTiledNavMesh()
 		CalcNavMeshProperties(TiledMeshParameters.maxTiles, TiledMeshParameters.maxPolys);
 		Config.MaxPolysPerTile = TiledMeshParameters.maxPolys;
 
+		DetourMesh->setWalkableClimb(Config.AgentMaxClimb);
+		DetourMesh->setWalkableHeight(Config.AgentHeight);
+		DetourMesh->setWalkableRadius(Config.AgentRadius);
+		DetourMesh->setBVQuantFactor(1.f / Config.cs);
+
 		if (TiledMeshParameters.maxTiles == 0)
 		{
 			UE_LOG(LogNavigation, Warning, TEXT("ConstructTiledNavMesh: Failed to create navmesh of size 0."));
