Added LLM tags for the navigation octree.
Also removed two unnecessary tags in the nav mesh code.
The new LLM tags mirror the existing STAT_Navigation_CollisionTreeMemory functionality.

diff --git a/Engine/Source/Runtime/NavigationSystem/Private/NavMesh/PImplRecastNavMesh.cpp b/Engine/Source/Runtime/NavigationSystem/Private/NavMesh/PImplRecastNavMesh.cpp
--- a/Engine/Source/Runtime/NavigationSystem/Private/NavMesh/PImplRecastNavMesh.cpp
+++ b/Engine/Source/Runtime/NavigationSystem/Private/NavMesh/PImplRecastNavMesh.cpp
@@ -206,14 +206,11 @@ static void DetourFree(void* Original, dtAllocHint Hint)
 	DEC_MEMORY_STAT_BY(STAT_Navigation_RecastMemory, Size);
 #endif // STATS
 
-	LLM_SCOPE(ELLMTag::NavigationRecast);
-
 	FMemory::Free(Original);
 }
 
 static void RecastFree(void* Original)
 {
-	LLM_SCOPE(ELLMTag::NavigationRecast);
 #if STATS
 	const uint32 Size = FMemory::GetAllocSize(Original);
 	DEC_DWORD_STAT_BY(STAT_NavigationMemory, Size);
diff --git a/Engine/Source/Runtime/NavigationSystem/Private/NavMesh/RecastNavMeshGenerator.cpp b/Engine/Source/Runtime/NavigationSystem/Private/NavMesh/RecastNavMeshGenerator.cpp
--- a/Engine/Source/Runtime/NavigationSystem/Private/NavMesh/RecastNavMeshGenerator.cpp
+++ b/Engine/Source/Runtime/NavigationSystem/Private/NavMesh/RecastNavMeshGenerator.cpp
@@ -2381,19 +2381,24 @@ void FRecastTileGenerator::GatherNavigationDataGeometry(const TSharedRef<FNaviga
 
 			if (!HasVoxelCache(ElementData->VoxelData, CachedVoxels, NumCachedVoxels))
 			{
-						// rasterize
+
+				// rasterize
 				PrepareVoxelCache(ElementData->CollisionData, ModifierInstance, SpanData);
 				CachedVoxels = SpanData.GetData();
 				NumCachedVoxels = SpanData.Num();
 
-						// encode
-				const int32 PrevElementMemory = ElementData->GetAllocatedSize();
-				FNavigationRelevantData* ModData = (FNavigationRelevantData*)&ElementData;
-				AddVoxelCache(ModData->VoxelData, CachedVoxels, NumCachedVoxels);
+				// encode
+				{
+					LLM_SCOPE_BYTAG(NavigationOctree);
+
+					const int32 PrevElementMemory = ElementData->GetAllocatedSize();
+					FNavigationRelevantData* ModData = (FNavigationRelevantData*)&ElementData;
+					AddVoxelCache(ModData->VoxelData, CachedVoxels, NumCachedVoxels);
 
-				const int32 NewElementMemory = ElementData->GetAllocatedSize();
-				const int32 ElementMemoryDelta = NewElementMemory - PrevElementMemory;
-				INC_MEMORY_STAT_BY(STAT_Navigation_CollisionTreeMemory, ElementMemoryDelta);
+					const int32 NewElementMemory = ElementData->GetAllocatedSize();
+					const int32 ElementMemoryDelta = NewElementMemory - PrevElementMemory;
+					INC_MEMORY_STAT_BY(STAT_Navigation_CollisionTreeMemory, ElementMemoryDelta);
+				}
 			}
 		}
 		else
diff --git a/Engine/Source/Runtime/NavigationSystem/Private/NavigationOctree.cpp b/Engine/Source/Runtime/NavigationSystem/Private/NavigationOctree.cpp
--- a/Engine/Source/Runtime/NavigationSystem/Private/NavigationOctree.cpp
+++ b/Engine/Source/Runtime/NavigationSystem/Private/NavigationOctree.cpp
@@ -4,6 +4,8 @@
 #include "AI/Navigation/NavRelevantInterface.h"
 #include "NavigationSystem.h"
 
+LLM_DEFINE_TAG(NavigationOctree);
+
 
 //----------------------------------------------------------------------//
 // FNavigationOctree
@@ -41,6 +43,8 @@ void FNavigationOctree::SetNavigableGeometryStoringMode(ENavGeometryStoringMode
 
 void FNavigationOctree::DemandLazyDataGathering(FNavigationRelevantData& ElementData)
 {
+	LLM_SCOPE_BYTAG(NavigationOctree);
+
 	UObject* ElementOb = ElementData.GetOwner();
 	if (ElementOb == nullptr)
 	{
@@ -107,6 +111,8 @@ void FNavigationOctree::DemandLazyDataGathering(FNavigationRelevantData& Element
 
 void FNavigationOctree::DemandChildLazyDataGathering(FNavigationRelevantData& ElementData, INavRelevantInterface& ChildNavInterface)
 {
+	LLM_SCOPE_BYTAG(NavigationOctree);
+
 	if (IsLazyGathering(ChildNavInterface))
 	{
 		ChildNavInterface.GetNavigationData(ElementData);
@@ -132,6 +138,8 @@ bool FNavigationOctree::IsLazyGathering(const INavRelevantInterface& ChildNavInt
 
 void FNavigationOctree::AddNode(UObject* ElementOb, INavRelevantInterface* NavElement, const FBox& Bounds, FNavigationOctreeElement& Element)
 {
+	LLM_SCOPE_BYTAG(NavigationOctree);
+
 	if (!Bounds.IsValid)
 	{
 		UE_LOG(LogNavigation, Warning, TEXT("%s: Ignoring %s because bounds are invalid."), ANSI_TO_TCHAR(__FUNCTION__), *GetNameSafe(ElementOb));
@@ -216,6 +224,8 @@ void FNavigationOctree::AddNode(UObject* ElementOb, INavRelevantInterface* NavEl
 
 void FNavigationOctree::AppendToNode(const FOctreeElementId2& Id, INavRelevantInterface* NavElement, const FBox& Bounds, FNavigationOctreeElement& Element)
 {
+	LLM_SCOPE_BYTAG(NavigationOctree);
+
 	FNavigationOctreeElement OrgData = GetElementById(Id);
 
 	Element = OrgData;
diff --git a/Engine/Source/Runtime/NavigationSystem/Public/NavigationOctree.h b/Engine/Source/Runtime/NavigationSystem/Public/NavigationOctree.h
--- a/Engine/Source/Runtime/NavigationSystem/Public/NavigationOctree.h
+++ b/Engine/Source/Runtime/NavigationSystem/Public/NavigationOctree.h
@@ -11,11 +11,14 @@
 #include "AI/NavigationModifier.h"
 #include "AI/Navigation/NavRelevantInterface.h"
 #include "Math/GenericOctree.h"
+#include "HAL/LowLevelMemTracker.h"
 
 class INavRelevantInterface;
 class FNavigationOctree;
 typedef FNavigationRelevantDataFilter FNavigationOctreeFilter;
 
+LLM_DECLARE_TAG(NavigationOctree);
+
 struct NAVIGATIONSYSTEM_API FNavigationOctreeElement
 {
 	FBoxSphereBounds Bounds;
@@ -132,6 +135,8 @@ public:
  */
 	inline void AddElement(const FNavigationOctreeElement& Element)
 	{
+		LLM_SCOPE_BYTAG(NavigationOctree);
+
 		DEC_MEMORY_STAT_BY(STAT_NavigationMemory, OctreeSizeBytes);
 		DEC_MEMORY_STAT_BY(STAT_Navigation_CollisionTreeMemory, OctreeSizeBytes);
 		TOctree2<FNavigationOctreeElement, FNavigationOctreeSemantics>::AddElement(Element);
