Fix for broken nav mesh links after UE5 integrate. Submitting for Frederic Doll.

diff --git a/Engine/Source/Runtime/NavigationSystem/Private/NavMesh/PImplRecastNavMesh.cpp b/Engine/Source/Runtime/NavigationSystem/Private/NavMesh/PImplRecastNavMesh.cpp
--- a/Engine/Source/Runtime/NavigationSystem/Private/NavMesh/PImplRecastNavMesh.cpp
+++ b/Engine/Source/Runtime/NavigationSystem/Private/NavMesh/PImplRecastNavMesh.cpp
@@ -605,6 +605,13 @@ void FPImplRecastNavMesh::Serialize( FArchive& ArWrapped, int32 NavMeshVersion )
 		{
 			UE_VLOG(NavMeshOwner, LogNavigation, Error, TEXT("Failed to allocate Recast navmesh"));
 		}
+		else if (NavMeshOwner)
+		{
+			DetourNavMesh->setWalkableClimb(NavMeshOwner->AgentMaxStepHeight);
+			DetourNavMesh->setWalkableHeight(NavMeshOwner->AgentHeight);
+			DetourNavMesh->setWalkableRadius(NavMeshOwner->AgentRadius);
+			DetourNavMesh->setBVQuantFactor(1.f / NavMeshOwner->CellSize);
+		}
 	}
 
 	int32 NumTiles = 0;
diff --git a/Engine/Source/Runtime/Navmesh/Private/Detour/DetourNavMesh.cpp b/Engine/Source/Runtime/Navmesh/Private/Detour/DetourNavMesh.cpp
--- a/Engine/Source/Runtime/Navmesh/Private/Detour/DetourNavMesh.cpp
+++ b/Engine/Source/Runtime/Navmesh/Private/Detour/DetourNavMesh.cpp
@@ -1615,6 +1615,7 @@ int dtNavMesh::queryPolygonsInTile(const dtMeshTile* tile, const dtReal* qmin, c
 		dtReal maxz = dtClamp(qmax[2], tbmin[2], tbmax[2]) - tbmin[2];
 
 		// Quantize
+		UE_CLOG(m_bvQuantFactor == 0.f, LogDetour, Warning, TEXT("dtNavMesh::queryPolygonsInTile bounding volume quantization factor is zero! The query might not return the right result"));
 		bmin[0] = (unsigned short)(m_bvQuantFactor * minx) & 0xfffe;
 		bmin[1] = (unsigned short)(m_bvQuantFactor * miny) & 0xfffe;
 		bmin[2] = (unsigned short)(m_bvQuantFactor * minz) & 0xfffe;
