Wrapping RawNavData pointer to help narrow a crash
(with some extra const)

diff --git a/Engine/Source/Runtime/NavigationSystem/Private/NavMesh/RecastNavMesh.cpp b/Engine/Source/Runtime/NavigationSystem/Private/NavMesh/RecastNavMesh.cpp
--- a/Engine/Source/Runtime/NavigationSystem/Private/NavMesh/RecastNavMesh.cpp
+++ b/Engine/Source/Runtime/NavigationSystem/Private/NavMesh/RecastNavMesh.cpp
@@ -135,18 +135,31 @@ FNavMeshTileData::FNavData::~FNavData()
 #endif // WITH_RECAST
 }
 
+// Temporary test to narrow a crash.
+void FNavMeshTileData::FNavData::TestPtr() const
+{
+	if (RawNavData != nullptr)
+	{
+		static uint8 Temp = 0;
+		Temp = *RawNavData;
+
+		const SIZE_T Size = FMemory::GetAllocSize((void*)RawNavData);
+		check(Size == 0 || Size == AllocatedSize);
+	}
+}
+
 FNavMeshTileData::FNavMeshTileData(uint8* RawData, int32 RawDataSize, int32 LayerIdx, FBox LayerBounds)
 	: LayerIndex(LayerIdx)
 	, LayerBBox(LayerBounds)
 	, DataSize(RawDataSize)
 {
 	INC_MEMORY_STAT_BY(STAT_Navigation_TileCacheMemory, DataSize);
-	NavData = MakeShareable(new FNavData(RawData));
+	NavData = MakeShareable(new FNavData(RawData, DataSize));
 }
 
 FNavMeshTileData::~FNavMeshTileData()
 {
-	if (NavData.IsUnique() && NavData->RawNavData)
+	if (NavData.IsUnique() && NavData->GetRawNavData())
 	{
 		// @todo this isn't accounting for the fact that NavData is a shared pointer
 		DEC_MEMORY_STAT_BY(STAT_Navigation_TileCacheMemory, DataSize);
@@ -157,10 +170,10 @@ uint8* FNavMeshTileData::Release()
 {
 	uint8* RawData = nullptr;
 
-	if (NavData.IsValid() && NavData->RawNavData) 
+	if (NavData.IsValid() && NavData->GetRawNavData()) 
 	{ 
-		RawData = NavData->RawNavData;
-		NavData->RawNavData = nullptr;
+		RawData = NavData->GetMutableRawNavData();
+		NavData->Reset();
 		DEC_MEMORY_STAT_BY(STAT_Navigation_TileCacheMemory, DataSize);
 	} 
  
@@ -179,8 +192,11 @@ void FNavMeshTileData::MakeUnique()
 #else
 		uint8* UniqueRawData = (uint8*)FMemory::Malloc(sizeof(uint8)*DataSize);
 #endif //WITH_RECAST
-		FMemory::Memcpy(UniqueRawData, NavData->RawNavData, DataSize);
-		NavData = MakeShareable(new FNavData(UniqueRawData));
+
+		NavData->TestPtr();
+		
+		FMemory::Memcpy(UniqueRawData, NavData->GetRawNavData(), DataSize);
+		NavData = MakeShareable(new FNavData(UniqueRawData, DataSize));
 	}
 }
 
diff --git a/Engine/Source/Runtime/NavigationSystem/Private/NavMesh/RecastNavMeshDataChunk.cpp b/Engine/Source/Runtime/NavigationSystem/Private/NavMesh/RecastNavMeshDataChunk.cpp
--- a/Engine/Source/Runtime/NavigationSystem/Private/NavMesh/RecastNavMeshDataChunk.cpp
+++ b/Engine/Source/Runtime/NavigationSystem/Private/NavMesh/RecastNavMeshDataChunk.cpp
@@ -53,7 +53,7 @@ FRecastTileData::FRecastTileData(int32 DataSize, uint8* RawData, int32 CacheData
 }
 
 // Helper to duplicate recast raw data
-static uint8* DuplicateRecastRawData(uint8* Src, int32 SrcSize)
+static uint8* DuplicateRecastRawData(const uint8* Src, int32 SrcSize)
 {
 #if WITH_RECAST	
 	uint8* DupData = (uint8*)dtAlloc(SrcSize, DT_ALLOC_PERM_TILE_DATA);
diff --git a/Engine/Source/Runtime/NavigationSystem/Private/NavMesh/RecastNavMeshGenerator.cpp b/Engine/Source/Runtime/NavigationSystem/Private/NavMesh/RecastNavMeshGenerator.cpp
--- a/Engine/Source/Runtime/NavigationSystem/Private/NavMesh/RecastNavMeshGenerator.cpp
+++ b/Engine/Source/Runtime/NavigationSystem/Private/NavMesh/RecastNavMeshGenerator.cpp
@@ -3709,7 +3709,7 @@ bool FRecastTileGenerator::GenerateNavigationDataLayer(FNavMeshBuildContext& Bui
 	GenerationContext.ResetIntermediateData();
 
 	// Decompress tile layer data. 
-	status = dtDecompressTileCacheLayer(&GenNavAllocator, &TileCompressor, (unsigned char*)CompressedData.GetData(), CompressedData.DataSize, &GenerationContext.Layer);
+	status = dtDecompressTileCacheLayer(&GenNavAllocator, &TileCompressor, (const unsigned char*)CompressedData.GetData(), CompressedData.DataSize, &GenerationContext.Layer);
 	if (dtStatusFailed(status))
 	{
 		BuildContext.log(RC_LOG_ERROR, "GenerateNavigationDataLayer: failed to decompress layer.");
diff --git a/Engine/Source/Runtime/NavigationSystem/Public/NavMesh/RecastNavMesh.h b/Engine/Source/Runtime/NavigationSystem/Public/NavMesh/RecastNavMesh.h
--- a/Engine/Source/Runtime/NavigationSystem/Public/NavMesh/RecastNavMesh.h
+++ b/Engine/Source/Runtime/NavigationSystem/Public/NavMesh/RecastNavMesh.h
@@ -442,10 +442,39 @@ struct FNavMeshTileData
 	// helper function so that we release NavData via dtFree not regular delete (for navigation mem stats)
 	struct FNavData
 	{
-		uint8* RawNavData;
-
-		FNavData(uint8* InNavData) : RawNavData(InNavData) {}
+		// Temporary test to help reproduce a crash.
+		void TestPtr() const;
+		
+		FNavData(uint8* InNavData, const int32 InDataSize) : RawNavData(InNavData)
+		{
+			if (RawNavData != nullptr)
+			{
+				// Temporary test to help reproduce a crash.
+				static uint8 Temp = 0;
+				Temp = *RawNavData;
+				
+				AllocatedSize = FMemory::GetAllocSize((void*)RawNavData);
+				check(AllocatedSize == 0 || AllocatedSize >= InDataSize);
+			}
+			else
+			{
+				AllocatedSize = 0;
+			}
+		}
 		~FNavData();
+
+		const uint8* GetRawNavData() const { return RawNavData; }
+		uint8* GetMutableRawNavData() { return RawNavData; }
+
+		void Reset()
+		{
+			RawNavData = nullptr;
+			AllocatedSize = 0;
+		}
+				
+	protected:
+		uint8* RawNavData;
+		SIZE_T AllocatedSize; // != DataSize
 	};
 	
 	// layer index
@@ -464,18 +493,18 @@ struct FNavMeshTileData
 	FORCEINLINE uint8* GetData()
 	{
 		check(NavData.IsValid());
-		return NavData->RawNavData;
+		return NavData->GetMutableRawNavData();
 	}
 
 	FORCEINLINE const uint8* GetData() const
 	{
 		check(NavData.IsValid());
-		return NavData->RawNavData;
+		return NavData->GetRawNavData();
 	}
 
 	FORCEINLINE uint8* GetDataSafe()
 	{
-		return NavData.IsValid() ? NavData->RawNavData : NULL;
+		return NavData.IsValid() ? NavData->GetMutableRawNavData() : NULL;
 	}
 
 	FORCEINLINE bool operator==(const uint8* RawData) const
diff --git a/Engine/Source/Runtime/Navmesh/Private/Detour/DetourNavMesh.cpp b/Engine/Source/Runtime/Navmesh/Private/Detour/DetourNavMesh.cpp
--- a/Engine/Source/Runtime/Navmesh/Private/Detour/DetourNavMesh.cpp
+++ b/Engine/Source/Runtime/Navmesh/Private/Detour/DetourNavMesh.cpp
@@ -1785,7 +1785,7 @@ dtStatus dtNavMesh::addTile(unsigned char* data, int dataSize, int flags,
 	const int clusterPolysSize = dtAlign(sizeof(unsigned short)*header->offMeshBase);
 #endif // WITH_NAVMESH_CLUSTER_LINKS
 
-	unsigned char* d = data + headerSize;
+	const unsigned char* d = data + headerSize;
 	tile->verts = (dtReal*)d; d += vertsSize;
 	tile->polys = (dtPoly*)d; d += polysSize;
 	tile->links = (dtLink*)d; d += linksSize;
diff --git a/Engine/Source/Runtime/Navmesh/Private/DetourTileCache/DetourTileCacheBuilder.cpp b/Engine/Source/Runtime/Navmesh/Private/DetourTileCache/DetourTileCacheBuilder.cpp
--- a/Engine/Source/Runtime/Navmesh/Private/DetourTileCache/DetourTileCacheBuilder.cpp
+++ b/Engine/Source/Runtime/Navmesh/Private/DetourTileCache/DetourTileCacheBuilder.cpp
@@ -2659,12 +2659,12 @@ dtStatus dtBuildTileCacheLayer(dtTileCacheCompressor* comp,
 void dtFreeTileCacheLayer(dtTileCacheAlloc* alloc, dtTileCacheLayer* layer)
 {
 	dtAssert(alloc);
-	// The layer is allocated as one conitguous blob of data.
+	// The layer is allocated as one contiguous blob of data.
 	alloc->free(layer);
 }
 
 dtStatus dtDecompressTileCacheLayer(dtTileCacheAlloc* alloc, dtTileCacheCompressor* comp,
-									unsigned char* compressed, const int compressedSize,
+									const unsigned char* compressed, const int compressedSize,
 									dtTileCacheLayer** layerOut)
 {
 	dtAssert(alloc);
diff --git a/Engine/Source/Runtime/Navmesh/Public/DetourTileCache/DetourTileCacheBuilder.h b/Engine/Source/Runtime/Navmesh/Public/DetourTileCache/DetourTileCacheBuilder.h
--- a/Engine/Source/Runtime/Navmesh/Public/DetourTileCache/DetourTileCacheBuilder.h
+++ b/Engine/Source/Runtime/Navmesh/Public/DetourTileCache/DetourTileCacheBuilder.h
@@ -196,7 +196,7 @@ NAVMESH_API dtStatus dtBuildTileCacheLayer(dtTileCacheCompressor* comp,
 NAVMESH_API void dtFreeTileCacheLayer(dtTileCacheAlloc* alloc, dtTileCacheLayer* layer);
 
 NAVMESH_API dtStatus dtDecompressTileCacheLayer(dtTileCacheAlloc* alloc, dtTileCacheCompressor* comp,
-									unsigned char* compressed, const int compressedSize,
+									const unsigned char* compressed, const int compressedSize,
 									dtTileCacheLayer** layerOut);
 
 NAVMESH_API dtTileCacheContourSet* dtAllocTileCacheContourSet(dtTileCacheAlloc* alloc);
