[Navmesh] Change bAutoDestroyWhenNoNavigation to be false by default

diff --git a/Engine/Source/Runtime/NavigationSystem/Private/NavigationData.cpp b/Engine/Source/Runtime/NavigationSystem/Private/NavigationData.cpp
--- a/Engine/Source/Runtime/NavigationSystem/Private/NavigationData.cpp
+++ b/Engine/Source/Runtime/NavigationSystem/Private/NavigationData.cpp
@@ -141,7 +141,7 @@ ANavigationData::ANavigationData(const FObjectInitializer& ObjectInitializer)
 	: Super(ObjectInitializer)
 	, bEnableDrawing(false)
 	, bForceRebuildOnLoad(false)
-	, bAutoDestroyWhenNoNavigation(true)
+	, bAutoDestroyWhenNoNavigation(false)
 	, bCanBeMainNavData(true)
 	, bCanSpawnOnRebuild(true)
 	, RuntimeGeneration(ERuntimeGenerationType::LegacyGeneration) //TODO: set to a valid value once bRebuildAtRuntime_DEPRECATED is removed
