[Navmesh] Fix broken "Mask Fill Collision Under" feature (regression from CL 17897123)
#rnx

diff --git a/Engine/Source/Runtime/NavigationSystem/Private/NavMesh/RecastNavMeshGenerator.cpp b/Engine/Source/Runtime/NavigationSystem/Private/NavMesh/RecastNavMeshGenerator.cpp
--- a/Engine/Source/Runtime/NavigationSystem/Private/NavMesh/RecastNavMeshGenerator.cpp
+++ b/Engine/Source/Runtime/NavigationSystem/Private/NavMesh/RecastNavMeshGenerator.cpp
@@ -4352,6 +4352,7 @@ void FRecastTileGenerator::MarkRasterizationMask(rcContext* /*BuildContext*/, rc
 		FBox RecastBox = Unreal2RecastBox(WorldBox);
 		FVector RecastPos;
 		FVector RecastExtent;
+		RecastBox.GetCenterAndExtents(RecastPos, RecastExtent);
 		check(OutMaskArray.Num() == SolidHF->width*SolidHF->height);
 		MarkBoxMask(&(RecastPos.X), &(RecastExtent.X), Mask, *SolidHF, OutMaskArray.GetData());
 	}
