Set IsSpatiallyLoaded to false by default for NavMeshBoundsVolume
#rnx

diff --git a/Engine/Source/Runtime/NavigationSystem/Private/NavMesh/NavMeshBoundsVolume.cpp b/Engine/Source/Runtime/NavigationSystem/Private/NavMesh/NavMeshBoundsVolume.cpp
--- a/Engine/Source/Runtime/NavigationSystem/Private/NavMesh/NavMeshBoundsVolume.cpp
+++ b/Engine/Source/Runtime/NavigationSystem/Private/NavMesh/NavMeshBoundsVolume.cpp
@@ -18,6 +18,10 @@ ANavMeshBoundsVolume::ANavMeshBoundsVolume(const FObjectInitializer& ObjectIniti
 	SupportedAgents.MarkInitialized();
 
 	bColored = true;
+
+#if WITH_EDITORONLY_DATA
+	bIsSpatiallyLoaded = false;
+#endif
 }
 
 #if WITH_EDITOR
