Fix "Navmesh needs to be rebuilt" error message while running a Listen Server and Client by not creating default static navdata in a non-editor context.

diff --git a/Engine/Source/Runtime/NavigationSystem/Private/NavigationSystem.cpp b/Engine/Source/Runtime/NavigationSystem/Private/NavigationSystem.cpp
--- a/Engine/Source/Runtime/NavigationSystem/Private/NavigationSystem.cpp
+++ b/Engine/Source/Runtime/NavigationSystem/Private/NavigationSystem.cpp
@@ -3973,18 +3973,22 @@ void UNavigationSystemV1::RemoveLevelCollisionFromOctree(ULevel* Level)
 	}
 }
 
-void UNavigationSystemV1::OnPostLoadMap(UWorld*)
+void UNavigationSystemV1::OnPostLoadMap(UWorld* LoadedWorld)
 {
-	UE_LOG(LogNavigation, Log, TEXT("UNavigationSystemV1::OnPostLoadMap"));
+	UE_LOG(LogNavigation, Verbose, TEXT("%s (Package: %s)"), ANSI_TO_TCHAR(__FUNCTION__), *GetNameSafe(LoadedWorld->GetOuter()));
 
-	// if map has been loaded and there are some navigation bounds volumes 
+	// If map has been loaded and there are some navigation bounds volumes 
 	// then create appropriate navigation structured
 	ANavigationData* NavData = GetDefaultNavDataInstance(FNavigationSystem::DontCreate);
 
 	// Do this if there's currently no navigation
-	if (NavData == NULL && bAutoCreateNavigationData == true && IsThereAnywhereToBuildNavigation() == true)
+	if ( NavData == nullptr &&
+		 bAutoCreateNavigationData &&
+		 IsThereAnywhereToBuildNavigation() &&
+		 (GetRuntimeGenerationType() != ERuntimeGenerationType::Static) )	// Prevent creating a static default nav instance out of the editor (GetRuntimeGenerationType() is always dynamic in editor).
 	{
 		NavData = GetDefaultNavDataInstance(FNavigationSystem::Create);
+		UE_LOG(LogNavigation, Verbose, TEXT("%s Created DefaultNavDataInstance %s"), ANSI_TO_TCHAR(__FUNCTION__), *GetNameSafe(NavData));
 	}
 }
 
