Crash fix for UActorFactory::NewActorClass that can be null.

diff --git a/Engine/Source/Runtime/NavigationSystem/Private/NavMesh/NavMeshBoundsVolume.cpp b/Engine/Source/Runtime/NavigationSystem/Private/NavMesh/NavMeshBoundsVolume.cpp
--- a/Engine/Source/Runtime/NavigationSystem/Private/NavMesh/NavMeshBoundsVolume.cpp
+++ b/Engine/Source/Runtime/NavigationSystem/Private/NavMesh/NavMeshBoundsVolume.cpp
@@ -66,7 +66,8 @@ void ANavMeshBoundsVolume::OnPostEngineInit()
 		{
 			// For ANavMeshBoundsVolume, do not use placement extent so that the volume embeds into the surface.
 			// When flush with the surface, the collision might not be in the volume and the navmesh might not generate.
-			if (Factory->NewActorClass->IsChildOf(ANavMeshBoundsVolume::StaticClass()))
+			const TSubclassOf<AActor> ActorClass = Factory->NewActorClass;
+			if (ActorClass != nullptr && ActorClass->IsChildOf(ANavMeshBoundsVolume::StaticClass()))
 			{
 				Factory->bUsePlacementExtent = false;
 			}
