[Navmesh] Fix bShouldDiscardSubLevelNavData
Moving conditions allowing the property to take effect from ARecastNavMesh::PostInitProperties() to ARecastNavMesh::PostLoad() and delaying the check to discard navdata if there is no navigation system yet.
Removing previous ShouldDiscardSubLevelNavData() function that would fallback to read value from NavSysCDO when there is no navigation system.
[RN] fix Warning: UNavigationSystemV1 bShouldDiscardSubLevelNavData property was not working in some cases. With this fix navdata from sublevel will be properly discarted if the property is true (it's default value).

diff --git a/Engine/Source/Runtime/Engine/Classes/AI/NavigationSystemBase.h b/Engine/Source/Runtime/Engine/Classes/AI/NavigationSystemBase.h
--- a/Engine/Source/Runtime/Engine/Classes/AI/NavigationSystemBase.h
+++ b/Engine/Source/Runtime/Engine/Classes/AI/NavigationSystemBase.h
@@ -202,7 +202,7 @@ namespace FNavigationSystem
 	DECLARE_DELEGATE_RetVal_OneParam(INavigationDataInterface*, FNavDataForActorSignature, const AActor& /*Actor*/);
 	DECLARE_DELEGATE_RetVal(TSubclassOf<AActor>, FNavDataClassFetchSignature);
 	DECLARE_DELEGATE_TwoParams(FWorldBoolBasedSignature, UWorld& /*World*/, const bool /*bShow*/);
-	DECLARE_MULTICAST_DELEGATE_OneParam(FOnNavigationInitDoneSignature, const UNavigationSystemBase&);
+	DECLARE_MULTICAST_DELEGATE_OneParam(FOnNavigationInitSignature, const UNavigationSystemBase&);
 }
 
 
@@ -254,7 +254,8 @@ public:
 
 	virtual bool IsWorldInitDone() const PURE_VIRTUAL(UNavigationSystemBase::IsWorldInitDone, return false;);
 
-	static FNavigationSystem::FOnNavigationInitDoneSignature& OnNavigationInitDoneStaticDelegate();
+	static FNavigationSystem::FOnNavigationInitSignature& OnNavigationInitStartStaticDelegate();
+	static FNavigationSystem::FOnNavigationInitSignature& OnNavigationInitDoneStaticDelegate();
 
 protected:
 	/**	Sets the Transform the Navigation System will use when converting from FromCoordType
diff --git a/Engine/Source/Runtime/Engine/Private/AI/NavigationSystemBase.cpp b/Engine/Source/Runtime/Engine/Private/AI/NavigationSystemBase.cpp
--- a/Engine/Source/Runtime/Engine/Private/AI/NavigationSystemBase.cpp
+++ b/Engine/Source/Runtime/Engine/Private/AI/NavigationSystemBase.cpp
@@ -28,6 +28,8 @@ namespace FNavigationSystem
 		const auto& Levels = InWorld.GetLevels();
 		for (ULevel* Level : Levels)
 		{
+			UE_LOG(LogNavigation, Verbose, TEXT("%s for %s"), ANSI_TO_TCHAR(__FUNCTION__), *GetFullNameSafe(Level));
+			
 			for (UNavigationDataChunk* NavChunk : Level->NavDataChunks)
 			{
 				if (NavChunk != nullptr)
@@ -68,7 +70,7 @@ namespace FNavigationSystem
 	void AddNavigationSystemToWorld(UWorld& WorldOwner, const FNavigationSystemRunMode RunMode, UNavigationSystemConfig* NavigationSystemConfig, const bool bInitializeForWorld, const bool bOverridePreviousNavSys)
 	{
 		TRACE_CPUPROFILER_EVENT_SCOPE(FNavigationSystem::AddNavigationSystemToWorld);
-		UE_LOG(LogNavigation, VeryVerbose, TEXT("%s (WorldOwner: %s)"), ANSI_TO_TCHAR(__FUNCTION__), *WorldOwner.GetOuter()->GetName());
+		UE_LOG(LogNavigation, Verbose, TEXT("%s bOverridePreviousNavSys=%i (WorldOwner: %s)"), ANSI_TO_TCHAR(__FUNCTION__), bOverridePreviousNavSys, *WorldOwner.GetOuter()->GetName());
 
 		const FNavigationSystemRunMode ResolvedRunMode = (RunMode == FNavigationSystemRunMode::InferFromWorldMode) ? FindRunModeFromWorldType(WorldOwner) : RunMode;
 
@@ -135,7 +137,8 @@ namespace FNavigationSystem
 		FNavDataClassFetchSignature GetDefaultNavDataClass;
 		FWorldBoolBasedSignature VerifyNavigationRenderingComponents;
 		FWorldBasedSignature Build;
-		FOnNavigationInitDoneSignature OnNavigationInitDone;
+		FOnNavigationInitSignature OnNavigationInitStart;
+		FOnNavigationInitSignature OnNavigationInitDone;
 #if WITH_EDITOR
 		FWorldBasedSignature OnPIEStart;
 		FWorldBasedSignature OnPIEEnd;
@@ -397,7 +400,8 @@ FNavigationSystem::FNavDataForActorSignature& UNavigationSystemBase::GetNavDataF
 FNavigationSystem::FNavDataClassFetchSignature& UNavigationSystemBase::GetDefaultNavDataClassDelegate() { return FNavigationSystem::Delegates.GetDefaultNavDataClass; }
 FNavigationSystem::FWorldBoolBasedSignature& UNavigationSystemBase::VerifyNavigationRenderingComponentsDelegate() { return FNavigationSystem::Delegates.VerifyNavigationRenderingComponents; }
 FNavigationSystem::FWorldBasedSignature& UNavigationSystemBase::BuildDelegate() { return FNavigationSystem::Delegates.Build; }
-FNavigationSystem::FOnNavigationInitDoneSignature& UNavigationSystemBase::OnNavigationInitDoneStaticDelegate() { return FNavigationSystem::Delegates.OnNavigationInitDone; }
+FNavigationSystem::FOnNavigationInitSignature& UNavigationSystemBase::OnNavigationInitStartStaticDelegate() { return FNavigationSystem::Delegates.OnNavigationInitStart; }
+FNavigationSystem::FOnNavigationInitSignature& UNavigationSystemBase::OnNavigationInitDoneStaticDelegate() { return FNavigationSystem::Delegates.OnNavigationInitDone; }
 #if WITH_EDITOR
 FNavigationSystem::FWorldBasedSignature& UNavigationSystemBase::OnPIEStartDelegate() { return FNavigationSystem::Delegates.OnPIEStart; }
 FNavigationSystem::FWorldBasedSignature& UNavigationSystemBase::OnPIEEndDelegate() { return FNavigationSystem::Delegates.OnPIEEnd; }
diff --git a/Engine/Source/Runtime/NavigationSystem/Private/NavMesh/RecastNavMesh.cpp b/Engine/Source/Runtime/NavigationSystem/Private/NavMesh/RecastNavMesh.cpp
--- a/Engine/Source/Runtime/NavigationSystem/Private/NavMesh/RecastNavMesh.cpp
+++ b/Engine/Source/Runtime/NavigationSystem/Private/NavMesh/RecastNavMesh.cpp
@@ -501,7 +501,22 @@ void ARecastNavMesh::CleanUp()
 
 void ARecastNavMesh::PostLoad()
 {
+	UE_LOG(LogNavigation, Verbose, TEXT("%s %s"), ANSI_TO_TCHAR(__FUNCTION__), *GetFullNameSafe(this));
+	
 	Super::PostLoad();
+	
+	if (const UWorld* World = GetWorld())
+	{
+		const UNavigationSystemBase* NavSys = World->GetNavigationSystem();
+		if (NavSys && NavSys->IsWorldInitDone())
+		{
+			CheckToDiscardSubLevelNavData(*NavSys);
+		}
+		else
+		{
+			UNavigationSystemBase::OnNavigationInitStartStaticDelegate().AddUObject(this, &ARecastNavMesh::CheckToDiscardSubLevelNavData);
+		}
+	}
 
 	UE_CLOG(TileSizeUU < CellSize, LogNavigation, Error, TEXT("%s: TileSizeUU (%f) being less than CellSize (%f) is an invalid case and will cause navmesh generation issues.")
 		, *GetName(), TileSizeUU, CellSize);
@@ -510,6 +525,36 @@ void ARecastNavMesh::PostLoad()
 	UpdatePolyRefBitsPreview();
 }
 
+void ARecastNavMesh::BeginDestroy()
+{
+	UNavigationSystemBase::OnNavigationInitStartStaticDelegate().RemoveAll(this);
+
+	Super::BeginDestroy();
+}
+
+void ARecastNavMesh::CheckToDiscardSubLevelNavData(const UNavigationSystemBase& BaseNavSys)
+{
+	// This used to be in ARecastNavMesh::PostInitProperties() but the OwningWorld is not always available, it might be too soon to query it.
+	// Moved here so sublevel data can be discarded when requested.
+	const UWorld* OwningWorld = GetWorld();
+	const UNavigationSystemV1* NavSys = Cast<UNavigationSystemV1>(&BaseNavSys);
+	if (OwningWorld && NavSys->ShouldDiscardSubLevelNavData(this))
+	{
+		// Get rid of instances saved within levels that are streamed-in
+		if ((GEngine->IsSettingUpPlayWorld() == false) // this is a @HACK
+			&& (OwningWorld->GetOutermost() != GetOutermost())
+			// If we are cooking, then let them all pass.
+			// They will be handled at load-time when running.
+			&& (IsRunningCommandlet() == false))
+		{
+			UE_LOG(LogNavigation, Verbose, TEXT("%s Discarding %s due to it not being part of PersistentLevel."), ANSI_TO_TCHAR(__FUNCTION__), *GetFullNameSafe(this));
+
+			// Marking self for deletion 
+			CleanUpAndMarkPendingKill();
+		}
+	}
+}
+
 void ARecastNavMesh::PostRegisterAllComponents()
 {
 	Super::PostRegisterAllComponents();
@@ -529,26 +574,6 @@ void ARecastNavMesh::PostInitProperties()
 		static const FNavMeshConfig::FRecastNamedFiltersCreator RecastNamedFiltersCreator(bUseVirtualFilters);
 		NavLinkFlag = FNavMeshConfig::NavLinkFlag;
 	}
-
-	UWorld* MyWorld = GetWorld();
-	if (MyWorld != nullptr 
-		&& HasAnyFlags(RF_NeedLoad) //  was loaded
-		&& FNavigationSystem::ShouldDiscardSubLevelNavData(*this))
-	{
-		// get rid of instances saved within levels that are streamed-in
-		if ((GEngine->IsSettingUpPlayWorld() == false) // this is a @HACK
-			&&	(MyWorld->GetOutermost() != GetOutermost())
-			// If we are cooking, then let them all pass.
-			// They will be handled at load-time when running.
-			&&	(IsRunningCommandlet() == false))
-		{
-			UE_LOG(LogNavigation, Log, TEXT("Discarding %s due to it not being part of PersistentLevel")
-				, *GetNameSafe(this));
-			
-			// marking self for deletion 
-			CleanUpAndMarkPendingKill();
-		}
-	}
 	
 	Super::PostInitProperties();
 
diff --git a/Engine/Source/Runtime/NavigationSystem/Private/NavigationData.cpp b/Engine/Source/Runtime/NavigationSystem/Private/NavigationData.cpp
--- a/Engine/Source/Runtime/NavigationSystem/Private/NavigationData.cpp
+++ b/Engine/Source/Runtime/NavigationSystem/Private/NavigationData.cpp
@@ -676,7 +676,7 @@ void ANavigationData::OnNavAreaAdded(const UClass* NavAreaClass, int32 AgentInde
 		{
 			SupportedAreas[i].AreaClass = NavAreaClass;
 			AreaClassToIdMap.Add(NavAreaClass, SupportedAreas[i].AreaID);
-			UE_VLOG_UELOG(this, LogNavigation, Verbose, TEXT("%s updated area %s with ID %d"), *GetName(), *AreaClassName, SupportedAreas[i].AreaID);
+			UE_VLOG_UELOG(this, LogNavigation, Verbose, TEXT("%s: updated area %s with ID %d"), *GetFullName(), *AreaClassName, SupportedAreas[i].AreaID);
 			return;
 		}
 	}
diff --git a/Engine/Source/Runtime/NavigationSystem/Private/NavigationSystem.cpp b/Engine/Source/Runtime/NavigationSystem/Private/NavigationSystem.cpp
--- a/Engine/Source/Runtime/NavigationSystem/Private/NavigationSystem.cpp
+++ b/Engine/Source/Runtime/NavigationSystem/Private/NavigationSystem.cpp
@@ -192,25 +192,6 @@ namespace FNavigationSystem
 		return false;
 	}
 
-	bool ShouldDiscardSubLevelNavData(ANavigationData& NavData)
-	{
-		const UWorld* World = NavData.GetWorld();
-
-		if (World && World->GetNavigationSystem())
-		{
-			const UNavigationSystemV1* NavSys = Cast<UNavigationSystemV1>(World->GetNavigationSystem());
-			if (NavSys)
-			{
-				return NavSys->ShouldDiscardSubLevelNavData(&NavData);
-			}
-		}
-
-		const UNavigationSystemV1* NavSysCDO = (*GEngine->NavigationSystemClass != nullptr)
-			? (GEngine->NavigationSystemClass->GetDefaultObject<const UNavigationSystemV1>())
-			: (const UNavigationSystemV1*)nullptr;
-		return NavSysCDO == nullptr || NavSysCDO->ShouldDiscardSubLevelNavData(&NavData);
-	}
-
 	void MakeAllComponentsNeverAffectNav(AActor& Actor)
 	{
 		const TSet<UActorComponent*> Components = Actor.GetComponents();
@@ -872,6 +853,8 @@ void UNavigationSystemV1::OnBeginTearingDown(UWorld* World)
 
 void UNavigationSystemV1::OnWorldInitDone(FNavigationSystemRunMode Mode)
 {
+	UNavigationSystemBase::OnNavigationInitStartStaticDelegate().Broadcast(*this);
+	
 	OperationMode = Mode;
 	DoInitialSetup();
 	
@@ -2252,6 +2235,8 @@ void UNavigationSystemV1::ProcessCustomLinkPendingRegistration()
 
 UNavigationSystemV1::ERegistrationResult UNavigationSystemV1::RegisterNavData(ANavigationData* NavData)
 {
+	UE_LOG(LogNavigation, Verbose, TEXT("%s %s"), ANSI_TO_TCHAR(__FUNCTION__), *GetFullNameSafe(NavData));
+	
 	if (NavData == NULL)
 	{
 		return RegistrationError;
@@ -2371,6 +2356,8 @@ UNavigationSystemV1::ERegistrationResult UNavigationSystemV1::RegisterNavData(AN
 
 void UNavigationSystemV1::UnregisterNavData(ANavigationData* NavData)
 {
+	UE_LOG(LogNavigation, Verbose, TEXT("%s %s"), ANSI_TO_TCHAR(__FUNCTION__), *GetFullNameSafe(NavData));
+	
 	NavDataSet.RemoveSingle(NavData);
 
 	if (NavData == NULL)
diff --git a/Engine/Source/Runtime/NavigationSystem/Public/NavMesh/RecastNavMesh.h b/Engine/Source/Runtime/NavigationSystem/Public/NavMesh/RecastNavMesh.h
--- a/Engine/Source/Runtime/NavigationSystem/Public/NavMesh/RecastNavMesh.h
+++ b/Engine/Source/Runtime/NavigationSystem/Public/NavMesh/RecastNavMesh.h
@@ -868,6 +868,7 @@ public:
 	virtual void PostInitProperties() override;
 	virtual void PostLoad() override;
 	virtual void PostRegisterAllComponents() override;
+	virtual void BeginDestroy() override;
 
 #if WITH_EDITOR
 	virtual void PostEditChangeProperty( struct FPropertyChangedEvent& PropertyChangedEvent) override;
@@ -1234,6 +1235,8 @@ protected:
 	 *	own extentions. Note: needs to derive from FRecastNavMeshGenerator */
 	virtual FRecastNavMeshGenerator* CreateGeneratorInstance();
 
+	void CheckToDiscardSubLevelNavData(const UNavigationSystemBase& NavSys);
+
 private:
 	friend struct FRecastGraphWrapper;
 	friend FRecastNavMeshGenerator;
diff --git a/Engine/Source/Runtime/NavigationSystem/Public/NavigationSystem.h b/Engine/Source/Runtime/NavigationSystem/Public/NavigationSystem.h
--- a/Engine/Source/Runtime/NavigationSystem/Public/NavigationSystem.h
+++ b/Engine/Source/Runtime/NavigationSystem/Public/NavigationSystem.h
@@ -73,7 +73,6 @@ namespace FNavigationSystem
 	};
 
 	bool NAVIGATIONSYSTEM_API ShouldLoadNavigationOnClient(ANavigationData& NavData);
-	bool NAVIGATIONSYSTEM_API ShouldDiscardSubLevelNavData(ANavigationData& NavData);
 
 	void NAVIGATIONSYSTEM_API MakeAllComponentsNeverAffectNav(AActor& Actor);
 }
