Fix missing array initialization in FPImplRecastNavMesh::FindPolysAroundCircle resulting in NodePolyCost not being returned

diff --git a/Engine/Source/Runtime/NavigationSystem/Private/NavMesh/PImplRecastNavMesh.cpp b/Engine/Source/Runtime/NavigationSystem/Private/NavMesh/PImplRecastNavMesh.cpp
--- a/Engine/Source/Runtime/NavigationSystem/Private/NavMesh/PImplRecastNavMesh.cpp
+++ b/Engine/Source/Runtime/NavigationSystem/Private/NavMesh/PImplRecastNavMesh.cpp
@@ -2039,8 +2039,8 @@ bool FPImplRecastNavMesh::FindPolysAroundCircle(const FVector& CenterPos, const
 
 		if (OutPolysCost)
 		{
-			OutPolysCost->Reset();
-			OutPolysCost->AddUninitialized(MaxSearchNodes);
+			PolysCost.Reset();
+			PolysCost.AddUninitialized(MaxSearchNodes);
 		}
 
 		FVector::FReal RecastLoc[3];
