Fix "navmesh needs to be rebuilt" message appearing even when there is nowhere to build navigation.

diff --git a/Engine/Source/Runtime/NavigationSystem/Private/NavigationSystem.cpp b/Engine/Source/Runtime/NavigationSystem/Private/NavigationSystem.cpp
--- a/Engine/Source/Runtime/NavigationSystem/Private/NavigationSystem.cpp
+++ b/Engine/Source/Runtime/NavigationSystem/Private/NavigationSystem.cpp
@@ -4411,6 +4411,12 @@ void UNavigationSystemV1::CycleNavigationDataDrawn()
 
 bool UNavigationSystemV1::IsNavigationDirty() const
 {
+	if (!IsThereAnywhereToBuildNavigation())
+	{
+		// Nowhere to build navigation so it can't be dirty.
+		return false;
+	}
+	
 #if !UE_BUILD_SHIPPING
 	if (DefaultDirtyAreasController.HadDirtyAreasReportedWhileAccumulationLocked())
 	{
