RecastNavMesh: Fixed cell size clamping based on issue reported via github PR.

PR #8899: [Bugfix][Recast] Fixed cell size adjustment in RecastNavMesh (Contributed by [REDACTED_PERSON])

diff --git a/Engine/Source/Runtime/NavigationSystem/Private/NavMesh/RecastNavMesh.cpp b/Engine/Source/Runtime/NavigationSystem/Private/NavMesh/RecastNavMesh.cpp
--- a/Engine/Source/Runtime/NavigationSystem/Private/NavMesh/RecastNavMesh.cpp
+++ b/Engine/Source/Runtime/NavigationSystem/Private/NavMesh/RecastNavMesh.cpp
@@ -48,11 +48,11 @@
 namespace 
 {
 	// Max tile size in voxels. Larger than this tiles will start to get slow to build.
-	const int32 ArbitraryMaxTileSizeVoxels = 1024;
+	constexpr int32 ArbitraryMaxTileSizeVoxels = 1024;
 	// Min tile size on voxels. Smaller tiles than this waste computation during voxelization because the border are will be larger than usable area.
-	const int32 ArbitraryMinTileSizeVoxels = 16;
+	constexpr int32 ArbitraryMinTileSizeVoxels = 16;
 	// Minimum tile size in multiples of agent radius.
-	const int32 ArbitraryMinTileSizeAgentRadius = 4; 
+	constexpr int32 ArbitraryMinTileSizeAgentRadius = 4; 
 
 	/** this helper function supplies a consistent way to keep TileSizeUU within defined bounds */
 	float GetClampedTileSizeUU(const float InTileSizeUU, const float CellSize, const float AgentRadius)
@@ -62,6 +62,18 @@ namespace
 		
 		return FMath::Clamp<float>(InTileSizeUU, MinTileSize, MaxTileSize);
 	}
+
+	// These should reflect the property clamping of FRecastNavMeshGenerationProperties::CellSize.  
+	// Minimum cell size.
+	constexpr float ArbitraryMinCellSize = 1.0f; 
+	// Maximum cell size.
+	constexpr float ArbitraryMaxCellSize = 1024.0f; 
+
+	float GetClampedCellSize(const float CellSize)
+	{
+		return FMath::Clamp(CellSize, ArbitraryMinCellSize, ArbitraryMaxCellSize);
+	}
+
 }
 
 FDetourTileLayout::FDetourTileLayout(const dtMeshTile& tile)
@@ -2690,19 +2702,22 @@ void ARecastNavMesh::PostEditChangeProperty(FPropertyChangedEvent& PropertyChang
 			}
 			else if (PropName == GET_MEMBER_NAME_CHECKED(ARecastNavMesh, TileSizeUU))
 			{
+				CellSize = GetClampedCellSize(CellSize);
 				TileSizeUU = GetClampedTileSizeUU(TileSizeUU, CellSize, AgentRadius);
-				
-				// trying to make cell size match TileSizeUU an integer number of times
-				const float AdjustedCellSize = TileSizeUU / FMath::TruncToInt(TileSizeUU / CellSize);
-				CellSize = FMath::Clamp(AdjustedCellSize, TileSizeUU / ArbitraryMinTileSizeVoxels, TileSizeUU / ArbitraryMaxTileSizeVoxels);
+						
+				// Match cell size to tile size.
+				CellSize = TileSizeUU / FMath::TruncToInt(TileSizeUU / CellSize);
 
 				// update config
 				FillConfig(NavDataConfig);
 			}
 			else if (PropName == GET_MEMBER_NAME_CHECKED(ARecastNavMesh, CellSize))
 			{
-				const float AdjustedTileSizeUU = CellSize * FMath::TruncToInt(TileSizeUU / CellSize);
-				TileSizeUU = GetClampedTileSizeUU(AdjustedTileSizeUU, CellSize, AgentRadius);
+				CellSize = GetClampedCellSize(CellSize);
+				TileSizeUU = GetClampedTileSizeUU(TileSizeUU, CellSize, AgentRadius);
+				
+				// Match tile size to cell size.
+				TileSizeUU = CellSize * FMath::TruncToInt(TileSizeUU / CellSize);
 
 				// update config
 				FillConfig(NavDataConfig);
