Add missing Insight trace on mesh card and navcollision builds

#rnx

diff --git a/Engine/Source/Runtime/Engine/Private/MeshCardRepresentation.cpp b/Engine/Source/Runtime/Engine/Private/MeshCardRepresentation.cpp
--- a/Engine/Source/Runtime/Engine/Private/MeshCardRepresentation.cpp
+++ b/Engine/Source/Runtime/Engine/Private/MeshCardRepresentation.cpp
@@ -191,6 +191,8 @@ void BeginCacheMeshCardRepresentation(const ITargetPlatform* TargetPlatform, USt
 
 void FCardRepresentationData::CacheDerivedData(const FString& InDDCKey, const ITargetPlatform* TargetPlatform, UStaticMesh* Mesh, UStaticMesh* GenerateSource, int32 MaxLumenMeshCards, bool bGenerateDistanceFieldAsIfTwoSided, FSourceMeshDataForDerivedDataTask* OptionalSourceMeshData)
 {
+	TRACE_CPUPROFILER_EVENT_SCOPE(FCardRepresentationData::CacheDerivedData);
+
 	TArray<uint8> DerivedData;
 
 	COOK_STAT(auto Timer = CardRepresentationCookStats::UsageStats.TimeSyncWork());
diff --git a/Engine/Source/Runtime/NavigationSystem/Private/NavCollision.cpp b/Engine/Source/Runtime/NavigationSystem/Private/NavCollision.cpp
--- a/Engine/Source/Runtime/NavigationSystem/Private/NavCollision.cpp
+++ b/Engine/Source/Runtime/NavigationSystem/Private/NavCollision.cpp
@@ -141,6 +141,8 @@ FDerivedDataNavCollisionCooker::FDerivedDataNavCollisionCooker(FName InFormat, U
 
 bool FDerivedDataNavCollisionCooker::Build( TArray<uint8>& OutData )
 {
+	TRACE_CPUPROFILER_EVENT_SCOPE(FDerivedDataNavCollisionCooker::Build);
+
 	if ((NavCollisionInstance->ConvexShapeIndices.Num() == 0) ||
 		(NavCollisionInstance->GetTriMeshCollision().VertexBuffer.Num() == 0 && NavCollisionInstance->GetConvexCollision().VertexBuffer.Num() == 0))
 	{
@@ -545,6 +547,8 @@ void UNavCollision::PostLoad()
 
 FByteBulkData* UNavCollision::GetCookedData(FName Format)
 {
+	TRACE_CPUPROFILER_EVENT_SCOPE(UNavCollision::GetCookedData);
+
 	const bool bUseConvexCollision = bGatherConvexGeometry || (BoxCollision.Num() > 0) || (CylinderCollision.Num() > 0);
 	if (IsTemplate() || !bUseConvexCollision)
 	{
