Remove some deprecated methods and members in navigation

diff --git a/Engine/Source/Runtime/NavigationSystem/Private/NavAreas/NavAreaMeta.cpp b/Engine/Source/Runtime/NavigationSystem/Private/NavAreas/NavAreaMeta.cpp
--- a/Engine/Source/Runtime/NavigationSystem/Private/NavAreas/NavAreaMeta.cpp
+++ b/Engine/Source/Runtime/NavigationSystem/Private/NavAreas/NavAreaMeta.cpp
@@ -7,17 +7,3 @@ UNavAreaMeta::UNavAreaMeta(const FObjectInitializer& ObjectInitializer)
 {
 	bIsMetaArea = true;
 }
-
-TSubclassOf<UNavArea> UNavAreaMeta::PickAreaClass(TSubclassOf<UNavArea> AreaClass, const AActor* Actor, const FNavAgentProperties& NavAgent)
-{
-	return Actor
-		? TSubclassOf<UNavArea>(UNavAreaBase::PickAreaClassForAgent(AreaClass, *Actor, NavAgent))
-		: AreaClass;
-}
-
-TSubclassOf<UNavArea> UNavAreaMeta::PickAreaClass(const AActor* Actor, const FNavAgentProperties& NavAgent)
-{
-	return Actor
-		? TSubclassOf<UNavArea>(PickAreaClassForAgent(*Actor, NavAgent))
-		: TSubclassOf<UNavArea>(GetClass());
-}
diff --git a/Engine/Source/Runtime/NavigationSystem/Private/NavigationSystem.cpp b/Engine/Source/Runtime/NavigationSystem/Private/NavigationSystem.cpp
--- a/Engine/Source/Runtime/NavigationSystem/Private/NavigationSystem.cpp
+++ b/Engine/Source/Runtime/NavigationSystem/Private/NavigationSystem.cpp
@@ -398,8 +398,6 @@ UNavigationSystemV1::UNavigationSystemV1(const FObjectInitializer& ObjectInitial
 	, bWorldInitDone(false)
 	, bCleanUpDone(false)
 	, CurrentlyDrawnNavDataIndex(0)
-	, NavOctree(nullptr)
-	, bNavOctreeLock(false)
 {
 #if WITH_EDITOR
 	NavUpdateLockFlags = 0;
@@ -4657,92 +4655,6 @@ void UNavigationSystemV1::UnregisterNavigationInvoker(AActor* Invoker)
 	}
 }
 
-//----------------------------------------------------------------------//
-// DEPRECATED
-//----------------------------------------------------------------------//
-void UNavigationSystemV1::RequestRegistration(ANavigationData* NavData, bool bTriggerRegistrationProcessing)
-{
-	if (NavData)
-	{
-		RequestRegistrationDeferred(*NavData);
-	}	
-}
-
-FVector UNavigationSystemV1::ProjectPointToNavigation(UObject* WorldContextObject, const FVector& Point, ANavigationData* NavData, TSubclassOf<UNavigationQueryFilter> FilterClass, const FVector QueryExtent)
-{
-	FNavLocation ProjectedPoint(Point);
-
-	UWorld* World = GEngine->GetWorldFromContextObject(WorldContextObject, EGetWorldErrorMode::LogAndReturnNull);
-	UNavigationSystemV1* NavSys = FNavigationSystem::GetCurrent<UNavigationSystemV1>(World);
-	if (NavSys)
-	{
-		ANavigationData* UseNavData = NavData ? NavData : NavSys->GetDefaultNavDataInstance(FNavigationSystem::DontCreate);
-		if (UseNavData)
-		{
-			NavSys->ProjectPointToNavigation(Point, ProjectedPoint, QueryExtent.IsNearlyZero() ? INVALID_NAVEXTENT : QueryExtent, UseNavData,
-				UNavigationQueryFilter::GetQueryFilter(*UseNavData, WorldContextObject, FilterClass));
-		}
-	}
-
-	return ProjectedPoint.Location;
-}
-
-FVector UNavigationSystemV1::GetRandomReachablePointInRadius(UObject* WorldContextObject, const FVector& Origin, float Radius, ANavigationData* NavData, TSubclassOf<UNavigationQueryFilter> FilterClass)
-{
-	FNavLocation RandomPoint;
-
-	UWorld* World = GEngine->GetWorldFromContextObject(WorldContextObject, EGetWorldErrorMode::LogAndReturnNull);
-	UNavigationSystemV1* NavSys = FNavigationSystem::GetCurrent<UNavigationSystemV1>(World);
-	if (NavSys)
-	{
-		ANavigationData* UseNavData = NavData ? NavData : NavSys->GetDefaultNavDataInstance(FNavigationSystem::DontCreate);
-		if (UseNavData)
-		{
-			NavSys->GetRandomReachablePointInRadius(Origin, Radius, RandomPoint, UseNavData, UNavigationQueryFilter::GetQueryFilter(*UseNavData, WorldContextObject, FilterClass));
-		}
-	}
-
-	return RandomPoint.Location;
-}
-
-FVector UNavigationSystemV1::GetRandomPointInNavigableRadius(UObject* WorldContextObject, const FVector& Origin, float Radius, ANavigationData* NavData, TSubclassOf<UNavigationQueryFilter> FilterClass)
-{
-	FNavLocation RandomPoint;
-
-	UWorld* World = GEngine->GetWorldFromContextObject(WorldContextObject, EGetWorldErrorMode::LogAndReturnNull);
-	UNavigationSystemV1* NavSys = FNavigationSystem::GetCurrent<UNavigationSystemV1>(World);
-	if (NavSys)
-	{
-		ANavigationData* UseNavData = NavData ? NavData : NavSys->GetDefaultNavDataInstance(FNavigationSystem::DontCreate);
-		if (UseNavData)
-		{
-			NavSys->GetRandomPointInNavigableRadius(Origin, Radius, RandomPoint, UseNavData, UNavigationQueryFilter::GetQueryFilter(*UseNavData, WorldContextObject, FilterClass));
-		}
-	}
-
-	return RandomPoint.Location;
-}
-
-bool UNavigationSystemV1::K2_GetRandomPointInNavigableRadius(UObject* WorldContextObject, const FVector& Origin, FVector& RandomLocation, float Radius, ANavigationData* NavData, TSubclassOf<UNavigationQueryFilter> FilterClass)
-{
-	return K2_GetRandomLocationInNavigableRadius(WorldContextObject, Origin, RandomLocation, Radius, NavData, FilterClass);
-}
-
-void UNavigationSystemV1::SimpleMoveToActor(AController* Controller, const AActor* Goal)
-{
-	UE_LOG(LogNavigation, Error, TEXT("SimpleMoveToActor is deprecated. Use UAIBlueprintHelperLibrary::SimpleMoveToActor instead"));
-}
-
-void UNavigationSystemV1::SimpleMoveToLocation(AController* Controller, const FVector& Goal)
-{
-	UE_LOG(LogNavigation, Error, TEXT("SimpleMoveToLocation is deprecated. Use UAIBlueprintHelperLibrary::SimpleMoveToLocation instead"));
-}
-
-ANavigationData* UNavigationSystemV1::CreateNavigationDataInstance(const FNavDataConfig& NavConfig)
-{
-	return CreateNavigationDataInstanceInLevel(NavConfig, nullptr);
-}
-
 //----------------------------------------------------------------------//
 // NEW STUFF!
 //----------------------------------------------------------------------//
diff --git a/Engine/Source/Runtime/NavigationSystem/Private/NavigationSystemTypes.cpp b/Engine/Source/Runtime/NavigationSystem/Private/NavigationSystemTypes.cpp
--- a/Engine/Source/Runtime/NavigationSystem/Private/NavigationSystemTypes.cpp
+++ b/Engine/Source/Runtime/NavigationSystem/Private/NavigationSystemTypes.cpp
@@ -248,25 +248,6 @@ namespace NavigationHelper
 			// AND has full colliding capabilities 
 			&& BodySetup.DefaultInstance.GetCollisionEnabled() == ECollisionEnabled::QueryAndPhysics;
 	}
-
-	//////////////////////////////////////////////////////////////////////////
-	// DEPRECATED FUNCTIONS
-
-	void DefaultNavLinkProcessorImpl(FCompositeNavModifier* OUT CompositeModifier, const AActor* Actor, const TArray<FNavigationLink>& IN NavLinks)
-	{
-		if (Actor)
-		{
-			DefaultNavLinkProcessorImpl(CompositeModifier, FNavLinkOwnerData(*Actor), NavLinks);
-		}
-	}
-
-	void DefaultNavLinkSegmentProcessorImpl(FCompositeNavModifier* OUT CompositeModifier, const AActor* Actor, const TArray<FNavigationSegmentLink>& IN NavLinks)
-	{
-		if (Actor)
-		{
-			DefaultNavLinkSegmentProcessorImpl(CompositeModifier, FNavLinkOwnerData(*Actor), NavLinks);
-		}
-	}
 }
 
 #include "NavLinkHostInterface.h"
diff --git a/Engine/Source/Runtime/NavigationSystem/Public/NavAreas/NavAreaMeta.h b/Engine/Source/Runtime/NavigationSystem/Public/NavAreas/NavAreaMeta.h
--- a/Engine/Source/Runtime/NavigationSystem/Public/NavAreas/NavAreaMeta.h
+++ b/Engine/Source/Runtime/NavigationSystem/Public/NavAreas/NavAreaMeta.h
@@ -20,10 +20,4 @@ class NAVIGATIONSYSTEM_API UNavAreaMeta : public UNavArea
 
 public:
 	UNavAreaMeta(const FObjectInitializer& ObjectInitializer);
-
-	UE_DEPRECATED(4.20, "UNavAreaMeta::PickAreaClass is deprecated. Use UNavArea::PickAreaClassForAgent instead")
-	static TSubclassOf<UNavArea> PickAreaClass(TSubclassOf<UNavArea> AreaClass, const AActor* Actor, const FNavAgentProperties& NavAgent);
-
-	UE_DEPRECATED(4.20, "UNavAreaMeta::PickAreaClass is deprecated. Use UNavArea::PickAreaClassForAgent instead")
-	virtual TSubclassOf<UNavArea> PickAreaClass(const AActor* Actor, const FNavAgentProperties& NavAgent);
 };
diff --git a/Engine/Source/Runtime/NavigationSystem/Public/NavigationSystem.h b/Engine/Source/Runtime/NavigationSystem/Public/NavigationSystem.h
--- a/Engine/Source/Runtime/NavigationSystem/Public/NavigationSystem.h
+++ b/Engine/Source/Runtime/NavigationSystem/Public/NavigationSystem.h
@@ -687,8 +687,6 @@ public:
 
 	/** Adds NavData to registration candidates queue - NavDataRegistrationQueue*/
 	virtual void RequestRegistrationDeferred(ANavigationData& NavData);
-	UE_DEPRECATED(4.24, "This version of RequestRegistration is deprecated. Please use the RequestRegistrationDeferred as the registration request is always queued now.")
-	virtual void RequestRegistration(ANavigationData* NavData, bool bTriggerRegistrationProcessing = true);
 
 protected:
 	void ApplySupportedAgentsFilter();
@@ -1216,59 +1214,6 @@ public:
 	// DEPRECATED
 	//----------------------------------------------------------------------//
 public:
-	UE_DEPRECATED(4.16, "This version of ProjectPointToNavigation is deprecated. Please use the new version")
-	UFUNCTION(BlueprintPure, Category = "AI|Navigation", meta = (WorldContext = "WorldContextObject", DisplayName = "ProjectPointToNavigation_DEPRECATED", ScriptNoExport, DeprecatedFunction, DeprecationMessage = "This version of ProjectPointToNavigation is deprecated. Please use the new version"))
-	static FVector ProjectPointToNavigation(UObject* WorldContextObject, const FVector& Point, ANavigationData* NavData = NULL, TSubclassOf<UNavigationQueryFilter> FilterClass = NULL, const FVector QueryExtent = FVector::ZeroVector);
-	UE_DEPRECATED(4.16, "This version of GetRandomReachablePointInRadius is deprecated. Please use the new version")
-	UFUNCTION(BlueprintPure, Category = "AI|Navigation", meta = (WorldContext = "WorldContextObject", DisplayName = "GetRandomReachablePointInRadius_DEPRECATED", ScriptNoExport, DeprecatedFunction, DeprecationMessage = "This version of GetRandomReachablePointInRadius is deprecated. Please use the new version"))
-	static FVector GetRandomReachablePointInRadius(UObject* WorldContextObject, const FVector& Origin, float Radius, ANavigationData* NavData = NULL, TSubclassOf<UNavigationQueryFilter> FilterClass = NULL);
-	UE_DEPRECATED(4.16, "This version of GetRandomPointInNavigableRadius is deprecated. Please use the new version")
-	UFUNCTION(BlueprintPure, Category = "AI|Navigation", meta = (WorldContext = "WorldContextObject", DisplayName = "GetRandomPointInNavigableRadius_DEPRECATED", ScriptNoExport, DeprecatedFunction, DeprecationMessage = "This version of GetRandomPointInNavigableRadius is deprecated. Please use the new version"))
-	static FVector GetRandomPointInNavigableRadius(UObject* WorldContextObject, const FVector& Origin, float Radius, ANavigationData* NavData = NULL, TSubclassOf<UNavigationQueryFilter> FilterClass = NULL);
-	UE_DEPRECATED(4.20, "SimpleMoveToActor is deprecated. Use UAIBlueprintHelperLibrary::SimpleMoveToActor instead")
-	UFUNCTION(BlueprintCallable, Category = "AI|Navigation", meta = (DisplayName = "SimpleMoveToActor_DEPRECATED", ScriptNoExport, DeprecatedFunction, DeprecationMessage = "SimpleMoveToActor is deprecated. Use AIBlueprintHelperLibrary::SimpleMoveToActor instead"))
-	static void SimpleMoveToActor(AController* Controller, const AActor* Goal);
-	UE_DEPRECATED(4.20, "SimpleMoveToLocation is deprecated. Use UAIBlueprintHelperLibrary::SimpleMoveToLocation instead")
-	UFUNCTION(BlueprintCallable, Category = "AI|Navigation", meta = (DisplayName = "SimpleMoveToLocation_DEPRECATED", ScriptNoExport, DeprecatedFunction, DeprecationMessage = "SimpleMoveToLocation is deprecated. Use AIBlueprintHelperLibrary::SimpleMoveToLocation instead"))
-	static void SimpleMoveToLocation(AController* Controller, const FVector& Goal);
-	UE_DEPRECATED(4.20, "UNavigationSystemV1::GetDefaultWalkableArea is deprecated. Use FNavigationSystem::GetDefaultWalkableArea instead")
-	static TSubclassOf<UNavAreaBase> GetDefaultWalkableArea() { return FNavigationSystem::GetDefaultWalkableArea(); }
-	UE_DEPRECATED(4.20, "UNavigationSystemV1::GetDefaultObstacleArea is deprecated. Use FNavigationSystem::GetDefaultObstacleArea instead")
-	static TSubclassOf<UNavAreaBase> GetDefaultObstacleArea() { return FNavigationSystem::GetDefaultObstacleArea(); }
-	UE_DEPRECATED(4.22, "This version is deprecated.  Please use GetRandomLocationInNavigableRadius instead")
-	UFUNCTION(BlueprintPure, Category = "AI|Navigation", meta = (WorldContext = "WorldContextObject", DisplayName = "GetRandomPointInNavigableRadius", ScriptName = "GetRandomPointInNavigableRadius"))
-	static bool K2_GetRandomPointInNavigableRadius(UObject* WorldContextObject, const FVector& Origin, FVector& RandomLocation, float Radius, ANavigationData* NavData = NULL, TSubclassOf<UNavigationQueryFilter> FilterClass = NULL);
-	UE_DEPRECATED(4.23, "This function is deprecated.  Please use CreateNavigationDataInstanceInLevel instead")
-	virtual ANavigationData* CreateNavigationDataInstance(const FNavDataConfig& NavConfig);
-	
-	UE_DEPRECATED(4.24, "This function is deprecated and no longer used. NavigationSystem is no longer involved in storing navoctree element IDs. See FNavigationOctree for more details.")
-	void SetObjectsNavOctreeId(const UObject& Object, FOctreeElementId2 Id) {}
-
-	UE_DEPRECATED(4.24, "This member is deprecated and no longer used.  Please access OctreeController instead")
-	TMap<uint32, FOctreeElementId2> ObjectToOctreeId;
-	UE_DEPRECATED(4.24, "This member is deprecated and no longer used.  Please access OctreeController instead")
-	TSet<FNavigationDirtyElement> PendingOctreeUpdates;
-	UE_DEPRECATED(4.24, "This member is deprecated and no longer used.  Please access OctreeController instead")
-	TSharedPtr<FNavigationOctree, ESPMode::ThreadSafe> NavOctree;
-	UE_DEPRECATED(4.24, "This member is deprecated and no longer used.  Please access OctreeController instead")
-	TMultiMap<UObject*, FWeakObjectPtr> OctreeChildNodesMap;
-	UE_DEPRECATED(4.24, "This member is deprecated and no longer used.  Please access OctreeController instead")
-	uint8 bNavOctreeLock : 1;
-
-	UE_DEPRECATED(4.24, "This member is deprecated and no longer used.  Please access DirtyAreasController instead")
-	UPROPERTY(config, EditAnywhere, Category = NavigationSystem)
-	float DirtyAreasUpdateFreq = 60;
-	UE_DEPRECATED(4.24, "This member is deprecated and no longer used.  Please access DirtyAreasController instead")
-	float DirtyAreasUpdateTime = 0;
-	UE_DEPRECATED(4.24, "This member is deprecated and no longer used.  Please access DirtyAreasController instead")
-	TArray<FNavigationDirtyArea> DirtyAreas;
-	UE_DEPRECATED(4.24, "This member is deprecated and no longer used.  Please access DirtyAreasController instead")
-	uint8 bCanAccumulateDirtyAreas : 1;
-#if !UE_BUILD_SHIPPING
-	UE_DEPRECATED(4.24, "This member is deprecated and no longer used.  Please access DirtyAreasController instead")
-	uint8 bDirtyAreasReportedWhileAccumulationLocked : 1;
-#endif // !UE_BUILD_SHIPPING
-
 	UE_DEPRECATED(4.26, "This version of RemoveNavigationBuildLock is deprecated. Please use the new version")
 	void RemoveNavigationBuildLock(uint8 Flags, bool bSkipRebuildInEditor) { RemoveNavigationBuildLock(Flags, bSkipRebuildInEditor ? ELockRemovalRebuildAction::RebuildIfNotInEditor : ELockRemovalRebuildAction::Rebuild);}
 };
diff --git a/Engine/Source/Runtime/NavigationSystem/Public/NavigationSystemTypes.h b/Engine/Source/Runtime/NavigationSystem/Public/NavigationSystemTypes.h
--- a/Engine/Source/Runtime/NavigationSystem/Public/NavigationSystemTypes.h
+++ b/Engine/Source/Runtime/NavigationSystem/Public/NavigationSystemTypes.h
@@ -177,11 +177,4 @@ namespace NavigationHelper
 	NAVIGATIONSYSTEM_API void DefaultNavLinkSegmentProcessorImpl(FCompositeNavModifier* OUT CompositeModifier, const FNavLinkOwnerData& OwnerData, const TArray<FNavigationSegmentLink>& IN NavLinks);
 
 	NAVIGATIONSYSTEM_API bool IsBodyNavigationRelevant(const UBodySetup& IN BodySetup);
-
-	// deprecated functions	
-	UE_DEPRECATED_FORGAME(4.17, "This function is now deprecated, please use override with OwnerData argument.")
-	NAVIGATIONSYSTEM_API void DefaultNavLinkProcessorImpl(FCompositeNavModifier* OUT CompositeModifier, const AActor* Actor, const TArray<FNavigationLink>& IN NavLinks);
-
-	UE_DEPRECATED_FORGAME(4.17, "This function is now deprecated, please use override with OwnerData argument.")
-	NAVIGATIONSYSTEM_API void DefaultNavLinkSegmentProcessorImpl(FCompositeNavModifier* OUT CompositeModifier, const AActor* Actor, const TArray<FNavigationSegmentLink>& IN NavLinks);
 }
\ No newline at end of file
