Bring back deprecated function with blueprint deprecation metadata to allow users to update.

diff --git a/Engine/Source/Runtime/NavigationSystem/Private/NavigationSystem.cpp b/Engine/Source/Runtime/NavigationSystem/Private/NavigationSystem.cpp
--- a/Engine/Source/Runtime/NavigationSystem/Private/NavigationSystem.cpp
+++ b/Engine/Source/Runtime/NavigationSystem/Private/NavigationSystem.cpp
@@ -4655,6 +4655,14 @@ void UNavigationSystemV1::UnregisterNavigationInvoker(AActor* Invoker)
 	}
 }
 
+//----------------------------------------------------------------------//
+// DEPRECATED
+//----------------------------------------------------------------------//
+bool UNavigationSystemV1::K2_GetRandomPointInNavigableRadius(UObject* WorldContextObject, const FVector& Origin, FVector& RandomLocation, float Radius, ANavigationData* NavData, TSubclassOf<UNavigationQueryFilter> FilterClass)
+{
+	return K2_GetRandomLocationInNavigableRadius(WorldContextObject, Origin, RandomLocation, Radius, NavData, FilterClass);
+}
+
 //----------------------------------------------------------------------//
 // NEW STUFF!
 //----------------------------------------------------------------------//
diff --git a/Engine/Source/Runtime/NavigationSystem/Public/NavigationSystem.h b/Engine/Source/Runtime/NavigationSystem/Public/NavigationSystem.h
--- a/Engine/Source/Runtime/NavigationSystem/Public/NavigationSystem.h
+++ b/Engine/Source/Runtime/NavigationSystem/Public/NavigationSystem.h
@@ -1214,6 +1214,11 @@ public:
 	// DEPRECATED
 	//----------------------------------------------------------------------//
 public:
+	// Note that this function was only deprecated for blueprint in 5.1
+	UE_DEPRECATED(4.22, "This version is deprecated.  Please use GetRandomLocationInNavigableRadius instead")
+	UFUNCTION(BlueprintPure, Category = "AI|Navigation", meta = (WorldContext = "WorldContextObject", DisplayName = "GetRandomPointInNavigableRadius", ScriptName = "GetRandomPointInNavigableRadius", DeprecatedFunction, DeprecationMessage = "GetRandomPointInNavigableRadius is deprecated. Use GetRandomLocationInNavigableRadius instead"))
+	static bool K2_GetRandomPointInNavigableRadius(UObject* WorldContextObject, const FVector& Origin, FVector& RandomLocation, float Radius, ANavigationData* NavData = NULL, TSubclassOf<UNavigationQueryFilter> FilterClass = NULL);
+	
 	UE_DEPRECATED(4.26, "This version of RemoveNavigationBuildLock is deprecated. Please use the new version")
 	void RemoveNavigationBuildLock(uint8 Flags, bool bSkipRebuildInEditor) { RemoveNavigationBuildLock(Flags, bSkipRebuildInEditor ? ELockRemovalRebuildAction::RebuildIfNotInEditor : ELockRemovalRebuildAction::Rebuild);}
 };
