[WP Dynamic Navmesh] First pass on WP dynamic navmesh
- Dynamic navmesh in a world partitioned map is now allowed to build a base navmesh and stream cells from it
- Ignore navigation dirtiness coming from objects loading/unloading that are part of the base navmesh
- Addition of vlog boxes on addition and removal of navigation data chunk actors

diff --git a/Engine/Source/Runtime/Engine/Classes/AI/Navigation/NavigationDirtyElement.h b/Engine/Source/Runtime/Engine/Classes/AI/Navigation/NavigationDirtyElement.h
--- a/Engine/Source/Runtime/Engine/Classes/AI/Navigation/NavigationDirtyElement.h
+++ b/Engine/Source/Runtime/Engine/Classes/AI/Navigation/NavigationDirtyElement.h
@@ -6,7 +6,7 @@
 
 class INavRelevantInterface;
 
-struct FNavigationDirtyElement
+struct ENGINE_API FNavigationDirtyElement
 {
 	/** object owning this element */
 	FWeakObjectPtr Owner;
@@ -29,20 +29,23 @@ struct FNavigationDirtyElement
 	/** request was invalidated while queued, use prev values to dirty area */
 	uint8 bInvalidRequest : 1;
 
+	/** requested during visibility change of the owning level (loading/unloading) */
+	uint8 bIsFromVisibilityChange : 1;
+
+	/** part of the base navmesh */
+	uint8 bIsInBaseNavmesh : 1;
+
 	FNavigationDirtyElement()
-		: NavInterface(0), FlagsOverride(0), PrevFlags(0), PrevBounds(ForceInit), bHasPrevData(false), bInvalidRequest(false)
+		: NavInterface(nullptr), FlagsOverride(0), PrevFlags(0), PrevBounds(ForceInit), bHasPrevData(false), bInvalidRequest(false), bIsFromVisibilityChange(false), bIsInBaseNavmesh(false)
 	{
 	}
 
 	FNavigationDirtyElement(UObject* InOwner)
-		: Owner(InOwner), NavInterface(0), FlagsOverride(0), PrevFlags(0), PrevBounds(ForceInit), bHasPrevData(false), bInvalidRequest(false)
+		: Owner(InOwner), NavInterface(nullptr), FlagsOverride(0), PrevFlags(0), PrevBounds(ForceInit), bHasPrevData(false), bInvalidRequest(false), bIsFromVisibilityChange(false), bIsInBaseNavmesh(false)
 	{
 	}
 
-	FNavigationDirtyElement(UObject* InOwner, INavRelevantInterface* InNavInterface, int32 InFlagsOverride = 0)
-		: Owner(InOwner), NavInterface(InNavInterface),	FlagsOverride(InFlagsOverride), PrevFlags(0), PrevBounds(ForceInit), bHasPrevData(false), bInvalidRequest(false)
-	{
-	}
+	FNavigationDirtyElement(UObject* InOwner, INavRelevantInterface* InNavInterface, int32 InFlagsOverride = 0, const bool bUseWorldPartitionedDynamicMode = false);
 
 	bool operator==(const FNavigationDirtyElement& Other) const 
 	{ 
diff --git a/Engine/Source/Runtime/Engine/Classes/AI/Navigation/NavigationTypes.h b/Engine/Source/Runtime/Engine/Classes/AI/Navigation/NavigationTypes.h
--- a/Engine/Source/Runtime/Engine/Classes/AI/Navigation/NavigationTypes.h
+++ b/Engine/Source/Runtime/Engine/Classes/AI/Navigation/NavigationTypes.h
@@ -59,6 +59,12 @@ namespace FNavigationSystem
 	{
 		return FVector::PointsAreSame(A.Min, B.Min) && FVector::PointsAreSame(A.Max, B.Max);
 	}
+
+	/** Returns true if the visibility of the owning level is currently changing (loading/unloading). */
+	ENGINE_API bool IsLevelVisibilityChanging(const UObject* Object);
+	
+	/** Objects placed directly in the level and objects placed in the base navmesh data layers are in the base navmesh. */
+	ENGINE_API bool IsInBaseNavmesh(const UObject* Object);
 }
 
 UENUM()
diff --git a/Engine/Source/Runtime/Engine/Private/AI/Navigation/NavigationDirtyElement.cpp b/Engine/Source/Runtime/Engine/Private/AI/Navigation/NavigationDirtyElement.cpp
new file mode 100644
--- /dev/null
+++ b/Engine/Source/Runtime/Engine/Private/AI/Navigation/NavigationDirtyElement.cpp
@@ -0,0 +1,25 @@
+﻿// Copyright Epic Games, Inc. All Rights Reserved.
+
+#include "AI/Navigation/NavigationDirtyElement.h"
+#include "AI/Navigation/NavigationTypes.h"
+
+FNavigationDirtyElement::FNavigationDirtyElement(UObject* InOwner, INavRelevantInterface* InNavInterface, int32 InFlagsOverride /*= 0*/, const bool bUseWorldPartitionedDynamicMode /*= false*/)
+	: Owner(InOwner)
+	, NavInterface(InNavInterface)
+	, FlagsOverride(InFlagsOverride)
+	, PrevFlags(0)
+	, PrevBounds(ForceInit)
+	, bHasPrevData(false)
+	, bInvalidRequest(false)
+{
+	if (bUseWorldPartitionedDynamicMode)
+	{
+		bIsFromVisibilityChange = FNavigationSystem::IsLevelVisibilityChanging(InOwner);
+		bIsInBaseNavmesh = FNavigationSystem::IsInBaseNavmesh(InOwner);
+	}
+	else
+	{
+		bIsFromVisibilityChange = false;
+		bIsInBaseNavmesh = false;
+	}
+}
diff --git a/Engine/Source/Runtime/Engine/Private/AI/Navigation/NavigationTypes.cpp b/Engine/Source/Runtime/Engine/Private/AI/Navigation/NavigationTypes.cpp
--- a/Engine/Source/Runtime/Engine/Private/AI/Navigation/NavigationTypes.cpp
+++ b/Engine/Source/Runtime/Engine/Private/AI/Navigation/NavigationTypes.cpp
@@ -7,7 +7,8 @@
 #include "EngineStats.h"
 #include "Components/ShapeComponent.h"
 #include "AI/Navigation/NavAreaBase.h"
-
+#include "GameFramework/WorldSettings.h"
+#include "WorldPartition/DataLayer/DataLayerAsset.h"
 
 DEFINE_STAT(STAT_Navigation_MetaAreaTranslation);
 
@@ -20,6 +21,56 @@ namespace FNavigationSystem
 	// and only those values will be used
 	const float FallbackAgentRadius = 35.f;
 	const float FallbackAgentHeight = 144.f;
+
+	bool IsLevelVisibilityChanging(const UObject* Object)
+	{
+		const UActorComponent* ObjectAsComponent = Cast<UActorComponent>(Object);
+		if (ObjectAsComponent)
+		{
+			if (const ULevel* Level = ObjectAsComponent->GetComponentLevel())
+			{
+				return Level->HasVisibilityChangeRequestPending();
+			}
+		}
+		else if (const AActor* Actor = Cast<AActor>(Object))
+		{
+			if (const ULevel* Level = Actor->GetLevel())
+			{
+				return Level->HasVisibilityChangeRequestPending();
+			}
+		}
+
+		return false;
+	}
+	
+	bool IsInBaseNavmesh(const UObject* Object)
+	{
+		const UActorComponent* ObjectAsComponent = Cast<UActorComponent>(Object);
+		if (const AActor* Actor = ObjectAsComponent ? ObjectAsComponent->GetOwner() : Cast<AActor>(Object))
+		{
+			if (!Actor->HasDataLayers())
+			{
+				return true;
+			}
+		
+			if (const UWorld* World = Object->GetWorld())
+			{
+				if (const AWorldSettings* WorldSettings = World->GetWorldSettings())
+				{
+					const TArray<TObjectPtr<UDataLayerAsset>>& BaseNavmeshLayers = WorldSettings->BaseNavmeshDataLayers;
+					for (const TObjectPtr<UDataLayerAsset>& DataLayer : BaseNavmeshLayers)
+					{
+						if (Actor->ContainsDataLayer(DataLayer))
+						{
+							return true;
+						}
+					}
+				}
+			}
+		}
+
+		return false;
+	}	
 }
 
 //----------------------------------------------------------------------//
diff --git a/Engine/Source/Runtime/Engine/Public/WorldPartition/NavigationData/NavigationDataChunkActor.h b/Engine/Source/Runtime/Engine/Public/WorldPartition/NavigationData/NavigationDataChunkActor.h
--- a/Engine/Source/Runtime/Engine/Public/WorldPartition/NavigationData/NavigationDataChunkActor.h
+++ b/Engine/Source/Runtime/Engine/Public/WorldPartition/NavigationData/NavigationDataChunkActor.h
@@ -41,6 +41,8 @@ public:
 	virtual FBox GetStreamingBounds() const override;
 	//~ End AActor Interface.
 #endif // WITH_EDITOR
+
+	virtual FBox GetBounds() const { return DataChunkActorBounds; }
 	
 protected:
 	//~ Begin AActor Interface.
diff --git a/Engine/Source/Runtime/NavigationSystem/Private/NavMesh/RecastNavMesh.cpp b/Engine/Source/Runtime/NavigationSystem/Private/NavMesh/RecastNavMesh.cpp
--- a/Engine/Source/Runtime/NavigationSystem/Private/NavMesh/RecastNavMesh.cpp
+++ b/Engine/Source/Runtime/NavigationSystem/Private/NavMesh/RecastNavMesh.cpp
@@ -2295,6 +2295,8 @@ void ARecastNavMesh::OnStreamingNavDataAdded(ANavigationDataChunkActor& InActor)
 
 	if (SupportsStreaming() && RecastNavMeshImpl)
 	{
+		UE_VLOG_BOX(this, LogNavigation, Log, InActor.GetBounds(), FColor::Blue, TEXT(""));
+		
 		URecastNavMeshDataChunk* NavDataChunk = GetNavigationDataChunk(InActor);
 		if (NavDataChunk)
 		{
@@ -2309,6 +2311,8 @@ void ARecastNavMesh::OnStreamingNavDataRemoved(ANavigationDataChunkActor& InActo
 
 	if (SupportsStreaming() && RecastNavMeshImpl)
 	{
+		UE_VLOG_BOX(this, LogNavigation, Log, InActor.GetBounds(), FColor::Red, TEXT(""));
+		
 		URecastNavMeshDataChunk* NavDataChunk = GetNavigationDataChunk(InActor);
 		if (NavDataChunk)
 		{
@@ -2782,7 +2786,7 @@ bool ARecastNavMesh::SupportsStreaming() const
 {
 	// Actually nothing prevents us to support streaming with dynamic generation
 	// Right now streaming in sub-level causes navmesh to build itself, so no point to stream tiles in
-	return (RuntimeGeneration != ERuntimeGenerationType::Dynamic);
+	return (RuntimeGeneration != ERuntimeGenerationType::Dynamic) || bIsWorldPartitioned;
 }
 
 FRecastNavMeshGenerator* ARecastNavMesh::CreateGeneratorInstance()
diff --git a/Engine/Source/Runtime/NavigationSystem/Private/NavMesh/RecastNavMeshGenerator.cpp b/Engine/Source/Runtime/NavigationSystem/Private/NavMesh/RecastNavMeshGenerator.cpp
--- a/Engine/Source/Runtime/NavigationSystem/Private/NavMesh/RecastNavMeshGenerator.cpp
+++ b/Engine/Source/Runtime/NavigationSystem/Private/NavMesh/RecastNavMeshGenerator.cpp
@@ -1779,7 +1779,7 @@ protected:
 			UE_LOG(LogNavigation, Log, TEXT("Recast: %s"), ANSI_TO_TCHAR( Msg ) );
 			break;
 		default:
-			UE_LOG(LogNavigation, Verbose, TEXT("Recast: %s"), ANSI_TO_TCHAR( Msg ) );
+			UE_LOG(LogNavigation, VeryVerbose, TEXT("Recast: %s"), ANSI_TO_TCHAR( Msg ) );
 			break;
 		}
 	}
diff --git a/Engine/Source/Runtime/NavigationSystem/Private/NavigationDataHandler.cpp b/Engine/Source/Runtime/NavigationSystem/Private/NavigationDataHandler.cpp
--- a/Engine/Source/Runtime/NavigationSystem/Private/NavigationDataHandler.cpp
+++ b/Engine/Source/Runtime/NavigationSystem/Private/NavigationDataHandler.cpp
@@ -69,7 +69,7 @@ FSetElementId FNavigationDataHandler::RegisterNavOctreeElement(UObject& ElementO
 
 		if (bCanAdd)
 		{
-			FNavigationDirtyElement UpdateInfo(&ElementOwner, &ElementInterface, GetDirtyFlagHelper(UpdateFlags, 0));
+			FNavigationDirtyElement UpdateInfo(&ElementOwner, &ElementInterface, GetDirtyFlagHelper(UpdateFlags, 0), DirtyAreasController.bUseWorldPartitionedDynamicMode);
 
 			SetId = OctreeController.PendingOctreeUpdates.FindId(UpdateInfo);
 			if (SetId.IsValidId())
@@ -96,7 +96,7 @@ void FNavigationDataHandler::AddElementToNavOctree(const FNavigationDirtyElement
 	{
 		if (DirtyElement.bHasPrevData)
 		{
-			DirtyAreasController.AddArea(DirtyElement.PrevBounds, DirtyElement.PrevFlags, [&DirtyElement] { return DirtyElement.Owner.Get(); });
+			DirtyAreasController.AddArea(DirtyElement.PrevBounds, DirtyElement.PrevFlags, [&DirtyElement] { return DirtyElement.Owner.Get(); }, &DirtyElement);
 		}
 
 		return;
@@ -146,7 +146,7 @@ void FNavigationDataHandler::AddElementToNavOctree(const FNavigationDirtyElement
 	if (!GeneratedData.IsEmpty())
 	{
 		const int32 DirtyFlag = DirtyElement.FlagsOverride ? DirtyElement.FlagsOverride : GeneratedData.Data->GetDirtyFlag();
-		DirtyAreasController.AddArea(GeneratedData.Bounds.GetBox(), DirtyFlag, [&ElementOwner] { return ElementOwner; });
+		DirtyAreasController.AddArea(GeneratedData.Bounds.GetBox(), DirtyFlag, [&ElementOwner] { return ElementOwner; }, &DirtyElement);
 	}
 }
 
diff --git a/Engine/Source/Runtime/NavigationSystem/Private/NavigationDirtyAreasController.cpp b/Engine/Source/Runtime/NavigationSystem/Private/NavigationDirtyAreasController.cpp
--- a/Engine/Source/Runtime/NavigationSystem/Private/NavigationDirtyAreasController.cpp
+++ b/Engine/Source/Runtime/NavigationSystem/Private/NavigationDirtyAreasController.cpp
@@ -2,6 +2,8 @@
 
 #include "NavigationDirtyAreasController.h"
 #include "NavigationData.h"
+#include "VisualLogger/VisualLogger.h"
+#include "AI/Navigation/NavigationDirtyElement.h"
 
 DEFINE_LOG_CATEGORY(LogNavigationDirtyArea);
 
@@ -10,6 +12,7 @@ DEFINE_LOG_CATEGORY(LogNavigationDirtyArea);
 //----------------------------------------------------------------------//
 FNavigationDirtyAreasController::FNavigationDirtyAreasController()
 	: bCanAccumulateDirtyAreas(true)
+	, bUseWorldPartitionedDynamicMode(false)
 #if !UE_BUILD_SHIPPING
 	, bDirtyAreasReportedWhileAccumulationLocked(false)
 	, bCanReportOversizedDirtyArea(false)
@@ -45,7 +48,7 @@ void FNavigationDirtyAreasController::Tick(const float DeltaSeconds, const TArra
 	}
 }
 
-void FNavigationDirtyAreasController::AddArea(const FBox& NewArea, const int32 Flags, const TFunction<UObject*()>& ObjectProviderFunc /*= nullptr*/)
+void FNavigationDirtyAreasController::AddArea(const FBox& NewArea, const int32 Flags, const TFunction<UObject*()>& ObjectProviderFunc /*= nullptr*/, const FNavigationDirtyElement* DirtyElement /*= nullptr*/)
 {
 #if !UE_BUILD_SHIPPING
 	// always keep track of reported areas even when filtered out by invalid area as long as flags are valid
@@ -65,8 +68,22 @@ void FNavigationDirtyAreasController::AddArea(const FBox& NewArea, const int32 F
 		return;
 	}
 
+	if (bUseWorldPartitionedDynamicMode)
+	{
+		// ObjectProviderFunc() is not always providing a valid object.
+		if (const bool bIsFromVisibilityChange = (DirtyElement && DirtyElement->bIsFromVisibilityChange) || FNavigationSystem::IsLevelVisibilityChanging(ObjectProviderFunc()))
+		{
+			// If the area is from the addition or removal of objects caused by level loading/unloading and it's already in the base navmesh ignore the dirtiness.
+			if (const bool bIsIncludedInBaseNavmesh = (DirtyElement && DirtyElement->bIsInBaseNavmesh) || FNavigationSystem::IsInBaseNavmesh(ObjectProviderFunc()))
+			{
+				UE_LOG(LogNavigationDirtyArea, VeryVerbose, TEXT("Ignoring dirtyness (visibility changed and in base navmesh). (object: %s)"), *GetFullNameSafe(ObjectProviderFunc ? ObjectProviderFunc() : nullptr));
+				return;
+			}
+		}		
+	}
+
 #if !UE_BUILD_SHIPPING
-	auto DumpExtraInfo = [ObjectProviderFunc, BoundsSize]() {
+	auto DumpExtraInfo = [ObjectProviderFunc, BoundsSize, NewArea]() {
 		const UObject* ObjectProvider = nullptr;
 		if (ObjectProviderFunc)
 		{
@@ -75,8 +92,13 @@ void FNavigationDirtyAreasController::AddArea(const FBox& NewArea, const int32 F
 
 		const UActorComponent* ObjectAsComponent = Cast<UActorComponent>(ObjectProvider);
 		const AActor* ComponentOwner = ObjectAsComponent ? ObjectAsComponent->GetOwner() : nullptr;
-		return FString::Printf(TEXT("Adding dirty area object = % | Potential component's owner = %s | Bounds size = %s)"), *GetFullNameSafe(ObjectProvider), *GetFullNameSafe(ComponentOwner), *BoundsSize.ToString());
+		if (ComponentOwner)
+		{
+			UE_VLOG_BOX(ComponentOwner, LogNavigationDirtyArea, Log, NewArea, FColor::Red, TEXT(""));
+		}
+		return FString::Printf(TEXT("Adding dirty area object: %s | Potential component's owner: %s | Bounds size: %s)"), *GetFullNameSafe(ObjectProvider), *GetFullNameSafe(ComponentOwner), *BoundsSize.ToString());
 	};
+
 	UE_LOG(LogNavigationDirtyArea, VeryVerbose, TEXT("%s"), *DumpExtraInfo());
 
 	if (ShouldReportOversizedDirtyArea() && BoundsSize.GetMax() > DirtyAreaWarningSizeThreshold)
@@ -115,6 +137,11 @@ void FNavigationDirtyAreasController::SetDirtyAreaWarningSizeThreshold(const flo
 #endif // !UE_BUILD_SHIPPING
 }
 
+void FNavigationDirtyAreasController::SetUseWorldPartitionedDynamicMode(bool bIsWPDynamic)
+{
+	bUseWorldPartitionedDynamicMode = bIsWPDynamic;
+}
+
 void FNavigationDirtyAreasController::SetCanReportOversizedDirtyArea(bool bCanReport)
 {
 #if !UE_BUILD_SHIPPING
diff --git a/Engine/Source/Runtime/NavigationSystem/Private/NavigationSystem.cpp b/Engine/Source/Runtime/NavigationSystem/Private/NavigationSystem.cpp
--- a/Engine/Source/Runtime/NavigationSystem/Private/NavigationSystem.cpp
+++ b/Engine/Source/Runtime/NavigationSystem/Private/NavigationSystem.cpp
@@ -1022,6 +1022,19 @@ void UNavigationSystemV1::OnWorldInitDone(FNavigationSystemRunMode Mode)
 	// Report oversized dirty areas only in game mode
 	DefaultDirtyAreasController.SetCanReportOversizedDirtyArea(Mode == FNavigationSystemRunMode::GameMode);
 
+	for (const ANavigationData* NavData : NavDataSet)
+	{
+		if (NavData)
+		{
+			const ARecastNavMesh* RecastNavMesh = Cast<ARecastNavMesh>(NavData);
+			if (RecastNavMesh && RecastNavMesh->bIsWorldPartitioned && NavData->GetRuntimeGenerationMode() > ERuntimeGenerationType::Static)
+			{
+				DefaultDirtyAreasController.SetUseWorldPartitionedDynamicMode(true);
+				break;
+			}
+		}
+	}
+
 	bWorldInitDone = true;
 	OnNavigationInitDone.Broadcast();
 	UNavigationSystemBase::OnNavigationInitDoneStaticDelegate().Broadcast(*this);
diff --git a/Engine/Source/Runtime/NavigationSystem/Private/NavigationSystemTypes.cpp b/Engine/Source/Runtime/NavigationSystem/Private/NavigationSystemTypes.cpp
--- a/Engine/Source/Runtime/NavigationSystem/Private/NavigationSystemTypes.cpp
+++ b/Engine/Source/Runtime/NavigationSystem/Private/NavigationSystemTypes.cpp
@@ -5,8 +5,6 @@
 #include "NavMesh/RecastNavMeshGenerator.h"
 #include "NavigationSystem.h"
 #include "Engine/World.h"
-#include "AbstractNavData.h"
-#include "NavigationOctree.h"
 #include "NavCollision.h"
 #include "PhysicsEngine/BodySetup.h"
 #include "Components/StaticMeshComponent.h"
diff --git a/Engine/Source/Runtime/NavigationSystem/Public/NavigationDirtyAreasController.h b/Engine/Source/Runtime/NavigationSystem/Public/NavigationDirtyAreasController.h
--- a/Engine/Source/Runtime/NavigationSystem/Public/NavigationDirtyAreasController.h
+++ b/Engine/Source/Runtime/NavigationSystem/Public/NavigationDirtyAreasController.h
@@ -22,6 +22,8 @@ struct NAVIGATIONSYSTEM_API FNavigationDirtyAreasController
 	TArray<FNavigationDirtyArea> DirtyAreas;
 
 	uint8 bCanAccumulateDirtyAreas : 1;
+	uint8 bUseWorldPartitionedDynamicMode : 1;
+
 #if !UE_BUILD_SHIPPING
 	uint8 bDirtyAreasReportedWhileAccumulationLocked : 1;
 private:
@@ -50,8 +52,9 @@ public:
 	 *	@param NewArea Bounding box of the affected area
 	 *	@param Flags Indicates the type of modification applied to the area
 	 *	@param ObjectProviderFunc Optional function to retrieve source object that can be use for error reporting and navmesh exclusion
+	 *	@param DirtyElement Optional dirty element
 	 */
-	void AddArea(const FBox& NewArea, const int32 Flags, const TFunction<UObject*()>& ObjectProviderFunc = nullptr);
+	void AddArea(const FBox& NewArea, const int32 Flags, const TFunction<UObject*()>& ObjectProviderFunc = nullptr, const FNavigationDirtyElement* DirtyElement = nullptr);
 	
 	bool IsDirty() const { return GetNumDirtyAreas() > 0; }
 	int32 GetNumDirtyAreas() const { return DirtyAreas.Num(); }
@@ -59,6 +62,7 @@ public:
 	void OnNavigationBuildLocked();
 	void OnNavigationBuildUnlocked();
 
+	void SetUseWorldPartitionedDynamicMode(bool bIsWPDynamic);
 	void SetCanReportOversizedDirtyArea(const bool bCanReport);
 	void SetDirtyAreaWarningSizeThreshold(const float Threshold);
 
