Extract custom navlink bounds computation in NavigationSystem custom link registration

diff --git a/Engine/Source/Runtime/NavigationSystem/Private/NavigationSystem.cpp b/Engine/Source/Runtime/NavigationSystem/Private/NavigationSystem.cpp
--- a/Engine/Source/Runtime/NavigationSystem/Private/NavigationSystem.cpp
+++ b/Engine/Source/Runtime/NavigationSystem/Private/NavigationSystem.cpp
@@ -2421,24 +2421,9 @@ void UNavigationSystemV1::RegisterCustomLink(INavLinkCustomInterface& CustomLink
 		UE_LOG(LogNavLink, VeryVerbose, TEXT("%s new navlink id %u."), ANSI_TO_TCHAR(__FUNCTION__), LinkId);
 		CustomLink.UpdateLinkId(LinkId);
 
-		UObject* CustomLinkOb = CustomLink.GetLinkOwner();
-		UActorComponent* OwnerComp = Cast<UActorComponent>(CustomLinkOb);
-		AActor* OwnerActor = OwnerComp ? OwnerComp->GetOwner() : Cast<AActor>(CustomLinkOb);
-
-		if (OwnerActor)
+		const FBox LinkBounds = ComputeCustomLinkBounds(CustomLink);
+		if (LinkBounds.IsValid)
 		{
-			ENavLinkDirection::Type DummyDir = ENavLinkDirection::BothWays;
-			FVector RelativePtA, RelativePtB;
-			CustomLink.GetLinkData(RelativePtA, RelativePtB, DummyDir);
-
-			const FTransform OwnerActorTM = OwnerActor->GetTransform();
-			const FVector WorldPtA = OwnerActorTM.TransformPosition(RelativePtA);
-			const FVector WorldPtB = OwnerActorTM.TransformPosition(RelativePtB);
-
-			FBox LinkBounds(ForceInitToZero);
-			LinkBounds += WorldPtA;
-			LinkBounds += WorldPtB;
-
 			AddDirtyArea(LinkBounds, FNavigationOctreeController::OctreeUpdate_Modifiers);
 		}
 	}
@@ -2497,6 +2482,29 @@ void UNavigationSystemV1::RequestCustomLinkUnregistering(INavLinkCustomInterface
 	}
 }
 
+FBox UNavigationSystemV1::ComputeCustomLinkBounds(const INavLinkCustomInterface& CustomLink)
+{
+	const UObject* CustomLinkOb = CustomLink.GetLinkOwner();
+	const UActorComponent* OwnerComp = Cast<UActorComponent>(CustomLinkOb);
+	const AActor* OwnerActor = OwnerComp ? OwnerComp->GetOwner() : Cast<AActor>(CustomLinkOb);
+
+	FBox LinkBounds(ForceInitToZero);
+	if (OwnerActor)
+	{
+		ENavLinkDirection::Type DummyDir = ENavLinkDirection::BothWays;
+		FVector RelativePtA, RelativePtB;
+		CustomLink.GetLinkData(RelativePtA, RelativePtB, DummyDir);
+
+		const FTransform OwnerActorTM = OwnerActor->GetTransform();
+		const FVector WorldPtA = OwnerActorTM.TransformPosition(RelativePtA);
+		const FVector WorldPtB = OwnerActorTM.TransformPosition(RelativePtB);
+
+		LinkBounds += WorldPtA;
+		LinkBounds += WorldPtB;
+	}
+	return LinkBounds;
+}
+
 void UNavigationSystemV1::RequestAreaUnregistering(UClass* NavAreaClass)
 {
 	for (TObjectIterator<UNavigationSystemV1> NavSysIt; NavSysIt; ++NavSysIt)
diff --git a/Engine/Source/Runtime/NavigationSystem/Public/NavigationSystem.h b/Engine/Source/Runtime/NavigationSystem/Public/NavigationSystem.h
--- a/Engine/Source/Runtime/NavigationSystem/Public/NavigationSystem.h
+++ b/Engine/Source/Runtime/NavigationSystem/Public/NavigationSystem.h
@@ -809,6 +809,9 @@ public:
 	/** updates custom link for all active navigation data instances */
 	void UpdateCustomLink(const INavLinkCustomInterface* CustomLink);
 
+	/** Return a Bounding Box containing the navlink points */
+	static FBox ComputeCustomLinkBounds(const INavLinkCustomInterface& CustomLink);
+
 	//----------------------------------------------------------------------//
 	// Areas
 	//----------------------------------------------------------------------//
