Fix gameplay debugger navmesh rendering issue
When using bRestrictBuildingToActiveTiles, GetDebugGeometry() would gather the data for all the active tiles for each tile in the TileSet instead of only once, leading to unsuable framerate.

diff --git a/Engine/Source/Runtime/NavigationSystem/Private/NavMesh/NavMeshRenderingComponent.cpp b/Engine/Source/Runtime/NavigationSystem/Private/NavMesh/NavMeshRenderingComponent.cpp
--- a/Engine/Source/Runtime/NavigationSystem/Private/NavMesh/NavMeshRenderingComponent.cpp
+++ b/Engine/Source/Runtime/NavigationSystem/Private/NavMesh/NavMeshRenderingComponent.cpp
@@ -511,14 +511,15 @@ void FNavMeshSceneProxyData::GatherData(const ARecastNavMesh* NavMesh, int32 InN
 		NavMesh->BeginBatchQuery();
 		if (TileSet.Num() > 0)
 		{
-			for (int32 Idx = 0; Idx < TileSet.Num(); Idx++)
+			bool bDone = false;
+			for (int32 Idx = 0; Idx < TileSet.Num() && !bDone; Idx++)
 			{
-				NavMesh->GetDebugGeometry(NavMeshGeometry, TileSet[Idx]);
+				bDone = NavMesh->GetDebugGeometryForTile(NavMeshGeometry, TileSet[Idx]);
 			}
 		}
 		else
 		{
-			NavMesh->GetDebugGeometry(NavMeshGeometry);
+			NavMesh->GetDebugGeometryForTile(NavMeshGeometry, INDEX_NONE);
 		}
 
 		const TArray<FVector>& MeshVerts = NavMeshGeometry.MeshVerts;
diff --git a/Engine/Source/Runtime/NavigationSystem/Private/NavMesh/PImplRecastNavMesh.cpp b/Engine/Source/Runtime/NavigationSystem/Private/NavMesh/PImplRecastNavMesh.cpp
--- a/Engine/Source/Runtime/NavigationSystem/Private/NavMesh/PImplRecastNavMesh.cpp
+++ b/Engine/Source/Runtime/NavigationSystem/Private/NavMesh/PImplRecastNavMesh.cpp
@@ -2718,17 +2718,13 @@ uint8 GetValidEnds(const dtNavMesh& NavMesh, const dtMeshTile& Tile, const dtPol
 	return ValidEnds;
 }
 
-/** 
- * @param PolyEdges			[out] Array of worldspace vertex locations for tile edges.  Edges are pairwise verts, i.e. [0,1], [2,3], etc
- * @param NavMeshEdges		[out] Array of worldspace vertex locations for the edge of the navmesh.  Edges are pairwise verts, i.e. [0,1], [2,3], etc
- * 
- * @todo PolyEdges and NavMeshEdges could probably be Index arrays into MeshVerts and be generated in the master loop instead of separate traversals. 
- */
-void FPImplRecastNavMesh::GetDebugGeometry(FRecastDebugGeometry& OutGeometry, int32 TileIndex) const
+bool FPImplRecastNavMesh::GetDebugGeometryForTile(FRecastDebugGeometry& OutGeometry, int32 TileIndex) const
 {
+	bool bDone = false;
 	if (DetourNavMesh == nullptr || TileIndex >= DetourNavMesh->getMaxTiles())
 	{
-		return;
+		bDone = true;
+		return bDone;
 	}
 				
 	check(NavMeshOwner);
@@ -2791,6 +2787,8 @@ void FPImplRecastNavMesh::GetDebugGeometry(FRecastDebugGeometry& OutGeometry, in
 				}
 			}
 		}
+
+		bDone = true;
 	}
 	else
 	{
@@ -2827,7 +2825,14 @@ void FPImplRecastNavMesh::GetDebugGeometry(FRecastDebugGeometry& OutGeometry, in
 
 			VertBase += GetTilesDebugGeometry(Generator, *Tile, VertBase, OutGeometry, TileIdx, ForbiddenFlags);
 		}
+
+		if (TileIndex == INDEX_NONE)
+		{
+			bDone = true;
+		}
 	}
+
+	return bDone;
 }
 
 int32 FPImplRecastNavMesh::GetTilesDebugGeometry(const FRecastNavMeshGenerator* Generator, const dtMeshTile& Tile, int32 VertBase, FRecastDebugGeometry& OutGeometry, int32 TileIdx, uint16 ForbiddenFlags) const
diff --git a/Engine/Source/Runtime/NavigationSystem/Private/NavMesh/RecastNavMesh.cpp b/Engine/Source/Runtime/NavigationSystem/Private/NavMesh/RecastNavMesh.cpp
--- a/Engine/Source/Runtime/NavigationSystem/Private/NavMesh/RecastNavMesh.cpp
+++ b/Engine/Source/Runtime/NavigationSystem/Private/NavMesh/RecastNavMesh.cpp
@@ -1946,12 +1946,19 @@ bool ARecastNavMesh::GetPolysWithinPathingDistance(FVector const& StartLoc, cons
 	return RecastNavMeshImpl && RecastNavMeshImpl->GetPolysWithinPathingDistance(StartLoc, PathingDistance, GetRightFilterRef(Filter), QueryOwner, FoundPolys, DebugData);
 }
 
+// Deprecated
 void ARecastNavMesh::GetDebugGeometry(FRecastDebugGeometry& OutGeometry, int32 TileIndex) const
+{
+	GetDebugGeometryForTile(OutGeometry, TileIndex);
+}
+
+bool ARecastNavMesh::GetDebugGeometryForTile(FRecastDebugGeometry& OutGeometry, int32 TileIndex) const
 {
 	if (RecastNavMeshImpl)
 	{
-		RecastNavMeshImpl->GetDebugGeometry(OutGeometry, TileIndex);
+		return RecastNavMeshImpl->GetDebugGeometryForTile(OutGeometry, TileIndex);
 	}
+	return true;
 }
 
 void ARecastNavMesh::RequestDrawingUpdate(bool bForce)
diff --git a/Engine/Source/Runtime/NavigationSystem/Public/NavMesh/PImplRecastNavMesh.h b/Engine/Source/Runtime/NavigationSystem/Public/NavMesh/PImplRecastNavMesh.h
--- a/Engine/Source/Runtime/NavigationSystem/Public/NavMesh/PImplRecastNavMesh.h
+++ b/Engine/Source/Runtime/NavigationSystem/Public/NavMesh/PImplRecastNavMesh.h
@@ -51,8 +51,12 @@ public:
 	 */
 	void Serialize(FArchive& Ar, int32 NavMeshVersion);
 
-	/** Debug rendering. */
-	void GetDebugGeometry(FRecastDebugGeometry& OutGeometry, int32 TileIndex = INDEX_NONE) const;
+	/* Gather debug geometry.
+	 * @params OutGeometry Output geometry.
+	 * @params TileIndex Used to collect geometry for a specific tile, INDEX_NONE will gather all tiles.
+	 * @return True if done collecting.
+	 */
+	bool GetDebugGeometryForTile(FRecastDebugGeometry& OutGeometry, int32 TileIndex) const;
 	
 	/** Returns bounding box for the whole navmesh. */
 	FBox GetNavMeshBounds() const;
diff --git a/Engine/Source/Runtime/NavigationSystem/Public/NavMesh/RecastNavMesh.h b/Engine/Source/Runtime/NavigationSystem/Public/NavMesh/RecastNavMesh.h
--- a/Engine/Source/Runtime/NavigationSystem/Public/NavMesh/RecastNavMesh.h
+++ b/Engine/Source/Runtime/NavigationSystem/Public/NavMesh/RecastNavMesh.h
@@ -1010,8 +1010,16 @@ public:
 	// Debug                                                                
 	//----------------------------------------------------------------------//
 	/** Debug rendering support. */
+	UE_DEPRECATED(5.1, "Please use the new signature of GetDebugGeometryForTile()")
 	void GetDebugGeometry(FRecastDebugGeometry& OutGeometry, int32 TileIndex = INDEX_NONE) const;
 
+	/* Gather debug geometry.
+	 * @params OutGeometry Output geometry.
+	 * @params TileIndex Used to collect geometry for a specific tile, INDEX_NONE will gather all tiles
+	 * @return True if done collecting.
+	 */
+	bool GetDebugGeometryForTile(FRecastDebugGeometry& OutGeometry, int32 TileIndex) const;
+
 	// @todo docuement
 	void DrawDebugPathCorridor(NavNodeRef const* PathPolys, int32 NumPathPolys, bool bPersistent=true) const;
 
