[Navmesh] Fix issue with HISM not modifying a dynamic navmesh when there is no instance on component register and instances are added later on.

diff --git a/Engine/Source/Runtime/Engine/Classes/AI/NavigationSystemBase.h b/Engine/Source/Runtime/Engine/Classes/AI/NavigationSystemBase.h
--- a/Engine/Source/Runtime/Engine/Classes/AI/NavigationSystemBase.h
+++ b/Engine/Source/Runtime/Engine/Classes/AI/NavigationSystemBase.h
@@ -140,7 +140,10 @@ namespace FNavigationSystem
 
 	ENGINE_API void OnComponentRegistered(UActorComponent& Comp);
 	ENGINE_API void OnComponentUnregistered(UActorComponent& Comp);
-
+	
+	ENGINE_API void RegisterComponent(UActorComponent& Comp);
+	ENGINE_API void UnregisterComponent(UActorComponent& Comp);
+	
 	ENGINE_API void RemoveActorData(AActor& Actor);
 
 	ENGINE_API bool HasComponentData(UActorComponent& Comp);
@@ -284,6 +287,8 @@ protected:
 	static FNavigationSystem::FActorBasedSignature& OnActorUnregisteredDelegate();
 	static FNavigationSystem::FActorComponentBasedSignature& OnComponentRegisteredDelegate();
 	static FNavigationSystem::FActorComponentBasedSignature& OnComponentUnregisteredDelegate();
+	static FNavigationSystem::FActorComponentBasedSignature& RegisterComponentDelegate();
+	static FNavigationSystem::FActorComponentBasedSignature& UnregisterComponentDelegate();
 	static FNavigationSystem::FActorBasedSignature& RemoveActorDataDelegate();
 	static FNavigationSystem::FBoolActorComponentBasedSignature& HasComponentDataDelegate();
 	static FNavigationSystem::FNavDataConfigBasedSignature& GetDefaultSupportedAgentDelegate();
diff --git a/Engine/Source/Runtime/Engine/Private/AI/NavigationSystemBase.cpp b/Engine/Source/Runtime/Engine/Private/AI/NavigationSystemBase.cpp
--- a/Engine/Source/Runtime/Engine/Private/AI/NavigationSystemBase.cpp
+++ b/Engine/Source/Runtime/Engine/Private/AI/NavigationSystemBase.cpp
@@ -132,6 +132,8 @@ namespace FNavigationSystem
 		FActorBasedSignature OnActorUnregistered;
 		FActorComponentBasedSignature OnComponentRegistered;
 		FActorComponentBasedSignature OnComponentUnregistered;
+		FActorComponentBasedSignature RegisterComponent;
+		FActorComponentBasedSignature UnregisterComponent;
 		FActorBasedSignature RemoveActorData;
 		FControllerBasedSignature StopMovement;
 		FBoolControllerBasedSignature IsFollowingAPath;
@@ -167,6 +169,8 @@ namespace FNavigationSystem
 			OnActorUnregistered.BindLambda([](AActor&) {});
 			OnComponentRegistered.BindLambda([](UActorComponent&) {});
 			OnComponentUnregistered.BindLambda([](UActorComponent&) {});
+			RegisterComponent.BindLambda([](UActorComponent&) {});
+			UnregisterComponent.BindLambda([](UActorComponent&) {});
 			RemoveActorData.BindLambda([](AActor&) {});
 			StopMovement.BindLambda([](const AController&) {});
 			IsFollowingAPath.BindLambda([](const AController&) { return false; });
@@ -206,6 +210,8 @@ namespace FNavigationSystem
 	void OnActorUnregistered(AActor& Actor) { Delegates.OnActorUnregistered.Execute(Actor); }
 	void OnComponentRegistered(UActorComponent& Comp) { Delegates.OnComponentRegistered.Execute(Comp); }
 	void OnComponentUnregistered(UActorComponent& Comp) { Delegates.OnComponentUnregistered.Execute(Comp); }
+	void RegisterComponent(UActorComponent& Comp) { Delegates.RegisterComponent.Execute(Comp); }
+	void UnregisterComponent(UActorComponent& Comp) { Delegates.UnregisterComponent.Execute(Comp); }
 	void RemoveActorData(AActor& Actor) { Delegates.RemoveActorData.Execute(Actor); }
 	bool HasComponentData(UActorComponent& Comp) { return Delegates.HasComponentData.Execute(Comp);	}
 	const FNavDataConfig& GetDefaultSupportedAgent() { return Delegates.GetDefaultSupportedAgent.Execute(); }
@@ -400,6 +406,8 @@ FNavigationSystem::FActorBasedSignature& UNavigationSystemBase::OnActorRegistere
 FNavigationSystem::FActorBasedSignature& UNavigationSystemBase::OnActorUnregisteredDelegate() { return FNavigationSystem::Delegates.OnActorUnregistered; }
 FNavigationSystem::FActorComponentBasedSignature& UNavigationSystemBase::OnComponentRegisteredDelegate() { return FNavigationSystem::Delegates.OnComponentRegistered; }
 FNavigationSystem::FActorComponentBasedSignature& UNavigationSystemBase::OnComponentUnregisteredDelegate() { return FNavigationSystem::Delegates.OnComponentUnregistered; }
+FNavigationSystem::FActorComponentBasedSignature& UNavigationSystemBase::RegisterComponentDelegate() { return FNavigationSystem::Delegates.RegisterComponent; }
+FNavigationSystem::FActorComponentBasedSignature& UNavigationSystemBase::UnregisterComponentDelegate() { return FNavigationSystem::Delegates.UnregisterComponent; }
 FNavigationSystem::FActorBasedSignature& UNavigationSystemBase::RemoveActorDataDelegate() { return FNavigationSystem::Delegates.RemoveActorData; }
 FNavigationSystem::FBoolActorComponentBasedSignature& UNavigationSystemBase::HasComponentDataDelegate() { return FNavigationSystem::Delegates.HasComponentData; }
 FNavigationSystem::FNavDataConfigBasedSignature& UNavigationSystemBase::GetDefaultSupportedAgentDelegate() { return FNavigationSystem::Delegates.GetDefaultSupportedAgent; }
diff --git a/Engine/Source/Runtime/Engine/Private/HierarchicalInstancedStaticMesh.cpp b/Engine/Source/Runtime/Engine/Private/HierarchicalInstancedStaticMesh.cpp
--- a/Engine/Source/Runtime/Engine/Private/HierarchicalInstancedStaticMesh.cpp
+++ b/Engine/Source/Runtime/Engine/Private/HierarchicalInstancedStaticMesh.cpp
@@ -3422,9 +3422,9 @@ void UHierarchicalInstancedStaticMeshComponent::PartialNavigationUpdate(int32 In
 		// Accumulate dirty areas and send them to navigation system once cluster tree is rebuilt
 		if (FNavigationSystem::HasComponentData(*this))
 		{
-			FTransform InstanceTransform(PerInstanceSMData[InstanceIdx].Transform);
-			FBox InstanceBox = GetStaticMesh()->GetBounds().TransformBy(InstanceTransform*GetComponentTransform()).GetBox(); // in world space
-			AccumulatedNavigationDirtyArea+= InstanceBox;
+			const FTransform InstanceTransform(PerInstanceSMData[InstanceIdx].Transform);
+			const FBox InstanceBox = GetStaticMesh()->GetBounds().TransformBy(InstanceTransform*GetComponentTransform()).GetBox(); // in world space
+			AccumulatedNavigationDirtyArea += InstanceBox;
 		}
 	}
 }
diff --git a/Engine/Source/Runtime/Engine/Private/InstancedStaticMesh.cpp b/Engine/Source/Runtime/Engine/Private/InstancedStaticMesh.cpp
--- a/Engine/Source/Runtime/Engine/Private/InstancedStaticMesh.cpp
+++ b/Engine/Source/Runtime/Engine/Private/InstancedStaticMesh.cpp
@@ -3079,6 +3079,13 @@ int32 UInstancedStaticMeshComponent::AddInstanceInternal(int32 InstanceIndex, FI
 	}
 #endif
 
+	// If it's the first instance, register the component. 
+	// If there was no instance on component register, component registration was skipped because of UInstancedStaticMeshComponent::IsNavigationRelevant().
+	if (GetInstanceCount() == 1 && IsNavigationRelevant())
+	{
+		FNavigationSystem::RegisterComponent(*this);
+	}
+	
 	PartialNavigationUpdate(InstanceIndex);
 
 	if (FInstancedStaticMeshDelegates::OnInstanceIndexUpdated.IsBound())
@@ -3131,6 +3138,13 @@ TArray<int32> UInstancedStaticMeshComponent::AddInstancesInternal(const TArray<F
 
 		if (SupportsPartialNavigationUpdate())
 		{
+			// If it's the first instance, register the component. 
+			// If there was no instance on component register, component registration was skipped because of UInstancedStaticMeshComponent::IsNavigationRelevant().
+			if (GetInstanceCount() == 1 && IsNavigationRelevant())
+			{
+				FNavigationSystem::RegisterComponent(*this);
+			}
+			
 			PartialNavigationUpdate(InstanceIndex);
 		}
 
@@ -3215,11 +3229,20 @@ bool UInstancedStaticMeshComponent::RemoveInstanceInternal(int32 InstanceIndex,
 	// remove instance
 	if (!InstanceAlreadyRemoved && PerInstanceSMData.IsValidIndex(InstanceIndex))
 	{
+		const bool bWasNavRelevant = IsNavigationRelevant();
+		
 		// Request navigation update
 		PartialNavigationUpdate(InstanceIndex);
 
 		PerInstanceSMData.RemoveAt(InstanceIndex);
 		PerInstanceSMCustomData.RemoveAt(InstanceIndex * NumCustomDataFloats, NumCustomDataFloats);
+
+		// If it's the last instance, unregister the component since component with no instances are not registered. 
+		// (because of GetInstanceCount() > 0 in UInstancedStaticMeshComponent::IsNavigationRelevant())
+		if (bWasNavRelevant && GetInstanceCount() == 0)
+		{
+			FNavigationSystem::UnregisterComponent(*this);
+		}
 	}
 
 #if WITH_EDITOR
diff --git a/Engine/Source/Runtime/NavigationSystem/Private/NavigationSystem.cpp b/Engine/Source/Runtime/NavigationSystem/Private/NavigationSystem.cpp
--- a/Engine/Source/Runtime/NavigationSystem/Private/NavigationSystem.cpp
+++ b/Engine/Source/Runtime/NavigationSystem/Private/NavigationSystem.cpp
@@ -442,6 +442,8 @@ UNavigationSystemV1::UNavigationSystemV1(const FObjectInitializer& ObjectInitial
 			UNavigationSystemBase::OnActorUnregisteredDelegate().BindLambda([](AActor& Actor) { UNavigationSystemV1::OnActorUnregistered(&Actor); });
 			UNavigationSystemBase::OnComponentRegisteredDelegate().BindLambda([](UActorComponent& Comp) { UNavigationSystemV1::OnComponentRegistered(&Comp); });
 			UNavigationSystemBase::OnComponentUnregisteredDelegate().BindLambda([](UActorComponent& Comp) { UNavigationSystemV1::OnComponentUnregistered(&Comp); });
+			UNavigationSystemBase::RegisterComponentDelegate().BindLambda([](UActorComponent& Comp) { UNavigationSystemV1::RegisterComponent(&Comp); });
+			UNavigationSystemBase::UnregisterComponentDelegate().BindLambda([](UActorComponent& Comp) { UNavigationSystemV1::UnregisterComponent(&Comp); });
 			UNavigationSystemBase::RemoveActorDataDelegate().BindLambda([](AActor& Actor) { UNavigationSystemV1::ClearNavOctreeAll(&Actor); });
 			UNavigationSystemBase::HasComponentDataDelegate().BindLambda([](UActorComponent& Comp) {
 				UNavigationSystemV1* NavSys = FNavigationSystem::GetCurrent<UNavigationSystemV1>(Comp.GetWorld());
@@ -2773,6 +2775,54 @@ ANavigationData* UNavigationSystemV1::GetNavDataWithID(const uint16 NavDataID) c
 	return NULL;
 }
 
+void UNavigationSystemV1::RegisterComponentToNavOctree(UActorComponent* Comp)
+{
+	if ((Comp == nullptr) || IsNavigationSystemStatic())
+	{
+		return;
+	}
+
+	SCOPE_CYCLE_COUNTER(STAT_DebugNavOctree);
+	INavRelevantInterface* NavInterface = Cast<INavRelevantInterface>(Comp);
+	if (NavInterface)
+	{
+		AActor* OwnerActor = Comp->GetOwner();
+		if (OwnerActor && OwnerActor->IsComponentRelevantForNavigation(Comp))
+		{
+			UNavigationSystemV1* NavSys = FNavigationSystem::GetCurrent<UNavigationSystemV1>(OwnerActor->GetWorld());
+			if (NavSys)
+			{
+				NavSys->RegisterNavOctreeElement(Comp, NavInterface, FNavigationOctreeController::OctreeUpdate_Default);
+			}
+		}
+	}
+}
+
+void UNavigationSystemV1::UnregisterComponentToNavOctree(UActorComponent* Comp)
+{
+	if ((Comp == nullptr) || IsNavigationSystemStatic())
+	{
+		return;
+	}
+
+	SCOPE_CYCLE_COUNTER(STAT_DebugNavOctree);
+	INavRelevantInterface* NavInterface = Cast<INavRelevantInterface>(Comp);
+	if (NavInterface)
+	{
+		AActor* OwnerActor = Comp->GetOwner();
+		if (OwnerActor)
+		{
+			// skip IsComponentRelevantForNavigation check, it's only for adding new stuff
+
+			UNavigationSystemV1* NavSys = FNavigationSystem::GetCurrent<UNavigationSystemV1>(OwnerActor->GetWorld());
+			if (NavSys)
+			{
+				NavSys->UnregisterNavOctreeElement(Comp, NavInterface, FNavigationOctreeController::OctreeUpdate_Default);
+			}
+		}
+	}
+}
+
 void UNavigationSystemV1::AddDirtyArea(const FBox& NewArea, int32 Flags)
 {
 	DefaultDirtyAreasController.AddArea(NewArea, Flags);
@@ -3059,50 +3109,22 @@ bool UNavigationSystemV1::ReplaceAreaInOctreeData(const UObject& Object, TSubcla
 
 void UNavigationSystemV1::OnComponentRegistered(UActorComponent* Comp)
 {
-	if ((Comp == nullptr) || IsNavigationSystemStatic())
-	{
-		return;
-	}
-
-	SCOPE_CYCLE_COUNTER(STAT_DebugNavOctree);
-	INavRelevantInterface* NavInterface = Cast<INavRelevantInterface>(Comp);
-	if (NavInterface)
-	{
-		AActor* OwnerActor = Comp->GetOwner();
-		if (OwnerActor && OwnerActor->IsComponentRelevantForNavigation(Comp))
-		{
-			UNavigationSystemV1* NavSys = FNavigationSystem::GetCurrent<UNavigationSystemV1>(OwnerActor->GetWorld());
-			if (NavSys)
-			{
-				NavSys->RegisterNavOctreeElement(Comp, NavInterface, FNavigationOctreeController::OctreeUpdate_Default);
-			}
-		}
-	}
+	RegisterComponentToNavOctree(Comp);
 }
 
 void UNavigationSystemV1::OnComponentUnregistered(UActorComponent* Comp)
 {
-	if ((Comp == nullptr) || IsNavigationSystemStatic())
-	{
-		return;
-	}
+	UnregisterComponentToNavOctree(Comp);
+}
 
-	SCOPE_CYCLE_COUNTER(STAT_DebugNavOctree);
-	INavRelevantInterface* NavInterface = Cast<INavRelevantInterface>(Comp);
-	if (NavInterface)
-	{
-		AActor* OwnerActor = Comp->GetOwner();
-		if (OwnerActor)
-		{
-			// skip IsComponentRelevantForNavigation check, it's only for adding new stuff
+void UNavigationSystemV1::RegisterComponent(UActorComponent* Comp)
+{
+	RegisterComponentToNavOctree(Comp);
+}
 
-			UNavigationSystemV1* NavSys = FNavigationSystem::GetCurrent<UNavigationSystemV1>(OwnerActor->GetWorld());
-			if (NavSys)
-			{
-				NavSys->UnregisterNavOctreeElement(Comp, NavInterface, FNavigationOctreeController::OctreeUpdate_Default);
-			}
-		}
-	}
+void UNavigationSystemV1::UnregisterComponent(UActorComponent* Comp)
+{
+	UnregisterComponentToNavOctree(Comp);
 }
 
 void UNavigationSystemV1::OnActorRegistered(AActor* Actor)
diff --git a/Engine/Source/Runtime/NavigationSystem/Public/NavigationSystem.h b/Engine/Source/Runtime/NavigationSystem/Public/NavigationSystem.h
--- a/Engine/Source/Runtime/NavigationSystem/Public/NavigationSystem.h
+++ b/Engine/Source/Runtime/NavigationSystem/Public/NavigationSystem.h
@@ -721,6 +721,9 @@ protected:
 	/** @return pointer to ANavigationData instance of given ID, or NULL if it was not found. Note it looks only through registered navigation data */
 	ANavigationData* GetNavDataWithID(const uint16 NavDataID) const;
 
+	static void RegisterComponentToNavOctree(UActorComponent* Comp);
+	static void UnregisterComponentToNavOctree(UActorComponent* Comp);
+
 public:
 	virtual void ReleaseInitialBuildingLock();
 
@@ -729,6 +732,8 @@ public:
 	//----------------------------------------------------------------------//
 	static void OnComponentRegistered(UActorComponent* Comp);
 	static void OnComponentUnregistered(UActorComponent* Comp);
+	static void RegisterComponent(UActorComponent* Comp);
+	static void UnregisterComponent(UActorComponent* Comp);
 	static void OnActorRegistered(AActor* Actor);
 	static void OnActorUnregistered(AActor* Actor);
 
