Improved Debug log regarding the tiles being dirtied by a dirty area (LogNavigationDirtyArea) by also showing how many tiles are new in the pending dirty list, instead of just the total, which sometimes is misleading

diff --git a/Engine/Source/Runtime/Engine/Classes/AI/Navigation/NavigationTypes.h b/Engine/Source/Runtime/Engine/Classes/AI/Navigation/NavigationTypes.h
--- a/Engine/Source/Runtime/Engine/Classes/AI/Navigation/NavigationTypes.h
+++ b/Engine/Source/Runtime/Engine/Classes/AI/Navigation/NavigationTypes.h
@@ -106,6 +106,16 @@ struct FNavigationDirtyArea
 	FNavigationDirtyArea() : Flags(0) {}
 	FNavigationDirtyArea(const FBox& InBounds, int32 InFlags, UObject* const InOptionalSourceObject = nullptr) : Bounds(InBounds), Flags(InFlags), OptionalSourceObject(InOptionalSourceObject) {}
 	FORCEINLINE bool HasFlag(ENavigationDirtyFlag::Type Flag) const { return (Flags & Flag) != 0; }
+
+	bool operator==(const FNavigationDirtyArea& Other) const 
+	{ 
+		return Flags == Other.Flags && OptionalSourceObject == Other.OptionalSourceObject && Bounds.Equals(Other.Bounds); 
+	}
+	
+	bool operator!=( const FNavigationDirtyArea& Other) const
+	{
+		return !(*this == Other);
+	}
 };
 
 USTRUCT()
diff --git a/Engine/Source/Runtime/NavigationSystem/Private/NavMesh/RecastNavMeshGenerator.cpp b/Engine/Source/Runtime/NavigationSystem/Private/NavMesh/RecastNavMeshGenerator.cpp
--- a/Engine/Source/Runtime/NavigationSystem/Private/NavMesh/RecastNavMeshGenerator.cpp
+++ b/Engine/Source/Runtime/NavigationSystem/Private/NavMesh/RecastNavMeshGenerator.cpp
@@ -5468,6 +5468,56 @@ void FRecastNavMeshGenerator::AddGeneratedTileLayer(int32 LayerIndex, FRecastTil
 	}
 }
 
+#if !UE_BUILD_SHIPPING
+void FRecastNavMeshGenerator::LogDirtyAreas(
+	const TMap<FPendingTileElement, TArray<FNavigationDirtyAreaPerTileDebugInformation>>& DirtyAreasDebuggingInformation) const
+{
+	// Helper struct used to collate the raw information provided to the method, needed for the log results
+	struct FNavigationDirtyAreaDebugInformation
+	{
+		FNavigationDirtyArea DirtyArea;
+		int32 NewlyAddedDirtyTiles = 0;
+		int32 TotalDirtyTiles = 0;
+	};
+	
+	// Array used to describe per dirty area the amount of dirtied tiles
+	TArray<FNavigationDirtyAreaDebugInformation> DirtyAreaToDirtyTilesCount;
+	for (const TPair<FPendingTileElement, TArray<FNavigationDirtyAreaPerTileDebugInformation>>& DirtyAreasDebuggingInformationPair : DirtyAreasDebuggingInformation)
+	{
+		for (const FNavigationDirtyAreaPerTileDebugInformation& DirtyAreaPerTileDebugInformation : DirtyAreasDebuggingInformationPair.Value)
+		{
+			FNavigationDirtyAreaDebugInformation* DirtyResultsTuple = DirtyAreaToDirtyTilesCount.FindByPredicate([&DirtyAreaPerTileDebugInformation](const FNavigationDirtyAreaDebugInformation& OtherDirtyResultsTuple)
+			{
+				return OtherDirtyResultsTuple.DirtyArea == DirtyAreaPerTileDebugInformation.DirtyArea;
+			});
+
+			if (!DirtyResultsTuple)
+			{
+				DirtyResultsTuple = &DirtyAreaToDirtyTilesCount.Add_GetRef({DirtyAreaPerTileDebugInformation.DirtyArea, 0, 0});
+			}
+
+			if (!DirtyAreaPerTileDebugInformation.bTileWasAlreadyAdded)
+			{
+				DirtyResultsTuple->NewlyAddedDirtyTiles++;
+			}
+			DirtyResultsTuple->TotalDirtyTiles++;
+		}
+	}
+		
+	for (const FNavigationDirtyAreaDebugInformation& DirtyResultsTuple : DirtyAreaToDirtyTilesCount)
+	{
+		const UObject* const SourceObject = DirtyResultsTuple.DirtyArea.OptionalSourceObject.Get();
+		const UActorComponent* const ObjectAsComponent = Cast<UActorComponent>(SourceObject);
+		const AActor* const ComponentOwner = ObjectAsComponent ? ObjectAsComponent->GetOwner() : nullptr;
+		
+		UE_LOG(LogNavigationDirtyArea, VeryVerbose,
+			TEXT("(navmesh: %s) Dirty area trying to dirt %d tiles (out of which %d are newly added/not pending) | Source Object = %s | Potential component's owner = %s | Bounds size = %s)"),
+			*GetNameSafe(GetOwner()), DirtyResultsTuple.TotalDirtyTiles, DirtyResultsTuple.NewlyAddedDirtyTiles, *GetFullNameSafe(SourceObject),
+			*GetFullNameSafe(ComponentOwner), *DirtyResultsTuple.DirtyArea.Bounds.ToString());
+	}
+}
+#endif
+
 ETimeSliceWorkResult FRecastNavMeshGenerator::AddGeneratedTilesTimeSliced(FRecastTileGenerator& TileGenerator, TArray<uint32>& OutResultTileIndices)
 {
 	SCOPE_CYCLE_COUNTER(STAT_Navigation_RecastAddGeneratedTiles);
@@ -5738,6 +5788,10 @@ void FRecastNavMeshGenerator::MarkDirtyTiles(const TArray<FNavigationDirtyArea>&
 		
 	// find all tiles that need regeneration
 	TSet<FPendingTileElement> DirtyTiles;
+#if !UE_BUILD_SHIPPING
+	// Used for debug purposes to track the number of new dirty tiles per area; Updated only if LogNavigationDirtyArea is VeryVerbose
+	TMap<FPendingTileElement, TArray<FNavigationDirtyAreaPerTileDebugInformation>> DirtyAreasDebugging;
+#endif
 	for (const FNavigationDirtyArea& DirtyArea : DirtyAreas)
 	{
 		// Static navmeshes accept only area modifiers updates
@@ -5786,11 +5840,6 @@ void FRecastNavMeshGenerator::MarkDirtyTiles(const TArray<FNavigationDirtyArea>&
 			bDoTileInclusionTest = (FindInclusionBoundEncapsulatingBox(CutDownArea) == INDEX_NONE);
 		}
 		
-#if !UE_BUILD_SHIPPING
-		// Used for debug purposes to track the number of tiles relevant to this area
-		uint32 RelevantTiles = 0;
-#endif
-		
 		const FRcTileBox TileBox(AdjustedAreaBounds, RcNavMeshOrigin, TileSizeInWorldUnits);
 
 		for (int32 TileY = TileBox.YMin; TileY <= TileBox.YMax; ++TileY)
@@ -5842,29 +5891,14 @@ void FRecastNavMeshGenerator::MarkDirtyTiles(const TArray<FNavigationDirtyArea>&
 				}
 				
 #if !UE_BUILD_SHIPPING
-				RelevantTiles++;
+				UE_SUPPRESS(LogNavigationDirtyArea, VeryVerbose, 
+				{
+					const bool bAlreadyAdded = ExistingElement != nullptr;
+					DirtyAreasDebugging.FindOrAdd(Element).Add({DirtyArea, bAlreadyAdded});
+				});
 #endif
 			}
 		}
-		
-#if !UE_BUILD_SHIPPING
-		UE_SUPPRESS(LogNavigationDirtyArea, VeryVerbose, 
-		{
-			if (RelevantTiles > 0)
-			{
-				if (const UObject* const SourceObject = DirtyArea.OptionalSourceObject.Get())
-				{
-					const UActorComponent* const ObjectAsComponent = Cast<UActorComponent>(SourceObject);
-					const AActor* const ComponentOwner = ObjectAsComponent ? ObjectAsComponent->GetOwner() : nullptr;
-					
-					UE_LOG(LogNavigationDirtyArea, VeryVerbose,
-						TEXT("(navmesh: %s) Dirty area trying to dirt %d tiles | Source Object = %s | Potential component's owner = %s | Bounds size = %s)"),
-						*GetNameSafe(GetOwner()), RelevantTiles, *GetFullNameSafe(SourceObject),
-						*GetFullNameSafe(ComponentOwner), *DirtyArea.Bounds.ToString());
-				}
-			}
-		});
-#endif
 	}
 	
 	int32 NumTilesMarked = DirtyTiles.Num();
@@ -5888,9 +5922,27 @@ void FRecastNavMeshGenerator::MarkDirtyTiles(const TArray<FNavigationDirtyArea>&
 				ExistingElement.DirtyAreas.Empty();
 			}
 			DirtyTiles.Remove(Id);
+
+#if !UE_BUILD_SHIPPING
+			UE_SUPPRESS(LogNavigationDirtyArea, VeryVerbose, 
+			{
+				// Flag everything in the map to set a true boolean (since it was already inserted)
+				for (FNavigationDirtyAreaPerTileDebugInformation& Pair : DirtyAreasDebugging.FindChecked(ExistingElement))
+				{
+					Pair.bTileWasAlreadyAdded = true;
+				}
+			});
+#endif
 		}
 	}
 
+#if !UE_BUILD_SHIPPING
+	UE_SUPPRESS(LogNavigationDirtyArea, VeryVerbose, 
+	{
+		LogDirtyAreas(DirtyAreasDebugging);
+	});
+#endif
+
 	// Append remaining new dirty tile elements
 	PendingDirtyTiles.Reserve(PendingDirtyTiles.Num() + DirtyTiles.Num());
 	for(const FPendingTileElement& Element : DirtyTiles)
diff --git a/Engine/Source/Runtime/NavigationSystem/Public/NavMesh/RecastNavMeshGenerator.h b/Engine/Source/Runtime/NavigationSystem/Public/NavMesh/RecastNavMeshGenerator.h
--- a/Engine/Source/Runtime/NavigationSystem/Public/NavMesh/RecastNavMeshGenerator.h
+++ b/Engine/Source/Runtime/NavigationSystem/Public/NavMesh/RecastNavMeshGenerator.h
@@ -810,6 +810,18 @@ protected:
 
 	void AddGeneratedTileLayer(int32 LayerIndex, FRecastTileGenerator& TileGenerator, const TMap<int32, dtPolyRef>& OldLayerTileIdMap, TArray<uint32>& OutResultTileIndices);
 
+#if !UE_BUILD_SHIPPING
+	/** Data struct used by 'LogDirtyAreas' that contains all the information regarding the areas that are being dirtied, per dirtied tile. */
+	struct FNavigationDirtyAreaPerTileDebugInformation
+	{
+		FNavigationDirtyArea DirtyArea;
+		bool bTileWasAlreadyAdded = false;
+	};
+
+	/** Used internally, when LogNavigationDirtyArea is VeryVerbose, to log the number of tiles a dirty area is requesting. */
+	void LogDirtyAreas(const TMap<FPendingTileElement, TArray<FNavigationDirtyAreaPerTileDebugInformation>>& DirtyAreasDebuggingInformation) const; 
+#endif
+	
 protected:
 	friend ARecastNavMesh;
 
