[WP navmesh] Limit tile building to the loaded space when using a WP dynamic navmesh
- Update and use active tiles when using URecastNavMeshDataChunk
- Keep track of dirty areas origin
- Fix some logs TileRef

diff --git a/Engine/Source/Editor/UnrealEd/Private/EditorBuildUtils.cpp b/Engine/Source/Editor/UnrealEd/Private/EditorBuildUtils.cpp
--- a/Engine/Source/Editor/UnrealEd/Private/EditorBuildUtils.cpp
+++ b/Engine/Source/Editor/UnrealEd/Private/EditorBuildUtils.cpp
@@ -902,9 +902,9 @@ void FEditorBuildUtils::TriggerNavigationBuilder(UWorld*& InOutWorld, FName Id)
 			bBuildingNavigationFromUserRequest = false;
 		}
 
-		const FString& LongPackageName = GetNameSafe(InOutWorld->GetPackage());
-		if (UE::EditorBuildUtils::bNavmeshAllowPartitionedBuildingFromEditor && ULevel::GetIsLevelPartitionedFromPackage(*LongPackageName))
+		if (UE::EditorBuildUtils::bNavmeshAllowPartitionedBuildingFromEditor && InOutWorld->IsPartitionedWorld())
 		{
+			const FString& LongPackageName = GetNameSafe(InOutWorld->GetPackage());
 			WorldPartitionBuildNavigation(LongPackageName);
 			InOutWorld = GEditor->GetEditorWorldContext().World();
 		}
diff --git a/Engine/Source/Runtime/NavigationSystem/Private/NavMesh/RecastNavMesh.cpp b/Engine/Source/Runtime/NavigationSystem/Private/NavMesh/RecastNavMesh.cpp
--- a/Engine/Source/Runtime/NavigationSystem/Private/NavMesh/RecastNavMesh.cpp
+++ b/Engine/Source/Runtime/NavigationSystem/Private/NavMesh/RecastNavMesh.cpp
@@ -794,6 +794,13 @@ void ARecastNavMesh::SortAreasForGenerator(TArray<FRecastAreaNavModifierElement>
 	Modifiers.Sort(FNavAreaSortPredicate());
 }
 
+const TArray<FIntPoint>& ARecastNavMesh::GetActiveTiles() const
+{
+	const FRecastNavMeshGenerator* MyGenerator = static_cast<const FRecastNavMeshGenerator*>(GetGenerator());
+	check(MyGenerator);
+	return MyGenerator->ActiveTiles;
+}
+
 TArray<FIntPoint>& ARecastNavMesh::GetActiveTiles()
 {
 	FRecastNavMeshGenerator* MyGenerator = static_cast<FRecastNavMeshGenerator*>(GetGenerator());
@@ -2801,6 +2808,11 @@ FRecastNavMeshGenerator* ARecastNavMesh::CreateGeneratorInstance()
 	return new FRecastNavMeshGenerator(*this);
 }
 
+bool ARecastNavMesh::IsUsingActiveTilesGeneration(const UNavigationSystemV1& NavSys) const
+{
+	return NavSys.IsActiveTilesGenerationEnabled() || bIsWorldPartitioned;
+}
+
 void ARecastNavMesh::ConditionalConstructGenerator()
 {	
 	if (NavDataGenerator.IsValid())
@@ -2824,7 +2836,7 @@ void ARecastNavMesh::ConditionalConstructGenerator()
 		UNavigationSystemV1* NavSys = FNavigationSystem::GetCurrent<UNavigationSystemV1>(World);
 		if (NavSys)
 		{
-			RestrictBuildingToActiveTiles(NavSys->IsActiveTilesGenerationEnabled());
+			RestrictBuildingToActiveTiles(IsUsingActiveTilesGeneration(*NavSys));
 		}
 	}
 }
diff --git a/Engine/Source/Runtime/NavigationSystem/Private/NavMesh/RecastNavMeshDataChunk.cpp b/Engine/Source/Runtime/NavigationSystem/Private/NavMesh/RecastNavMeshDataChunk.cpp
--- a/Engine/Source/Runtime/NavigationSystem/Private/NavMesh/RecastNavMeshDataChunk.cpp
+++ b/Engine/Source/Runtime/NavigationSystem/Private/NavMesh/RecastNavMeshDataChunk.cpp
@@ -64,6 +64,19 @@ static uint8* DuplicateRecastRawData(const uint8* Src, int32 SrcSize)
 	return DupData;
 }
 
+namespace UE::NavMesh::Private
+{
+	bool IsUsingActiveTileGeneration(const FPImplRecastNavMesh& NavMeshImpl)
+	{
+		const UNavigationSystemV1* NavSys = FNavigationSystem::GetCurrent<UNavigationSystemV1>(NavMeshImpl.NavMeshOwner->GetWorld());
+		if (NavSys)
+		{
+			return NavMeshImpl.NavMeshOwner->IsUsingActiveTilesGeneration(*NavSys);
+		}
+		return false;
+	}
+} // namespace UE::NavMesh::Private
+
 //----------------------------------------------------------------------//
 // URecastNavMeshDataChunk                                                                
 //----------------------------------------------------------------------//
@@ -192,6 +205,8 @@ TArray<uint32> URecastNavMeshDataChunk::AttachTiles(FPImplRecastNavMesh& NavMesh
 
 TArray<uint32> URecastNavMeshDataChunk::AttachTiles(FPImplRecastNavMesh& NavMeshImpl, const bool bKeepCopyOfData, const bool bKeepCopyOfCacheData)
 {
+	UE_LOG(LogNavigation, Verbose, TEXT("%s Attaching to NavMesh - %s"), ANSI_TO_TCHAR(__FUNCTION__), *NavigationDataName.ToString());
+	
 	TArray<uint32> Result;
 	Result.Reserve(Tiles.Num());
 
@@ -200,15 +215,43 @@ TArray<uint32> URecastNavMeshDataChunk::AttachTiles(FPImplRecastNavMesh& NavMesh
 
 	if (NavMesh != nullptr)
 	{
+		TArray<FIntPoint>& ActiveTiles = NavMeshImpl.NavMeshOwner->GetActiveTiles();
+		const bool bUsingActiveTiles = UE::NavMesh::Private::IsUsingActiveTileGeneration(NavMeshImpl);
+		
 		for (FRecastTileData& TileData : Tiles)
 		{
 			if (!TileData.bAttached && TileData.TileRawData.IsValid())
 			{
+				if (TileData.TileRawData->RawData == nullptr)
+				{
+					UE_LOG(LogNavigation, Warning, TEXT("Null rawdata. This can be caused by the reuse of unloaded sublevels. 's.ForceGCAfterLevelStreamedOut 1' can be used until this gets fixed."));
+					continue;
+				}
+				
+				const dtMeshHeader* Header = (dtMeshHeader*)TileData.TileRawData->RawData;
+				if (Header->version != DT_NAVMESH_VERSION)
+				{
+					continue;
+				}
+				
+				// If there was a previous tile at the location remove it
+				if (const dtMeshTile* PreExistingTile = NavMesh->getTileAt(Header->x, Header->y, Header->layer))
+				{
+					if (const dtTileRef PreExistingTileRef = NavMesh->getTileRef(PreExistingTile))
+					{
+						UE_LOG(LogNavigation, VeryVerbose, TEXT("%s> Tile (%d,%d:%d), %s removing TileRef: 0x%llx"),
+							*NavMeshImpl.NavMeshOwner->GetName(), Header->x, Header->y, Header->layer, ANSI_TO_TCHAR(__FUNCTION__), PreExistingTileRef);
+						
+						NavMesh->removeTile(PreExistingTileRef, nullptr, nullptr);	
+					}
+				}
+
 				// Attach mesh tile to target nav mesh 
 				dtTileRef TileRef = 0;
 				const dtMeshTile* MeshTile = nullptr;
 
 				dtStatus status = NavMesh->addTile(TileData.TileRawData->RawData, TileData.TileDataSize, DT_TILE_FREE_DATA, 0, &TileRef);
+
 				if (dtStatusFailed(status))
 				{
 					continue;
@@ -224,6 +267,14 @@ TArray<uint32> URecastNavMeshDataChunk::AttachTiles(FPImplRecastNavMesh& NavMesh
 					TileData.bAttached = true;
 				}
 
+				UE_LOG(LogNavigation, VeryVerbose, TEXT("  %s> Tile (%d,%d:%d), %s added TileRef: 0x%llx"),
+					*NavMeshImpl.NavMeshOwner->GetName(), TileData.X, TileData.Y, TileData.Layer, ANSI_TO_TCHAR(__FUNCTION__), TileRef);
+				
+				if (bUsingActiveTiles)
+				{
+					ActiveTiles.AddUnique(FIntPoint(TileData.X, TileData.Y));					
+				}
+				
 				if (bKeepCopyOfData == false)
 				{
 					// We don't own tile data anymore it will be released by recast navmesh 
@@ -281,6 +332,8 @@ TArray<uint32> URecastNavMeshDataChunk::DetachTiles(FPImplRecastNavMesh& NavMesh
 
 TArray<uint32> URecastNavMeshDataChunk::DetachTiles(FPImplRecastNavMesh& NavMeshImpl, const bool bTakeDataOwnership, const bool bTakeCacheDataOwnership)
 {
+	UE_LOG(LogNavigation, Verbose, TEXT("%s Detaching from %s"), ANSI_TO_TCHAR(__FUNCTION__), *NavigationDataName.ToString());
+
 	TArray<uint32> Result;
 	Result.Reserve(Tiles.Num());
 
@@ -289,6 +342,12 @@ TArray<uint32> URecastNavMeshDataChunk::DetachTiles(FPImplRecastNavMesh& NavMesh
 
 	if (NavMesh != nullptr)
 	{
+		TArray<FIntPoint>& ActiveTiles = NavMeshImpl.NavMeshOwner->GetActiveTiles();
+		const bool bUsingActiveTiles = UE::NavMesh::Private::IsUsingActiveTileGeneration(NavMeshImpl);
+
+		TArray<const dtMeshTile*> ExtraMeshTiles;
+		const bool bIsDynamic = NavMeshImpl.NavMeshOwner->SupportsRuntimeGeneration();
+		
 		for (FRecastTileData& TileData : Tiles)
 		{
 			if (TileData.bAttached)
@@ -309,6 +368,9 @@ TArray<uint32> URecastNavMeshDataChunk::DetachTiles(FPImplRecastNavMesh& NavMesh
 							TileData.TileCacheRawData->RawData = TileCacheData.Release();
 						}
 					}
+
+					UE_LOG(LogNavigation, VeryVerbose, TEXT("  %s> Tile (%d,%d:%d), removing TileRef: 0x%llx"),
+						*NavMeshImpl.NavMeshOwner->GetName(), TileData.X, TileData.Y, TileData.Layer, TileRef);
 				
 					NavMeshImpl.RemoveTileCacheLayer(TileData.X, TileData.Y, TileData.Layer);
 
@@ -322,9 +384,36 @@ TArray<uint32> URecastNavMeshDataChunk::DetachTiles(FPImplRecastNavMesh& NavMesh
 						// In the editor we have a copy of tile data so just release tile in navmesh
 						NavMesh->removeTile(TileRef, nullptr, nullptr);
 					}
+
+					if (bUsingActiveTiles)
+					{
+						ActiveTiles.RemoveSwap(FIntPoint(TileData.X, TileData.Y));	
+					}
 						
 					Result.Add(NavMesh->decodePolyIdTile(TileRef));
 				}
+
+				if (bIsDynamic)
+				{
+					// Remove any tile remaining
+					const int32 MaxTiles = NavMesh->getTileCountAt(TileData.X, TileData.Y);
+					if (MaxTiles > 0)
+					{
+						ExtraMeshTiles.SetNumZeroed(MaxTiles, false);
+						const int32 MeshTilesCount = NavMesh->getTilesAt(TileData.X, TileData.Y, ExtraMeshTiles.GetData(), MaxTiles);
+						for (int32 i = 0; i < MeshTilesCount; ++i)
+						{
+							const dtMeshTile* ExtraMeshTile = ExtraMeshTiles[i];
+							dtTileRef ExtraTileRef = NavMesh->getTileRef(ExtraMeshTile);
+							if (ExtraTileRef)
+							{
+								NavMesh->removeTile(ExtraTileRef, nullptr, nullptr);
+								Result.Add(NavMesh->decodePolyIdTile(ExtraTileRef));
+							}
+						}
+					}
+				}
+				
 			}
 
 			TileData.bAttached = false;
@@ -397,12 +486,6 @@ void URecastNavMeshDataChunk::ReleaseTiles()
 	Tiles.Reset();
 }
 
-void URecastNavMeshDataChunk::GatherTiles(const FPImplRecastNavMesh* NavMeshImpl, const TArray<int32>& TileIndices)
-{
-	const EGatherTilesCopyMode CopyMode = NavMeshImpl->NavMeshOwner->SupportsRuntimeGeneration() ? EGatherTilesCopyMode::CopyDataAndCacheData : EGatherTilesCopyMode::CopyData;
-	GetTiles(NavMeshImpl, TileIndices, CopyMode);
-}
-
 void URecastNavMeshDataChunk::GetTiles(const FPImplRecastNavMesh* NavMeshImpl, const TArray<int32>& TileIndices, const EGatherTilesCopyMode CopyMode, const bool bMarkAsAttached /*= true*/)
 {
 	Tiles.Empty(TileIndices.Num());
diff --git a/Engine/Source/Runtime/NavigationSystem/Private/NavMesh/RecastNavMeshGenerator.cpp b/Engine/Source/Runtime/NavigationSystem/Private/NavMesh/RecastNavMeshGenerator.cpp
--- a/Engine/Source/Runtime/NavigationSystem/Private/NavMesh/RecastNavMeshGenerator.cpp
+++ b/Engine/Source/Runtime/NavigationSystem/Private/NavMesh/RecastNavMeshGenerator.cpp
@@ -113,6 +113,15 @@ int32 GetTilesCountHelper(const dtNavMesh* DetourMesh)
 	return NumTiles;
 }
 
+namespace UE::NavMesh::Private
+{
+	// @todo: Promote to the public interface when this becomes less experimental.
+	bool IsWorldPartitionedDynamicNavmesh(const ARecastNavMesh& NavMesh)
+	{
+		return NavMesh.bIsWorldPartitioned && NavMesh.SupportsRuntimeGeneration();
+	}
+} // namespace UE::NavMesh::Private
+
 /**
  * Exports geometry to OBJ file. Can be used to verify NavMesh generation in RecastDemo app
  * @param FileName - full name of OBJ file with extension
@@ -5335,7 +5344,7 @@ TArray<uint32> FRecastNavMeshGenerator::RemoveTileLayers(const int32 TileX, cons
 				dtPolyRef TileRef = DetourMesh->getTileRef(Tiles[i]);
 
 				NumActiveTiles--;
-				UE_LOG(LogNavigation, Verbose, TEXT("%s> Tile (%d,%d:%d), removing TileRef: 0x%X (active:%d)"),
+				UE_LOG(LogNavigation, VeryVerbose, TEXT("%s> Tile (%d,%d:%d), removing TileRef: 0x%llx (active count: %d)"),
 					*DestNavMesh->GetName(), TileX, TileY, LayerIndex, TileRef, NumActiveTiles);
 
 				DetourMesh->removeTile(TileRef, nullptr, nullptr);
@@ -5404,7 +5413,7 @@ void FRecastNavMeshGenerator::AddGeneratedTileLayer(int32 LayerIndex, FRecastTil
 		if (OldTileRef)
 		{
 			NumActiveTiles--;
-			UE_LOG(LogNavigation, Verbose, TEXT("%s> Tile (%d,%d:%d), removing TileRef: 0x%X (active:%d)"),
+			UE_LOG(LogNavigation, VeryVerbose, TEXT("%s> Tile (%d,%d:%d), removing TileRef: 0x%llx (active count: %d)"),
 				*DestNavMesh->GetName(), TileX, TileY, LayerIndex, OldTileRef, NumActiveTiles);
 
 			DetourMesh->removeTile(OldTileRef, nullptr, nullptr);
@@ -5449,7 +5458,7 @@ void FRecastNavMeshGenerator::AddGeneratedTileLayer(int32 LayerIndex, FRecastTil
 				OutResultTileIndices.AddUnique(DetourMesh->decodePolyIdTile(ResultTileRef));
 				NumActiveTiles++;
 
-				UE_LOG(LogNavigation, Verbose, TEXT("%s> Tile (%d,%d:%d), added TileRef: 0x%X (active:%d)"),
+				UE_LOG(LogNavigation, VeryVerbose, TEXT("%s> Tile (%d,%d:%d), added TileRef: 0x%llx (active count: %d)"),
 					*DestNavMesh->GetName(), TileX, TileY, LayerIndex, ResultTileRef, NumActiveTiles);
 
 				{
@@ -5551,10 +5560,7 @@ ETimeSliceWorkResult FRecastNavMeshGenerator::AddGeneratedTilesTimeSliced(FRecas
 	case EAddGeneratedTilesTimeSlicedState::AddTiles:
 	{
 		if (DetourMesh != nullptr
-			// no longer testing this here, we can live with a stray unwanted tile here 
-			// and there. It will be removed the next time around the invokers get
-			// updated 
-			// && IsInActiveSet(FIntPoint(TileX, TileY))
+			&& (UE::NavMesh::Private::IsWorldPartitionedDynamicNavmesh(*DestNavMesh) && IsInActiveSet(FIntPoint(TileX, TileY)))
 			&& SyncTimeSlicedData.AddGenTilesLayerIndex != INDEX_NONE)
 		{
 			for (; SyncTimeSlicedData.AddGenTilesLayerIndex < TileGenerator.GetDirtyLayersMask().Num(); ++SyncTimeSlicedData.AddGenTilesLayerIndex)
@@ -5618,10 +5624,7 @@ TArray<uint32> FRecastNavMeshGenerator::AddGeneratedTiles(FRecastTileGenerator&
 	const int32 FirstDirtyTileIndex = TileGenerator.GetDirtyLayersMask().Find(true);
 
 	if (DetourMesh != nullptr
-		// no longer testing this here, we can live with a stray unwanted tile here 
-		// and there. It will be removed the next time around the invokers get
-		// updated 
-		// && IsInActiveSet(FIntPoint(TileX, TileY))
+		&& (UE::NavMesh::Private::IsWorldPartitionedDynamicNavmesh(*DestNavMesh) && IsInActiveSet(FIntPoint(TileX, TileY)))
 		&& FirstDirtyTileIndex != INDEX_NONE)
 	{
 		TArray<FNavMeshTileData> TileLayers = TileGenerator.GetNavigationData();
@@ -5799,6 +5802,8 @@ void FRecastNavMeshGenerator::MarkDirtyTiles(const TArray<FNavigationDirtyArea>&
 		{
 			continue;
 		}
+
+		UE_VLOG_BOX(OwnerNav, LogNavigation, VeryVerbose, DirtyArea.Bounds, FColor::Blue, TEXT("DirtyArea"));
 		
 		// (if bUseVirtualGeometryFilteringAndDirtying is true) Ignore dirty areas flagged by a source object that is not supposed to apply to this navmesh
 		if (bUseVirtualGeometryFilteringAndDirtying && NavSys && NavOctreeInstance && NavDataConfig)
@@ -5848,6 +5853,11 @@ void FRecastNavMeshGenerator::MarkDirtyTiles(const TArray<FNavigationDirtyArea>&
 			{
 				if (IsInActiveSet(FIntPoint(TileX, TileY)) == false)
 				{
+					UE_SUPPRESS(LogNavigation, VeryVerbose,
+					{
+						const FBox TileBounds = CalculateTileBounds(TileX, TileY, RcNavMeshOrigin, TotalNavBounds, TileSizeInWorldUnits);
+						UE_VLOG_BOX(OwnerNav, LogNavigation, VeryVerbose, TileBounds, FColor::Red, TEXT("Not in active set"));
+					});
 					continue;
 				}
 
diff --git a/Engine/Source/Runtime/NavigationSystem/Private/NavigationDataHandler.cpp b/Engine/Source/Runtime/NavigationSystem/Private/NavigationDataHandler.cpp
--- a/Engine/Source/Runtime/NavigationSystem/Private/NavigationDataHandler.cpp
+++ b/Engine/Source/Runtime/NavigationSystem/Private/NavigationDataHandler.cpp
@@ -27,7 +27,7 @@ void FNavigationDataHandler::RemoveNavOctreeElementId(const FOctreeElementId2& E
 		const FNavigationOctreeElement& ElementData = OctreeController.NavOctree->GetElementById(ElementId);
 		const int32 DirtyFlag = GetDirtyFlagHelper(UpdateFlags, ElementData.Data->GetDirtyFlag());
 		// mark area occupied by given actor as dirty
-		DirtyAreasController.AddArea(ElementData.Bounds.GetBox(), DirtyFlag, [&ElementData] { return ElementData.Data->SourceObject.Get(); });
+		DirtyAreasController.AddArea(ElementData.Bounds.GetBox(), DirtyFlag, [&ElementData] { return ElementData.Data->SourceObject.Get(); }, nullptr, "Remove from navoctree");
 		OctreeController.NavOctree->RemoveNode(ElementId);
 	}
 }
@@ -96,7 +96,7 @@ void FNavigationDataHandler::AddElementToNavOctree(const FNavigationDirtyElement
 	{
 		if (DirtyElement.bHasPrevData)
 		{
-			DirtyAreasController.AddArea(DirtyElement.PrevBounds, DirtyElement.PrevFlags, [&DirtyElement] { return DirtyElement.Owner.Get(); }, &DirtyElement);
+			DirtyAreasController.AddArea(DirtyElement.PrevBounds, DirtyElement.PrevFlags, [&DirtyElement] { return DirtyElement.Owner.Get(); }, &DirtyElement, "Addition to navoctree (invalid request)");
 		}
 
 		return;
@@ -146,7 +146,7 @@ void FNavigationDataHandler::AddElementToNavOctree(const FNavigationDirtyElement
 	if (!GeneratedData.IsEmpty())
 	{
 		const int32 DirtyFlag = DirtyElement.FlagsOverride ? DirtyElement.FlagsOverride : GeneratedData.Data->GetDirtyFlag();
-		DirtyAreasController.AddArea(GeneratedData.Bounds.GetBox(), DirtyFlag, [&ElementOwner] { return ElementOwner; }, &DirtyElement);
+		DirtyAreasController.AddArea(GeneratedData.Bounds.GetBox(), DirtyFlag, [&ElementOwner] { return ElementOwner; }, &DirtyElement, "Addition to navoctree");
 	}
 }
 
@@ -342,7 +342,7 @@ bool FNavigationDataHandler::UpdateNavOctreeElementBounds(UActorComponent& Comp,
 			if (ElementId != nullptr && ensure(OctreeController.IsValidElement(*ElementId)))
 			{
 				const FNavigationOctreeElement& ElementData = OctreeController.NavOctree->GetElementById(*ElementId);
-				DirtyAreasController.AddArea(DirtyArea, ElementData.Data->GetDirtyFlag(), [&Comp] { return &Comp; });
+				DirtyAreasController.AddArea(DirtyArea, ElementData.Data->GetDirtyFlag(), [&Comp] { return &Comp; }, nullptr, "Bounds change");
 			}
 		}
 
@@ -431,7 +431,7 @@ void FNavigationDataHandler::AddLevelCollisionToOctree(ULevel& Level)
 			if (!Bounds.GetExtent().IsNearlyZero())
 			{
 				OctreeController.NavOctree->AddNode(&Level, nullptr, Bounds, BSPElem);
-				DirtyAreasController.AddArea(Bounds, ENavigationDirtyFlag::All, [&Level] { return &Level; });
+				DirtyAreasController.AddArea(Bounds, ENavigationDirtyFlag::All, [&Level] { return &Level; }, nullptr, "Add level");
 
 				UE_LOG(LogNavOctree, Log, TEXT("ADD %s"), *Level.GetName());
 			}
@@ -453,7 +453,7 @@ void FNavigationDataHandler::RemoveLevelCollisionFromOctree(ULevel& Level)
 			{
 				// mark area occupied by given actor as dirty
 				const FNavigationOctreeElement& ElementData = OctreeController.NavOctree->GetElementById(*ElementId);
-				DirtyAreasController.AddArea(ElementData.Bounds.GetBox(), ENavigationDirtyFlag::All, [&Level] { return &Level; });
+				DirtyAreasController.AddArea(ElementData.Bounds.GetBox(), ENavigationDirtyFlag::All, [&Level] { return &Level; }, nullptr, "Remove level");
 			}
 
 			OctreeController.NavOctree->RemoveNode(*ElementId);
diff --git a/Engine/Source/Runtime/NavigationSystem/Private/NavigationDirtyAreasController.cpp b/Engine/Source/Runtime/NavigationSystem/Private/NavigationDirtyAreasController.cpp
--- a/Engine/Source/Runtime/NavigationSystem/Private/NavigationDirtyAreasController.cpp
+++ b/Engine/Source/Runtime/NavigationSystem/Private/NavigationDirtyAreasController.cpp
@@ -48,7 +48,8 @@ void FNavigationDirtyAreasController::Tick(const float DeltaSeconds, const TArra
 	}
 }
 
-void FNavigationDirtyAreasController::AddArea(const FBox& NewArea, const int32 Flags, const TFunction<UObject*()>& ObjectProviderFunc /*= nullptr*/, const FNavigationDirtyElement* DirtyElement /*= nullptr*/)
+void FNavigationDirtyAreasController::AddArea(const FBox& NewArea, const int32 Flags, const TFunction<UObject*()>& ObjectProviderFunc /*= nullptr*/,
+	const FNavigationDirtyElement* DirtyElement /*= nullptr*/, const FName& DebugReason /*= NAME_None*/)
 {
 #if !UE_BUILD_SHIPPING
 	// always keep track of reported areas even when filtered out by invalid area as long as flags are valid
@@ -57,14 +58,16 @@ void FNavigationDirtyAreasController::AddArea(const FBox& NewArea, const int32 F
 
 	if (!NewArea.IsValid)
 	{
-		UE_LOG(LogNavigationDirtyArea, Warning, TEXT("Skipping dirty area creation because of invalid bounds (object: %s)"), *GetFullNameSafe(ObjectProviderFunc ? ObjectProviderFunc() : nullptr));
+		UE_LOG(LogNavigationDirtyArea, Warning, TEXT("Skipping dirty area creation because of invalid bounds (object: %s, from: %s)"),
+			*GetFullNameSafe(ObjectProviderFunc ? ObjectProviderFunc() : nullptr), *DebugReason.ToString());
 		return;
 	}
 
 	const FVector2D BoundsSize(NewArea.GetSize());
 	if (BoundsSize.IsNearlyZero())
 	{
-		UE_LOG(LogNavigationDirtyArea, Warning, TEXT("Skipping dirty area creation because of empty bounds (object: %s)"), *GetFullNameSafe(ObjectProviderFunc ? ObjectProviderFunc() : nullptr));
+		UE_LOG(LogNavigationDirtyArea, Warning, TEXT("Skipping dirty area creation because of empty bounds (object: %s, from: %s)"),
+			*GetFullNameSafe(ObjectProviderFunc ? ObjectProviderFunc() : nullptr), *DebugReason.ToString());
 		return;
 	}
 
@@ -76,14 +79,15 @@ void FNavigationDirtyAreasController::AddArea(const FBox& NewArea, const int32 F
 			// If the area is from the addition or removal of objects caused by level loading/unloading and it's already in the base navmesh ignore the dirtiness.
 			if (const bool bIsIncludedInBaseNavmesh = (DirtyElement && DirtyElement->bIsInBaseNavmesh) || FNavigationSystem::IsInBaseNavmesh(ObjectProviderFunc()))
 			{
-				UE_LOG(LogNavigationDirtyArea, VeryVerbose, TEXT("Ignoring dirtyness (visibility changed and in base navmesh). (object: %s)"), *GetFullNameSafe(ObjectProviderFunc ? ObjectProviderFunc() : nullptr));
+				UE_LOG(LogNavigationDirtyArea, VeryVerbose, TEXT("Ignoring dirtyness (visibility changed and in base navmesh). (object: %s from: %s)"),
+					*GetFullNameSafe(ObjectProviderFunc ? ObjectProviderFunc() : nullptr), *DebugReason.ToString());
 				return;
 			}
 		}		
 	}
 
 #if !UE_BUILD_SHIPPING
-	auto DumpExtraInfo = [ObjectProviderFunc, BoundsSize, NewArea]() {
+	auto DumpExtraInfo = [ObjectProviderFunc, DebugReason, BoundsSize, NewArea]() {
 		const UObject* ObjectProvider = nullptr;
 		if (ObjectProviderFunc)
 		{
@@ -96,7 +100,11 @@ void FNavigationDirtyAreasController::AddArea(const FBox& NewArea, const int32 F
 		{
 			UE_VLOG_BOX(ComponentOwner, LogNavigationDirtyArea, Log, NewArea, FColor::Red, TEXT(""));
 		}
-		return FString::Printf(TEXT("Adding dirty area object: %s | Potential component's owner: %s | Bounds size: %s)"), *GetFullNameSafe(ObjectProvider), *GetFullNameSafe(ComponentOwner), *BoundsSize.ToString());
+		return FString::Printf(TEXT("Adding dirty area object: %s (from: %s) | Potential component's owner: %s | Bounds size: %s"),
+			*GetFullNameSafe(ObjectProvider),
+			*DebugReason.ToString(),
+			*GetFullNameSafe(ComponentOwner),
+			*BoundsSize.ToString());
 	};
 
 	UE_LOG(LogNavigationDirtyArea, VeryVerbose, TEXT("%s"), *DumpExtraInfo());
diff --git a/Engine/Source/Runtime/NavigationSystem/Private/NavigationSystem.cpp b/Engine/Source/Runtime/NavigationSystem/Private/NavigationSystem.cpp
--- a/Engine/Source/Runtime/NavigationSystem/Private/NavigationSystem.cpp
+++ b/Engine/Source/Runtime/NavigationSystem/Private/NavigationSystem.cpp
@@ -2823,12 +2823,12 @@ void UNavigationSystemV1::UnregisterComponentToNavOctree(UActorComponent* Comp)
 	}
 }
 
-void UNavigationSystemV1::AddDirtyArea(const FBox& NewArea, int32 Flags)
+void UNavigationSystemV1::AddDirtyArea(const FBox& NewArea, int32 Flags, const char* SourceText)
 {
-	DefaultDirtyAreasController.AddArea(NewArea, Flags);
+	DefaultDirtyAreasController.AddArea(NewArea, Flags, nullptr, nullptr, SourceText);
 }
 
-void UNavigationSystemV1::AddDirtyAreas(const TArray<FBox>& NewAreas, int32 Flags)
+void UNavigationSystemV1::AddDirtyAreas(const TArray<FBox>& NewAreas, int32 Flags, const char* SourceText)
 { 
 	if (Flags == 0)
 	{
@@ -2837,7 +2837,7 @@ void UNavigationSystemV1::AddDirtyAreas(const TArray<FBox>& NewAreas, int32 Flag
 
 	for (int32 NewAreaIndex = 0; NewAreaIndex < NewAreas.Num(); NewAreaIndex++)
 	{
-		AddDirtyArea(NewAreas[NewAreaIndex], Flags);
+		AddDirtyArea(NewAreas[NewAreaIndex], Flags, SourceText);
 	}
 }
 
@@ -3367,7 +3367,7 @@ void UNavigationSystemV1::PerformNavigationBoundsUpdate(const TArray<FNavigation
 		}
 
 		// Propagate to generators areas that needs to be updated
-		AddDirtyAreas(UpdatedAreas, ENavigationDirtyFlag::All | ENavigationDirtyFlag::NavigationBounds);
+		AddDirtyAreas(UpdatedAreas, ENavigationDirtyFlag::All | ENavigationDirtyFlag::NavigationBounds, "Navigation bounds update");
 	}
 
 	// I'm not sure why we even do the following as part of this function
@@ -4054,7 +4054,7 @@ void UNavigationSystemV1::OnActorMoved(AActor* Actor)
 
 void UNavigationSystemV1::OnNavigationDirtied(const FBox& Bounds)
 {
-	AddDirtyArea(Bounds, ENavigationDirtyFlag::All);
+	AddDirtyArea(Bounds, ENavigationDirtyFlag::All, "OnNavigationDirtied");
 }
 
 void UNavigationSystemV1::OnReloadComplete(EReloadCompleteReason Reason)
diff --git a/Engine/Source/Runtime/NavigationSystem/Public/NavMesh/RecastNavMesh.h b/Engine/Source/Runtime/NavigationSystem/Public/NavMesh/RecastNavMesh.h
--- a/Engine/Source/Runtime/NavigationSystem/Public/NavMesh/RecastNavMesh.h
+++ b/Engine/Source/Runtime/NavigationSystem/Public/NavMesh/RecastNavMesh.h
@@ -49,6 +49,7 @@ class dtNavMesh;
 class dtQueryFilter;
 class FRecastNavMeshGenerator;
 struct dtMeshTile;
+class UNavigationSystemV1;
 
 UENUM()
 namespace ERecastPartitioning
@@ -925,11 +926,13 @@ public:
 	virtual void AttachNavMeshDataChunk(URecastNavMeshDataChunk& NavDataChunk);
 	virtual void DetachNavMeshDataChunk(URecastNavMeshDataChunk& NavDataChunk);
 
+	const TArray<FIntPoint>& GetActiveTiles() const;
+	TArray<FIntPoint>& GetActiveTiles(); 
+	
 protected:
 	/** Serialization helper. */
 	void SerializeRecastNavMesh(FArchive& Ar, FPImplRecastNavMesh*& NavMesh, int32 NavMeshVersion);
 
-	TArray<FIntPoint>& GetActiveTiles(); 
 	virtual void RestrictBuildingToActiveTiles(bool InRestrictBuildingToActiveTiles) override;
 
 	virtual void OnRegistered() override;
@@ -1222,6 +1225,11 @@ public:
 	virtual bool NeedsRebuild() const override;
 	virtual bool SupportsRuntimeGeneration() const override;
 	virtual bool SupportsStreaming() const override;
+
+	/** When using active tiles generation, navigation is only allowed to be runtime generated on a subset of tiles.
+	 *  The subset is be defined by navinvokers or loaded world partitioned cells. */
+	bool IsUsingActiveTilesGeneration(const UNavigationSystemV1& NavSys) const;
+
 	virtual void ConditionalConstructGenerator() override;
 	void UpdateGenerationProperties(const FRecastNavMeshGenerationProperties& GenerationProps);
 	bool ShouldGatherDataOnGameThread() const { return bDoFullyAsyncNavDataGathering == false; }
diff --git a/Engine/Source/Runtime/NavigationSystem/Public/NavMesh/RecastNavMeshDataChunk.h b/Engine/Source/Runtime/NavigationSystem/Public/NavMesh/RecastNavMeshDataChunk.h
--- a/Engine/Source/Runtime/NavigationSystem/Public/NavMesh/RecastNavMeshDataChunk.h
+++ b/Engine/Source/Runtime/NavigationSystem/Public/NavMesh/RecastNavMeshDataChunk.h
@@ -101,9 +101,6 @@ class NAVIGATIONSYSTEM_API URecastNavMeshDataChunk : public UNavigationDataChunk
 	/** Releases all tiles that this chunk holds */
 	void ReleaseTiles();
 
-	UE_DEPRECATED(4.26, "Use GetTiles() instead.")
-	void GatherTiles(const FPImplRecastNavMesh* NavMeshImpl, const TArray<int32>& TileIndices);
-
 	/** Collect tiles with data and/or cache data from the provided TileIndices. */
 	void GetTiles(const FPImplRecastNavMesh* NavMeshImpl, const TArray<int32>& TileIndices, const EGatherTilesCopyMode CopyMode, const bool bMarkAsAttached = true);
 
diff --git a/Engine/Source/Runtime/NavigationSystem/Public/NavigationDirtyAreasController.h b/Engine/Source/Runtime/NavigationSystem/Public/NavigationDirtyAreasController.h
--- a/Engine/Source/Runtime/NavigationSystem/Public/NavigationDirtyAreasController.h
+++ b/Engine/Source/Runtime/NavigationSystem/Public/NavigationDirtyAreasController.h
@@ -53,8 +53,10 @@ public:
 	 *	@param Flags Indicates the type of modification applied to the area
 	 *	@param ObjectProviderFunc Optional function to retrieve source object that can be use for error reporting and navmesh exclusion
 	 *	@param DirtyElement Optional dirty element
+	 *	@param DebugReason Source of the new area
 	 */
-	void AddArea(const FBox& NewArea, const int32 Flags, const TFunction<UObject*()>& ObjectProviderFunc = nullptr, const FNavigationDirtyElement* DirtyElement = nullptr);
+	void AddArea(const FBox& NewArea, const int32 Flags, const TFunction<UObject*()>& ObjectProviderFunc = nullptr,
+		const FNavigationDirtyElement* DirtyElement = nullptr, const FName& DebugReason = NAME_None);
 	
 	bool IsDirty() const { return GetNumDirtyAreas() > 0; }
 	int32 GetNumDirtyAreas() const { return DirtyAreas.Num(); }
diff --git a/Engine/Source/Runtime/NavigationSystem/Public/NavigationSystem.h b/Engine/Source/Runtime/NavigationSystem/Public/NavigationSystem.h
--- a/Engine/Source/Runtime/NavigationSystem/Public/NavigationSystem.h
+++ b/Engine/Source/Runtime/NavigationSystem/Public/NavigationSystem.h
@@ -763,8 +763,8 @@ public:
 	/** updates bounds of all components implementing INavRelevantInterface */
 	static void UpdateNavOctreeBounds(AActor* Actor);
 
-	void AddDirtyArea(const FBox& NewArea, int32 Flags);
-	void AddDirtyAreas(const TArray<FBox>& NewAreas, int32 Flags);
+	void AddDirtyArea(const FBox& NewArea, int32 Flags, const char* SourceText = nullptr);
+	void AddDirtyAreas(const TArray<FBox>& NewAreas, int32 Flags, const char* SourceText = nullptr);
 	bool HasDirtyAreasQueued() const;
 	int32 GetNumDirtyAreas() const;
 
