[Navmesh] Fix check triggering in ARecastNavMesh::GetActiveTiles
#rnx

diff --git a/Engine/Source/Runtime/NavigationSystem/Private/NavMesh/RecastNavMeshDataChunk.cpp b/Engine/Source/Runtime/NavigationSystem/Private/NavMesh/RecastNavMeshDataChunk.cpp
--- a/Engine/Source/Runtime/NavigationSystem/Private/NavMesh/RecastNavMeshDataChunk.cpp
+++ b/Engine/Source/Runtime/NavigationSystem/Private/NavMesh/RecastNavMeshDataChunk.cpp
@@ -215,8 +215,13 @@ TArray<uint32> URecastNavMeshDataChunk::AttachTiles(FPImplRecastNavMesh& NavMesh
 
 	if (NavMesh != nullptr)
 	{
-		TArray<FIntPoint>& ActiveTiles = NavMeshImpl.NavMeshOwner->GetActiveTiles();
-		const bool bUsingActiveTiles = UE::NavMesh::Private::IsUsingActiveTileGeneration(NavMeshImpl);
+		check(NavMeshImpl.NavMeshOwner);
+		
+		TArray<FIntPoint>* ActiveTiles = nullptr;
+		if (UE::NavMesh::Private::IsUsingActiveTileGeneration(NavMeshImpl))
+		{
+			ActiveTiles = &NavMeshImpl.NavMeshOwner->GetActiveTiles();
+		}
 		
 		for (FRecastTileData& TileData : Tiles)
 		{
@@ -270,9 +275,9 @@ TArray<uint32> URecastNavMeshDataChunk::AttachTiles(FPImplRecastNavMesh& NavMesh
 				UE_LOG(LogNavigation, VeryVerbose, TEXT("  %s> Tile (%d,%d:%d), %s added TileRef: 0x%llx"),
 					*NavMeshImpl.NavMeshOwner->GetName(), TileData.X, TileData.Y, TileData.Layer, ANSI_TO_TCHAR(__FUNCTION__), TileRef);
 				
-				if (bUsingActiveTiles)
+				if (ActiveTiles)
 				{
-					ActiveTiles.AddUnique(FIntPoint(TileData.X, TileData.Y));					
+					ActiveTiles->AddUnique(FIntPoint(TileData.X, TileData.Y));					
 				}
 				
 				if (bKeepCopyOfData == false)
@@ -342,8 +347,13 @@ TArray<uint32> URecastNavMeshDataChunk::DetachTiles(FPImplRecastNavMesh& NavMesh
 
 	if (NavMesh != nullptr)
 	{
-		TArray<FIntPoint>& ActiveTiles = NavMeshImpl.NavMeshOwner->GetActiveTiles();
-		const bool bUsingActiveTiles = UE::NavMesh::Private::IsUsingActiveTileGeneration(NavMeshImpl);
+		check(NavMeshImpl.NavMeshOwner);
+		
+		TArray<FIntPoint>* ActiveTiles = nullptr;
+		if (UE::NavMesh::Private::IsUsingActiveTileGeneration(NavMeshImpl))
+		{
+			ActiveTiles = &NavMeshImpl.NavMeshOwner->GetActiveTiles();
+		}
 
 		TArray<const dtMeshTile*> ExtraMeshTiles;
 		const bool bIsDynamic = NavMeshImpl.NavMeshOwner->SupportsRuntimeGeneration();
@@ -385,9 +395,9 @@ TArray<uint32> URecastNavMeshDataChunk::DetachTiles(FPImplRecastNavMesh& NavMesh
 						NavMesh->removeTile(TileRef, nullptr, nullptr);
 					}
 
-					if (bUsingActiveTiles)
+					if (ActiveTiles)
 					{
-						ActiveTiles.RemoveSwap(FIntPoint(TileData.X, TileData.Y));	
+						ActiveTiles->RemoveSwap(FIntPoint(TileData.X, TileData.Y));	
 					}
 						
 					Result.Add(NavMesh->decodePolyIdTile(TileRef));
