Fix for nav mesh not being present in various levels.

Issue was associated with IsWorldPartitionedDynamicNavmesh change.

diff --git a/Engine/Source/Runtime/NavigationSystem/Private/NavMesh/RecastNavMeshGenerator.cpp b/Engine/Source/Runtime/NavigationSystem/Private/NavMesh/RecastNavMeshGenerator.cpp
--- a/Engine/Source/Runtime/NavigationSystem/Private/NavMesh/RecastNavMeshGenerator.cpp
+++ b/Engine/Source/Runtime/NavigationSystem/Private/NavMesh/RecastNavMeshGenerator.cpp
@@ -5560,7 +5560,7 @@
 	case EAddGeneratedTilesTimeSlicedState::AddTiles:
 	{
 		if (DetourMesh != nullptr
-			&& (UE::NavMesh::Private::IsWorldPartitionedDynamicNavmesh(*DestNavMesh) && IsInActiveSet(FIntPoint(TileX, TileY)))
+			&& (!UE::NavMesh::Private::IsWorldPartitionedDynamicNavmesh(*DestNavMesh) || IsInActiveSet(FIntPoint(TileX, TileY)))
 			&& SyncTimeSlicedData.AddGenTilesLayerIndex != INDEX_NONE)
 		{
 			for (; SyncTimeSlicedData.AddGenTilesLayerIndex < TileGenerator.GetDirtyLayersMask().Num(); ++SyncTimeSlicedData.AddGenTilesLayerIndex)
@@ -5624,7 +5624,7 @@
 	const int32 FirstDirtyTileIndex = TileGenerator.GetDirtyLayersMask().Find(true);
 
 	if (DetourMesh != nullptr
-		&& (UE::NavMesh::Private::IsWorldPartitionedDynamicNavmesh(*DestNavMesh) && IsInActiveSet(FIntPoint(TileX, TileY)))
+		&& (!UE::NavMesh::Private::IsWorldPartitionedDynamicNavmesh(*DestNavMesh) || IsInActiveSet(FIntPoint(TileX, TileY)))
 		&& FirstDirtyTileIndex != INDEX_NONE)
 	{
 		TArray<FNavMeshTileData> TileLayers = TileGenerator.GetNavigationData();
