[WP Dynamic Navmesh] Handle preexisting spawned objects when loading cells

diff --git a/Engine/Source/Runtime/Engine/Classes/AI/Navigation/NavRelevantInterface.h b/Engine/Source/Runtime/Engine/Classes/AI/Navigation/NavRelevantInterface.h
--- a/Engine/Source/Runtime/Engine/Classes/AI/Navigation/NavRelevantInterface.h
+++ b/Engine/Source/Runtime/Engine/Classes/AI/Navigation/NavRelevantInterface.h
@@ -21,12 +21,15 @@ struct FNavigationRelevantDataFilter
 	uint32 bIncludeAreas : 1;
 	/** pass when actor has any modifier with meta area */
 	uint32 bIncludeMetaAreas : 1;
+	/** fail if from level loading (only valid in WP dynamic mode) */
+	uint32 bExcludeLoadedData : 1;
 
 	FNavigationRelevantDataFilter() 
 		: bIncludeGeometry(false)
 		, bIncludeOffmeshLinks(false)
 		, bIncludeAreas(false)
 		, bIncludeMetaAreas(false)
+		, bExcludeLoadedData(false)
 	{}
 };
 
@@ -77,11 +80,16 @@ struct ENGINE_API FNavigationRelevantData : public TSharedFromThis<FNavigationRe
 
 	uint32 bSupportsGatheringGeometrySlices : 1;
 
+	/** From level loading (only valid in WP dynamic mode) */
+	uint32 bLoadedData : 1;
+
 	FNavigationRelevantData(UObject& Source)
 		: SourceObject(&Source)
 		, bPendingLazyGeometryGathering(false)
 		, bPendingLazyModifiersGathering(false)
 		, bPendingChildLazyModifiersGathering(false)
+		, bSupportsGatheringGeometrySlices(false)
+		, bLoadedData(false)
 	{}
 
 	FORCEINLINE bool HasGeometry() const { return VoxelData.Num() || CollisionData.Num(); }
diff --git a/Engine/Source/Runtime/Engine/Private/AI/Navigation/NavigationTypes.cpp b/Engine/Source/Runtime/Engine/Private/AI/Navigation/NavigationTypes.cpp
--- a/Engine/Source/Runtime/Engine/Private/AI/Navigation/NavigationTypes.cpp
+++ b/Engine/Source/Runtime/Engine/Private/AI/Navigation/NavigationTypes.cpp
@@ -143,6 +143,10 @@ bool FNavigationRelevantData::HasPerInstanceTransforms() const
 
 bool FNavigationRelevantData::IsMatchingFilter(const FNavigationRelevantDataFilter& Filter) const
 {
+	if (Filter.bExcludeLoadedData && bLoadedData)
+	{
+		return false;
+	}
 	return (Filter.bIncludeGeometry && HasGeometry()) ||
 		(Filter.bIncludeOffmeshLinks && (Modifiers.HasPotentialLinks() || Modifiers.HasLinks())) ||
 		(Filter.bIncludeAreas && Modifiers.HasAreas()) ||
diff --git a/Engine/Source/Runtime/NavigationSystem/Private/NavMesh/RecastNavMesh.cpp b/Engine/Source/Runtime/NavigationSystem/Private/NavMesh/RecastNavMesh.cpp
--- a/Engine/Source/Runtime/NavigationSystem/Private/NavMesh/RecastNavMesh.cpp
+++ b/Engine/Source/Runtime/NavigationSystem/Private/NavMesh/RecastNavMesh.cpp
@@ -20,6 +20,8 @@
 #include "NavMesh/RecastQueryFilter.h"
 #include "VisualLogger/VisualLogger.h"
 #include "WorldPartition/NavigationData/NavigationDataChunkActor.h"
+#include "Math/Color.h"
+#include "NavigationDataHandler.h"
 
 #if WITH_EDITOR
 #include "ObjectEditorUtils.h"
@@ -2316,6 +2318,28 @@ void ARecastNavMesh::OnStreamingNavDataAdded(ANavigationDataChunkActor& InActor)
 		{
 			AttachNavMeshDataChunk(*NavDataChunk);
 		}
+
+		if (IsWorldPartitionedDynamicNavmesh())
+		{
+			// Add dirtiness for preexisting elements that are not part of the base navmesh.
+			UNavigationSystemV1* NavSys = FNavigationSystem::GetCurrent<UNavigationSystemV1>(GetWorld());
+			if (NavSys)
+			{
+				FNavigationOctreeFilter Filter;
+				Filter.bIncludeGeometry = true;
+				Filter.bExcludeLoadedData = true;
+			
+				TArray<FNavigationOctreeElement> NavElements;
+				NavSys->FindElementsInNavOctree(InActor.GetBounds(), Filter, NavElements);
+
+				for (const FNavigationOctreeElement& NavElement : NavElements)
+				{
+					UE_VLOG_BOX(this, LogNavigation, Verbose, NavElement.Bounds.GetBox(), FColor::Orange, TEXT(""));
+
+					NavSys->AddDirtyArea(NavElement.Bounds.GetBox(), ENavigationDirtyFlag::All, [&NavElement] { return NavElement.Data->SourceObject.Get(); }, "Streaming data added");
+				}
+			}
+		}
 	}
 }
 
@@ -2803,6 +2827,11 @@ bool ARecastNavMesh::SupportsStreaming() const
 	return (RuntimeGeneration != ERuntimeGenerationType::Dynamic) || bIsWorldPartitioned;
 }
 
+bool ARecastNavMesh::IsWorldPartitionedDynamicNavmesh() const
+{
+	return bIsWorldPartitioned && SupportsRuntimeGeneration();
+}
+
 FRecastNavMeshGenerator* ARecastNavMesh::CreateGeneratorInstance()
 {
 	return new FRecastNavMeshGenerator(*this);
diff --git a/Engine/Source/Runtime/NavigationSystem/Private/NavMesh/RecastNavMeshGenerator.cpp b/Engine/Source/Runtime/NavigationSystem/Private/NavMesh/RecastNavMeshGenerator.cpp
--- a/Engine/Source/Runtime/NavigationSystem/Private/NavMesh/RecastNavMeshGenerator.cpp
+++ b/Engine/Source/Runtime/NavigationSystem/Private/NavMesh/RecastNavMeshGenerator.cpp
@@ -113,15 +113,6 @@
 	return NumTiles;
 }
 
-namespace UE::NavMesh::Private
-{
-	// @todo: Promote to the public interface when this becomes less experimental.
-	bool IsWorldPartitionedDynamicNavmesh(const ARecastNavMesh& NavMesh)
-	{
-		return NavMesh.bIsWorldPartitioned && NavMesh.SupportsRuntimeGeneration();
-	}
-} // namespace UE::NavMesh::Private
-
 /**
  * Exports geometry to OBJ file. Can be used to verify NavMesh generation in RecastDemo app
  * @param FileName - full name of OBJ file with extension
@@ -5560,7 +5551,7 @@
 	case EAddGeneratedTilesTimeSlicedState::AddTiles:
 	{
 		if (DetourMesh != nullptr
-			&& (!UE::NavMesh::Private::IsWorldPartitionedDynamicNavmesh(*DestNavMesh) || IsInActiveSet(FIntPoint(TileX, TileY)))
+			&& (!DestNavMesh->IsWorldPartitionedDynamicNavmesh() || IsInActiveSet(FIntPoint(TileX, TileY)))
 			&& SyncTimeSlicedData.AddGenTilesLayerIndex != INDEX_NONE)
 		{
 			for (; SyncTimeSlicedData.AddGenTilesLayerIndex < TileGenerator.GetDirtyLayersMask().Num(); ++SyncTimeSlicedData.AddGenTilesLayerIndex)
@@ -5624,7 +5615,7 @@
 	const int32 FirstDirtyTileIndex = TileGenerator.GetDirtyLayersMask().Find(true);
 
 	if (DetourMesh != nullptr
-		&& (!UE::NavMesh::Private::IsWorldPartitionedDynamicNavmesh(*DestNavMesh) || IsInActiveSet(FIntPoint(TileX, TileY)))
+		&& (!DestNavMesh->IsWorldPartitionedDynamicNavmesh() || IsInActiveSet(FIntPoint(TileX, TileY)))
 		&& FirstDirtyTileIndex != INDEX_NONE)
 	{
 		TArray<FNavMeshTileData> TileLayers = TileGenerator.GetNavigationData();
diff --git a/Engine/Source/Runtime/NavigationSystem/Private/NavigationDataHandler.cpp b/Engine/Source/Runtime/NavigationSystem/Private/NavigationDataHandler.cpp
--- a/Engine/Source/Runtime/NavigationSystem/Private/NavigationDataHandler.cpp
+++ b/Engine/Source/Runtime/NavigationSystem/Private/NavigationDataHandler.cpp
@@ -109,6 +109,13 @@ void FNavigationDataHandler::AddElementToNavOctree(const FNavigationDirtyElement
 	}
 
 	FNavigationOctreeElement GeneratedData(*ElementOwner);
+
+	// In WP dynamic mode, store if this is loaded data.
+	if (DirtyAreasController.bUseWorldPartitionedDynamicMode)
+	{
+		GeneratedData.Data->bLoadedData = DirtyElement.bIsFromVisibilityChange || FNavigationSystem::IsLevelVisibilityChanging(ElementOwner);
+	}
+	
 	const FBox ElementBounds = DirtyElement.NavInterface->GetNavigationBounds();
 
 	UObject* NavigationParent = DirtyElement.NavInterface->GetNavigationParent();
@@ -425,6 +432,13 @@ void FNavigationDataHandler::AddLevelCollisionToOctree(ULevel& Level)
 		if (!ElementId && LevelGeom && LevelGeom->Num() > 0)
 		{
 			FNavigationOctreeElement BSPElem(Level);
+			
+			// In WP dynamic mode, store if this is loaded data.
+			if (DirtyAreasController.bUseWorldPartitionedDynamicMode)
+			{
+				BSPElem.Data->bLoadedData = Level.HasVisibilityChangeRequestPending();
+			}
+			
 			FRecastNavMeshGenerator::ExportVertexSoupGeometry(*LevelGeom, *BSPElem.Data);
 
 			const FBox& Bounds = BSPElem.Data->Bounds;
diff --git a/Engine/Source/Runtime/NavigationSystem/Private/NavigationDirtyAreasController.cpp b/Engine/Source/Runtime/NavigationSystem/Private/NavigationDirtyAreasController.cpp
--- a/Engine/Source/Runtime/NavigationSystem/Private/NavigationDirtyAreasController.cpp
+++ b/Engine/Source/Runtime/NavigationSystem/Private/NavigationDirtyAreasController.cpp
@@ -73,11 +73,15 @@ void FNavigationDirtyAreasController::AddArea(const FBox& NewArea, const int32 F
 
 	if (bUseWorldPartitionedDynamicMode)
 	{
+		// Both conditions must be true to ignore dirtiness.
+		//  If it's only a visibility change and it's not in the base navmesh, it's the case created by loading a data layer (dirtiness must be applied)
+		//  If there is no visibility change, the change is not from loading/unloading a cell (dirtiness must be applied)
+		
 		// ObjectProviderFunc() is not always providing a valid object.
-		if (const bool bIsFromVisibilityChange = (DirtyElement && DirtyElement->bIsFromVisibilityChange) || FNavigationSystem::IsLevelVisibilityChanging(ObjectProviderFunc()))
+		if (const bool bIsFromVisibilityChange = (DirtyElement && DirtyElement->bIsFromVisibilityChange) || (ObjectProviderFunc && FNavigationSystem::IsLevelVisibilityChanging(ObjectProviderFunc())))
 		{
 			// If the area is from the addition or removal of objects caused by level loading/unloading and it's already in the base navmesh ignore the dirtiness.
-			if (const bool bIsIncludedInBaseNavmesh = (DirtyElement && DirtyElement->bIsInBaseNavmesh) || FNavigationSystem::IsInBaseNavmesh(ObjectProviderFunc()))
+			if (const bool bIsIncludedInBaseNavmesh = (DirtyElement && DirtyElement->bIsInBaseNavmesh) || (ObjectProviderFunc && FNavigationSystem::IsInBaseNavmesh(ObjectProviderFunc())))
 			{
 				UE_LOG(LogNavigationDirtyArea, VeryVerbose, TEXT("Ignoring dirtyness (visibility changed and in base navmesh). (object: %s from: %s)"),
 					*GetFullNameSafe(ObjectProviderFunc ? ObjectProviderFunc() : nullptr), *DebugReason.ToString());
diff --git a/Engine/Source/Runtime/NavigationSystem/Private/NavigationSystem.cpp b/Engine/Source/Runtime/NavigationSystem/Private/NavigationSystem.cpp
--- a/Engine/Source/Runtime/NavigationSystem/Private/NavigationSystem.cpp
+++ b/Engine/Source/Runtime/NavigationSystem/Private/NavigationSystem.cpp
@@ -2823,12 +2823,17 @@ void UNavigationSystemV1::UnregisterComponentToNavOctree(UActorComponent* Comp)
 	}
 }
 
-void UNavigationSystemV1::AddDirtyArea(const FBox& NewArea, int32 Flags, const char* SourceText)
+void UNavigationSystemV1::AddDirtyArea(const FBox& NewArea, int32 Flags, const FName& DebugReason /*= NAME_None*/)
 {
-	DefaultDirtyAreasController.AddArea(NewArea, Flags, nullptr, nullptr, SourceText);
+	DefaultDirtyAreasController.AddArea(NewArea, Flags, nullptr, nullptr, DebugReason);
 }
 
-void UNavigationSystemV1::AddDirtyAreas(const TArray<FBox>& NewAreas, int32 Flags, const char* SourceText)
+void UNavigationSystemV1::AddDirtyArea(const FBox& NewArea, int32 Flags, const TFunction<UObject*()>& ObjectProviderFunc, const FName& DebugReason /*= NAME_None*/)
+{
+	DefaultDirtyAreasController.AddArea(NewArea, Flags, ObjectProviderFunc, nullptr, DebugReason);
+}
+
+void UNavigationSystemV1::AddDirtyAreas(const TArray<FBox>& NewAreas, int32 Flags, const FName& DebugReason /*= NAME_None*/)
 { 
 	if (Flags == 0)
 	{
@@ -2837,7 +2842,7 @@ void UNavigationSystemV1::AddDirtyAreas(const TArray<FBox>& NewAreas, int32 Flag
 
 	for (int32 NewAreaIndex = 0; NewAreaIndex < NewAreas.Num(); NewAreaIndex++)
 	{
-		AddDirtyArea(NewAreas[NewAreaIndex], Flags, SourceText);
+		AddDirtyArea(NewAreas[NewAreaIndex], Flags, DebugReason);
 	}
 }
 
diff --git a/Engine/Source/Runtime/NavigationSystem/Public/NavMesh/RecastNavMesh.h b/Engine/Source/Runtime/NavigationSystem/Public/NavMesh/RecastNavMesh.h
--- a/Engine/Source/Runtime/NavigationSystem/Public/NavMesh/RecastNavMesh.h
+++ b/Engine/Source/Runtime/NavigationSystem/Public/NavMesh/RecastNavMesh.h
@@ -1226,6 +1226,8 @@ public:
 	virtual bool SupportsRuntimeGeneration() const override;
 	virtual bool SupportsStreaming() const override;
 
+	bool IsWorldPartitionedDynamicNavmesh() const;
+	
 	/** When using active tiles generation, navigation is only allowed to be runtime generated on a subset of tiles.
 	 *  The subset is be defined by navinvokers or loaded world partitioned cells. */
 	bool IsUsingActiveTilesGeneration(const UNavigationSystemV1& NavSys) const;
diff --git a/Engine/Source/Runtime/NavigationSystem/Public/NavigationSystem.h b/Engine/Source/Runtime/NavigationSystem/Public/NavigationSystem.h
--- a/Engine/Source/Runtime/NavigationSystem/Public/NavigationSystem.h
+++ b/Engine/Source/Runtime/NavigationSystem/Public/NavigationSystem.h
@@ -763,8 +763,9 @@ public:
 	/** updates bounds of all components implementing INavRelevantInterface */
 	static void UpdateNavOctreeBounds(AActor* Actor);
 
-	void AddDirtyArea(const FBox& NewArea, int32 Flags, const char* SourceText = nullptr);
-	void AddDirtyAreas(const TArray<FBox>& NewAreas, int32 Flags, const char* SourceText = nullptr);
+	void AddDirtyArea(const FBox& NewArea, int32 Flags, const FName& DebugReason = NAME_None);
+	void AddDirtyArea(const FBox& NewArea, int32 Flags, const TFunction<UObject*()>& ObjectProviderFunc, const FName& DebugReason = NAME_None);
+	void AddDirtyAreas(const TArray<FBox>& NewAreas, int32 Flags, const FName& DebugReason = NAME_None);
 	bool HasDirtyAreasQueued() const;
 	int32 GetNumDirtyAreas() const;
 
