[Navmesh] Fix crash when using WP static navmesh

diff --git a/Engine/Source/Runtime/NavigationSystem/Private/NavMesh/RecastNavMesh.cpp b/Engine/Source/Runtime/NavigationSystem/Private/NavMesh/RecastNavMesh.cpp
--- a/Engine/Source/Runtime/NavigationSystem/Private/NavMesh/RecastNavMesh.cpp
+++ b/Engine/Source/Runtime/NavigationSystem/Private/NavMesh/RecastNavMesh.cpp
@@ -2891,7 +2891,7 @@ FRecastNavMeshGenerator* ARecastNavMesh::CreateGeneratorInstance()
 
 bool ARecastNavMesh::IsUsingActiveTilesGeneration(const UNavigationSystemV1& NavSys) const
 {
-	return NavSys.IsActiveTilesGenerationEnabled() || bIsWorldPartitioned;
+	return SupportsRuntimeGeneration() && (NavSys.IsActiveTilesGenerationEnabled() || bIsWorldPartitioned);
 }
 
 void ARecastNavMesh::ConditionalConstructGenerator()
