[Navmesh] Use FNavTileRef instead if int32 for updated tiles to fix issue where a tile would be reuse from the tile pool while it would still be referenced in the updated tile list (causing unpredictable results).
Also includes improved tile logging.

diff --git a/Engine/Source/Runtime/NavigationSystem/Private/NavMesh/PImplRecastNavMesh.cpp b/Engine/Source/Runtime/NavigationSystem/Private/NavMesh/PImplRecastNavMesh.cpp
--- a/Engine/Source/Runtime/NavigationSystem/Private/NavMesh/PImplRecastNavMesh.cpp
+++ b/Engine/Source/Runtime/NavigationSystem/Private/NavMesh/PImplRecastNavMesh.cpp
@@ -2846,7 +2846,17 @@ int32 FPImplRecastNavMesh::GetTilesDebugGeometry(const FRecastNavMeshGenerator*
 #endif
 
 	const bool bIsBeingBuilt = Generator != nullptr && !!NavMeshOwner->bDistinctlyDrawTilesBeingBuilt
-		&& Generator->IsTileChanged(TileIdx == INDEX_NONE ? DetourNavMesh->decodePolyIdTile(DetourNavMesh->getTileRef(&Tile)) : TileIdx);
+		&& Generator->IsTileChanged(FNavTileRef(DetourNavMesh->getTileRef(&Tile)));
+
+	UE_SUPPRESS(LogNavigation, VeryVerbose,
+	{
+		if (bIsBeingBuilt)
+		{
+			const dtTileRef TileRef = DetourNavMesh->getTileRef(&Tile);
+			UE_LOG(LogNavigation, VeryVerbose, TEXT("%s TileId: %d Salt: %d TileRef: 0x%llx bIsBeingBuilt"),
+				ANSI_TO_TCHAR(__FUNCTION__), DetourNavMesh->decodePolyIdTile(TileRef), DetourNavMesh->decodePolyIdSalt(TileRef), TileRef);	
+		}
+	});
 
 	// add all the poly verts
 	FVector::FReal* F = Tile.verts;
diff --git a/Engine/Source/Runtime/NavigationSystem/Private/NavMesh/RecastNavMesh.cpp b/Engine/Source/Runtime/NavigationSystem/Private/NavMesh/RecastNavMesh.cpp
--- a/Engine/Source/Runtime/NavigationSystem/Private/NavMesh/RecastNavMesh.cpp
+++ b/Engine/Source/Runtime/NavigationSystem/Private/NavMesh/RecastNavMesh.cpp
@@ -78,6 +78,19 @@ namespace
 
 }
 
+namespace UE::NavMesh::Private
+{
+	FNavTileRef GetTileRefFromPolyRef(const dtNavMesh& DetourMesh, const NavNodeRef PolyRef)
+	{
+		unsigned int Salt = 0;
+		unsigned int TileIndex = 0;
+		unsigned int PolyIndex = 0;
+		DetourMesh.decodePolyId(PolyRef, Salt, TileIndex, PolyIndex);
+		const dtTileRef TileRef = DetourMesh.encodePolyId(Salt, TileIndex, 0);
+		return FNavTileRef(TileRef);
+	}
+} // namespace UE::NavMesh::Private
+
 FDetourTileLayout::FDetourTileLayout(const dtMeshTile& tile)
 {
 	const dtMeshHeader* header = tile.header;
@@ -285,6 +298,43 @@ uint32 FRecastDebugGeometry::GetAllocatedSize() const
 	return Size;
 }
 
+//----------------------------------------------------------------------//
+// FNavTileRef
+//----------------------------------------------------------------------//
+
+// Only use for deprecation
+void FNavTileRef::DeprecatedGetTileIdsFromNavTileRefs(const FPImplRecastNavMesh* RecastNavMeshImpl, const TArray<FNavTileRef>& InTileRefs, TArray<uint32>& OutTileIds)
+{
+	if (RecastNavMeshImpl)
+	{
+		if (const dtNavMesh* DetourMesh = RecastNavMeshImpl->DetourNavMesh)
+		{
+			OutTileIds.Reserve(InTileRefs.Num());
+			for (const FNavTileRef TileRef : InTileRefs)
+			{
+				OutTileIds.Add(DetourMesh->decodePolyIdTile((dtTileRef)TileRef));
+			}
+		}
+	}
+}
+
+// Only use for deprecation
+void FNavTileRef::DeprecatedMakeTileRefsFromTileIds(const FPImplRecastNavMesh* RecastNavMeshImpl, const TArray<uint32>& InTileIds, TArray<FNavTileRef>& OutTileRefs)
+{
+	if (RecastNavMeshImpl)
+	{
+		if (const dtNavMesh* DetourMesh = RecastNavMeshImpl->DetourNavMesh)
+		{
+			OutTileRefs.Reserve(InTileIds.Num());
+			for (const uint32 TileId : InTileIds)
+			{
+				const dtTileRef TileRef = DetourMesh->encodePolyId(0, TileId, 0);
+				OutTileRefs.Add(FNavTileRef(TileRef));
+			}
+		}
+	}
+}
+
 //----------------------------------------------------------------------//
 // ARecastNavMesh
 //----------------------------------------------------------------------//
@@ -810,6 +860,15 @@ TArray<FIntPoint>& ARecastNavMesh::GetActiveTiles()
 	return MyGenerator->ActiveTiles;
 }
 
+void ARecastNavMesh::LogRecastTile(const TCHAR* Caller, const FName& Prefix, const FName& OperationName, const dtNavMesh& DetourMesh, const int32 TileX, const int32 TileY, const int32 LayerIndex, const dtTileRef TileRef) const
+{
+	UE_LOG(LogNavigation, VeryVerbose, TEXT("%s> %s Tile (%d,%d:%d), %s TileId: %d, Salt: %d, TileRef: 0x%llx (%s)"),
+		   *GetName(), *Prefix.ToString(),
+		   TileX, TileY, LayerIndex, *OperationName.ToString(),
+		   DetourMesh.decodePolyIdTile(TileRef), DetourMesh.decodePolyIdSalt(TileRef), TileRef, 
+		   Caller);
+}
+
 void ARecastNavMesh::RestrictBuildingToActiveTiles(bool InRestrictBuildingToActiveTiles)
 {
 	FRecastNavMeshGenerator* MyGenerator = static_cast<FRecastNavMeshGenerator*>(GetGenerator());
@@ -2093,12 +2152,28 @@ void ARecastNavMesh::DrawDebugPathCorridor(NavNodeRef const* PathPolys, int32 Nu
 #endif // ENABLE_DRAW_DEBUG
 }
 
+// Deprecated
 void ARecastNavMesh::OnNavMeshTilesUpdated(const TArray<uint32>& ChangedTiles)
+{
+	TArray<FNavTileRef> ChangedTileRefs;
+	FNavTileRef::DeprecatedMakeTileRefsFromTileIds(RecastNavMeshImpl, ChangedTiles, ChangedTileRefs);
+	OnNavMeshTilesUpdated(ChangedTileRefs);
+}
+
+void ARecastNavMesh::OnNavMeshTilesUpdated(const TArray<FNavTileRef>& ChangedTiles)
 {
 	InvalidateAffectedPaths(ChangedTiles);
 }
 
+// Deprecated
 void ARecastNavMesh::InvalidateAffectedPaths(const TArray<uint32>& ChangedTiles)
+{
+	TArray<FNavTileRef> ChangedTileRefs;
+	FNavTileRef::DeprecatedMakeTileRefsFromTileIds(RecastNavMeshImpl, ChangedTiles, ChangedTileRefs);
+	InvalidateAffectedPaths(ChangedTileRefs);
+}
+
+void ARecastNavMesh::InvalidateAffectedPaths(const TArray<FNavTileRef>& ChangedTiles)
 {
 	const int32 PathsCount = ActivePaths.Num();
 	const int32 ChangedTilesCount = ChangedTiles.Num();
@@ -2107,6 +2182,12 @@ void ARecastNavMesh::InvalidateAffectedPaths(const TArray<uint32>& ChangedTiles)
 	{
 		return;
 	}
+	
+	const dtNavMesh* DetourMesh = RecastNavMeshImpl->DetourNavMesh;
+	if (DetourMesh == nullptr)
+	{
+		return;
+	}
 
 	// Paths can be registered from async pathfinding thread.
 	// Theoretically paths are invalidated synchronously by the navigation system 
@@ -2140,8 +2221,8 @@ void ARecastNavMesh::InvalidateAffectedPaths(const TArray<uint32>& ChangedTiles)
 				const NavNodeRef* PathPoly = Path->PathCorridor.GetData();
 				for (int32 NodeIndex = 0; NodeIndex < PathLenght; ++NodeIndex, ++PathPoly)
 				{
-					const uint32 NodeTileIdx = RecastNavMeshImpl->GetTileIndexFromPolyRef(*PathPoly);
-					if (ChangedTiles.Contains(NodeTileIdx))
+					const FNavTileRef NavTileRef = UE::NavMesh::Private::GetTileRefFromPolyRef(*DetourMesh, *PathPoly);
+					if (ChangedTiles.Contains(NavTileRef))
 					{
 						SharedPath->Invalidate();
 						ActivePaths.RemoveAtSwap(PathIndex, 1, /*bAllowShrinking=*/false);
@@ -2427,7 +2508,7 @@ void ARecastNavMesh::OnStreamingLevelAdded(ULevel* InLevel, UWorld* InWorld)
 
 void ARecastNavMesh::AttachNavMeshDataChunk(URecastNavMeshDataChunk& NavDataChunk)
 {
-	TArray<uint32> AttachedIndices = NavDataChunk.AttachTiles(*RecastNavMeshImpl);
+	const TArray<FNavTileRef> AttachedIndices = NavDataChunk.AttachTiles(*this);
 	if (AttachedIndices.Num() > 0)
 	{
 		InvalidateAffectedPaths(AttachedIndices);
@@ -2451,7 +2532,7 @@ void ARecastNavMesh::OnStreamingLevelRemoved(ULevel* InLevel, UWorld* InWorld)
 
 void ARecastNavMesh::DetachNavMeshDataChunk(URecastNavMeshDataChunk& NavDataChunk)
 {
-	TArray<uint32> DetachedIndices = NavDataChunk.DetachTiles(*RecastNavMeshImpl);
+	const TArray<FNavTileRef> DetachedIndices = NavDataChunk.DetachTiles(*this);
 	if (DetachedIndices.Num() > 0)
 	{
 		InvalidateAffectedPaths(DetachedIndices);
diff --git a/Engine/Source/Runtime/NavigationSystem/Private/NavMesh/RecastNavMeshDataChunk.cpp b/Engine/Source/Runtime/NavigationSystem/Private/NavMesh/RecastNavMeshDataChunk.cpp
--- a/Engine/Source/Runtime/NavigationSystem/Private/NavMesh/RecastNavMeshDataChunk.cpp
+++ b/Engine/Source/Runtime/NavigationSystem/Private/NavMesh/RecastNavMeshDataChunk.cpp
@@ -8,6 +8,7 @@
 #include "NavMesh/RecastHelpers.h"
 #include "NavMesh/RecastVersion.h"
 #include "Detour/DetourNavMeshBuilder.h"
+#include "NavMesh/RecastNavMeshGenerator.h"
 
 //----------------------------------------------------------------------//
 // FRecastTileData                                                                
@@ -66,12 +67,12 @@ static uint8* DuplicateRecastRawData(const uint8* Src, int32 SrcSize)
 
 namespace UE::NavMesh::Private
 {
-	bool IsUsingActiveTileGeneration(const FPImplRecastNavMesh& NavMeshImpl)
+	bool IsUsingActiveTileGeneration(const ARecastNavMesh& NavMesh)
 	{
-		const UNavigationSystemV1* NavSys = FNavigationSystem::GetCurrent<UNavigationSystemV1>(NavMeshImpl.NavMeshOwner->GetWorld());
+		const UNavigationSystemV1* NavSys = FNavigationSystem::GetCurrent<UNavigationSystemV1>(NavMesh.GetWorld());
 		if (NavSys)
 		{
-			return NavMeshImpl.NavMeshOwner->IsUsingActiveTilesGeneration(*NavSys);
+			return NavMesh.IsUsingActiveTilesGeneration(*NavSys);
 		}
 		return false;
 	}
@@ -191,36 +192,74 @@ void URecastNavMeshDataChunk::SerializeRecastData(FArchive& Ar, int32 NavMeshVer
 }
 #endif// WITH_RECAST
 
+// Deprecated
 TArray<uint32> URecastNavMeshDataChunk::AttachTiles(FPImplRecastNavMesh& NavMeshImpl)
 {
-	check(NavMeshImpl.NavMeshOwner && NavMeshImpl.NavMeshOwner->GetWorld());
-	const bool bIsGameWorld = NavMeshImpl.NavMeshOwner->GetWorld()->IsGameWorld();
+	check(NavMeshImpl.NavMeshOwner);
+	const TArray<FNavTileRef> TileRefs = AttachTiles(*NavMeshImpl.NavMeshOwner);
+	TArray<uint32> TileIds;
+	FNavTileRef::DeprecatedGetTileIdsFromNavTileRefs(&NavMeshImpl, TileRefs, TileIds);
+	return TileIds;
+}
+
+// Deprecated
+TArray<uint32> URecastNavMeshDataChunk::AttachTiles(FPImplRecastNavMesh& NavMeshImpl, const bool bKeepCopyOfData, const bool bKeepCopyOfCacheData)
+{
+	check(NavMeshImpl.NavMeshOwner);
+	const TArray<FNavTileRef> TileRefs = AttachTiles(*NavMeshImpl.NavMeshOwner, bKeepCopyOfData, bKeepCopyOfCacheData);
+	TArray<uint32> TileIds;
+	FNavTileRef::DeprecatedGetTileIdsFromNavTileRefs(&NavMeshImpl, TileRefs, TileIds);
+	return TileIds;
+}
+
+// Deprecated
+TArray<uint32> URecastNavMeshDataChunk::DetachTiles(FPImplRecastNavMesh& NavMeshImpl)
+{
+	check(NavMeshImpl.NavMeshOwner);
+	const TArray<FNavTileRef> TileRefs = DetachTiles(*NavMeshImpl.NavMeshOwner);
+	TArray<uint32> TileIds;
+	FNavTileRef::DeprecatedGetTileIdsFromNavTileRefs(&NavMeshImpl, TileRefs, TileIds);
+	return TileIds;
+}
+
+// Deprecated
+TArray<uint32> URecastNavMeshDataChunk::DetachTiles(FPImplRecastNavMesh& NavMeshImpl, const bool bTakeDataOwnership, const bool bTakeCacheDataOwnership)
+{
+	check(NavMeshImpl.NavMeshOwner);
+	const TArray<FNavTileRef> TileRefs = DetachTiles(*NavMeshImpl.NavMeshOwner, bTakeDataOwnership, bTakeCacheDataOwnership);
+	TArray<uint32> TileIds;
+	FNavTileRef::DeprecatedGetTileIdsFromNavTileRefs(&NavMeshImpl, TileRefs, TileIds);
+	return TileIds;
+}
+
+TArray<FNavTileRef> URecastNavMeshDataChunk::AttachTiles(ARecastNavMesh& NavMesh)
+{
+	check(NavMesh.GetWorld());
+	const bool bIsGameWorld = NavMesh.GetWorld()->IsGameWorld();
 
 	// In editor we still need to own the data so a copy will be made.
 	const bool bKeepCopyOfData = !bIsGameWorld;
 	const bool bKeepCopyOfCacheData = !bIsGameWorld;
 
-	return AttachTiles(NavMeshImpl, bKeepCopyOfData, bKeepCopyOfCacheData);
+	return AttachTiles(NavMesh, bKeepCopyOfData, bKeepCopyOfCacheData);
 }
 
-TArray<uint32> URecastNavMeshDataChunk::AttachTiles(FPImplRecastNavMesh& NavMeshImpl, const bool bKeepCopyOfData, const bool bKeepCopyOfCacheData)
+TArray<FNavTileRef> URecastNavMeshDataChunk::AttachTiles(ARecastNavMesh& NavMesh, const bool bKeepCopyOfData, const bool bKeepCopyOfCacheData)
 {
 	UE_LOG(LogNavigation, Verbose, TEXT("%s Attaching to NavMesh - %s"), ANSI_TO_TCHAR(__FUNCTION__), *NavigationDataName.ToString());
 	
-	TArray<uint32> Result;
+	TArray<FNavTileRef> Result;
 	Result.Reserve(Tiles.Num());
 
 #if WITH_RECAST	
-	dtNavMesh* NavMesh = NavMeshImpl.DetourNavMesh;
+	dtNavMesh* DetourNavMesh = NavMesh.GetRecastMesh();
 
-	if (NavMesh != nullptr)
+	if (DetourNavMesh != nullptr)
 	{
-		check(NavMeshImpl.NavMeshOwner);
-		
 		TArray<FIntPoint>* ActiveTiles = nullptr;
-		if (UE::NavMesh::Private::IsUsingActiveTileGeneration(NavMeshImpl))
+		if (UE::NavMesh::Private::IsUsingActiveTileGeneration(NavMesh))
 		{
-			ActiveTiles = &NavMeshImpl.NavMeshOwner->GetActiveTiles();
+			ActiveTiles = &NavMesh.GetActiveTiles();
 		}
 		
 		for (FRecastTileData& TileData : Tiles)
@@ -240,14 +279,13 @@ TArray<uint32> URecastNavMeshDataChunk::AttachTiles(FPImplRecastNavMesh& NavMesh
 				}
 				
 				// If there was a previous tile at the location remove it
-				if (const dtMeshTile* PreExistingTile = NavMesh->getTileAt(Header->x, Header->y, Header->layer))
+				if (const dtMeshTile* PreExistingTile = DetourNavMesh->getTileAt(Header->x, Header->y, Header->layer))
 				{
-					if (const dtTileRef PreExistingTileRef = NavMesh->getTileRef(PreExistingTile))
+					if (const dtTileRef PreExistingTileRef = DetourNavMesh->getTileRef(PreExistingTile))
 					{
-						UE_LOG(LogNavigation, VeryVerbose, TEXT("%s> Tile (%d,%d:%d), %s removing TileRef: 0x%llx"),
-							*NavMeshImpl.NavMeshOwner->GetName(), Header->x, Header->y, Header->layer, ANSI_TO_TCHAR(__FUNCTION__), PreExistingTileRef);
+						NavMesh.LogRecastTile(ANSI_TO_TCHAR(__FUNCTION__), FName("   "), FName("removing"), *DetourNavMesh, Header->x, Header->y, Header->layer, PreExistingTileRef);
 						
-						NavMesh->removeTile(PreExistingTileRef, nullptr, nullptr);	
+						DetourNavMesh->removeTile(PreExistingTileRef, nullptr, nullptr);	
 					}
 				}
 
@@ -255,7 +293,7 @@ TArray<uint32> URecastNavMeshDataChunk::AttachTiles(FPImplRecastNavMesh& NavMesh
 				dtTileRef TileRef = 0;
 				const dtMeshTile* MeshTile = nullptr;
 
-				dtStatus status = NavMesh->addTile(TileData.TileRawData->RawData, TileData.TileDataSize, DT_TILE_FREE_DATA, 0, &TileRef);
+				dtStatus status = DetourNavMesh->addTile(TileData.TileRawData->RawData, TileData.TileDataSize, DT_TILE_FREE_DATA, 0, &TileRef);
 
 				if (dtStatusFailed(status))
 				{
@@ -263,7 +301,7 @@ TArray<uint32> URecastNavMeshDataChunk::AttachTiles(FPImplRecastNavMesh& NavMesh
 				}
 				else
 				{
-					MeshTile = NavMesh->getTileByRef(TileRef);
+					MeshTile = DetourNavMesh->getTileByRef(TileRef);
 					check(MeshTile);
 					
 					TileData.X = MeshTile->header->x;
@@ -272,8 +310,7 @@ TArray<uint32> URecastNavMeshDataChunk::AttachTiles(FPImplRecastNavMesh& NavMesh
 					TileData.bAttached = true;
 				}
 
-				UE_LOG(LogNavigation, VeryVerbose, TEXT("  %s> Tile (%d,%d:%d), %s added TileRef: 0x%llx"),
-					*NavMeshImpl.NavMeshOwner->GetName(), TileData.X, TileData.Y, TileData.Layer, ANSI_TO_TCHAR(__FUNCTION__), TileRef);
+				NavMesh.LogRecastTile(ANSI_TO_TCHAR(__FUNCTION__), FName("   "), FName("added"), *DetourNavMesh, TileData.X, TileData.Y, TileData.Layer, TileRef);
 				
 				if (ActiveTiles)
 				{
@@ -298,7 +335,7 @@ TArray<uint32> URecastNavMeshDataChunk::AttachTiles(FPImplRecastNavMesh& NavMesh
 					FBox TileBBox = Recast2UnrealBox(MeshTile->header->bmin, MeshTile->header->bmax);
 
 					FNavMeshTileData LayerData(TileData.TileCacheRawData->RawData, TileData.TileCacheDataSize, TileData.Layer, TileBBox);
-					NavMeshImpl.AddTileCacheLayer(TileData.X, TileData.Y, TileData.Layer, LayerData);
+					NavMesh.GetRecastNavMeshImpl()->AddTileCacheLayer(TileData.X, TileData.Y, TileData.Layer, LayerData);
 
 					if (bKeepCopyOfCacheData == false)
 					{
@@ -313,7 +350,7 @@ TArray<uint32> URecastNavMeshDataChunk::AttachTiles(FPImplRecastNavMesh& NavMesh
 					}
 				}
 
-				Result.Add(NavMesh->decodePolyIdTile(TileRef));
+				Result.Add(FNavTileRef(TileRef));
 			}
 		}
 	}
@@ -323,40 +360,38 @@ TArray<uint32> URecastNavMeshDataChunk::AttachTiles(FPImplRecastNavMesh& NavMesh
 	return Result;
 }
 
-TArray<uint32> URecastNavMeshDataChunk::DetachTiles(FPImplRecastNavMesh& NavMeshImpl)
+TArray<FNavTileRef> URecastNavMeshDataChunk::DetachTiles(ARecastNavMesh& NavMesh)
 {
-	check(NavMeshImpl.NavMeshOwner && NavMeshImpl.NavMeshOwner->GetWorld());
-	const bool bIsGameWorld = NavMeshImpl.NavMeshOwner->GetWorld()->IsGameWorld();
+	check(NavMesh.GetWorld());
+	const bool bIsGameWorld = NavMesh.GetWorld()->IsGameWorld();
 
 	// Keep data in game worlds (in editor we have a copy of the data so we don't keep it).
 	const bool bTakeDataOwnership = bIsGameWorld;
 	const bool bTakeCacheDataOwnership = bIsGameWorld;
 
-	return DetachTiles(NavMeshImpl, bTakeDataOwnership, bTakeCacheDataOwnership);
+	return DetachTiles(NavMesh, bTakeDataOwnership, bTakeCacheDataOwnership);
 }
 
-TArray<uint32> URecastNavMeshDataChunk::DetachTiles(FPImplRecastNavMesh& NavMeshImpl, const bool bTakeDataOwnership, const bool bTakeCacheDataOwnership)
+TArray<FNavTileRef> URecastNavMeshDataChunk::DetachTiles(ARecastNavMesh& NavMesh, const bool bTakeDataOwnership, const bool bTakeCacheDataOwnership)
 {
 	UE_LOG(LogNavigation, Verbose, TEXT("%s Detaching from %s"), ANSI_TO_TCHAR(__FUNCTION__), *NavigationDataName.ToString());
 
-	TArray<uint32> Result;
+	TArray<FNavTileRef> Result;
 	Result.Reserve(Tiles.Num());
 
 #if WITH_RECAST
-	dtNavMesh* NavMesh = NavMeshImpl.DetourNavMesh;
+	dtNavMesh* DetourNavMesh = NavMesh.GetRecastMesh();
 
-	if (NavMesh != nullptr)
+	if (DetourNavMesh != nullptr)
 	{
-		check(NavMeshImpl.NavMeshOwner);
-		
 		TArray<FIntPoint>* ActiveTiles = nullptr;
-		if (UE::NavMesh::Private::IsUsingActiveTileGeneration(NavMeshImpl))
+		if (UE::NavMesh::Private::IsUsingActiveTileGeneration(NavMesh))
 		{
-			ActiveTiles = &NavMeshImpl.NavMeshOwner->GetActiveTiles();
+			ActiveTiles = &NavMesh.GetActiveTiles();
 		}
 
 		TArray<const dtMeshTile*> ExtraMeshTiles;
-		const bool bIsDynamic = NavMeshImpl.NavMeshOwner->SupportsRuntimeGeneration();
+		const bool bIsDynamic = NavMesh.SupportsRuntimeGeneration();
 		
 		for (FRecastTileData& TileData : Tiles)
 		{
@@ -364,14 +399,14 @@ TArray<uint32> URecastNavMeshDataChunk::DetachTiles(FPImplRecastNavMesh& NavMesh
 			{
 				// Detach tile cache layer and take ownership over compressed data
 				dtTileRef TileRef = 0;
-				const dtMeshTile* MeshTile = NavMesh->getTileAt(TileData.X, TileData.Y, TileData.Layer);
+				const dtMeshTile* MeshTile = DetourNavMesh->getTileAt(TileData.X, TileData.Y, TileData.Layer);
 				if (MeshTile)
 				{
-					TileRef = NavMesh->getTileRef(MeshTile);
+					TileRef = DetourNavMesh->getTileRef(MeshTile);
 
 					if (bTakeCacheDataOwnership)
 					{
-						FNavMeshTileData TileCacheData = NavMeshImpl.GetTileCacheLayer(TileData.X, TileData.Y, TileData.Layer);
+						FNavMeshTileData TileCacheData = NavMesh.GetRecastNavMeshImpl()->GetTileCacheLayer(TileData.X, TileData.Y, TileData.Layer);
 						if (TileCacheData.IsValid())
 						{
 							TileData.TileCacheDataSize = TileCacheData.DataSize;
@@ -379,20 +414,19 @@ TArray<uint32> URecastNavMeshDataChunk::DetachTiles(FPImplRecastNavMesh& NavMesh
 						}
 					}
 
-					UE_LOG(LogNavigation, VeryVerbose, TEXT("  %s> Tile (%d,%d:%d), removing TileRef: 0x%llx"),
-						*NavMeshImpl.NavMeshOwner->GetName(), TileData.X, TileData.Y, TileData.Layer, TileRef);
+					NavMesh.LogRecastTile(ANSI_TO_TCHAR(__FUNCTION__), FName("   "), FName("removing"), *DetourNavMesh, TileData.X, TileData.Y, TileData.Layer, TileRef);
 				
-					NavMeshImpl.RemoveTileCacheLayer(TileData.X, TileData.Y, TileData.Layer);
+					NavMesh.GetRecastNavMeshImpl()->RemoveTileCacheLayer(TileData.X, TileData.Y, TileData.Layer);
 
 					if (bTakeDataOwnership)
 					{
 						// Remove tile from navmesh and take ownership of tile raw data
-						NavMesh->removeTile(TileRef, &TileData.TileRawData->RawData, &TileData.TileDataSize);
+						DetourNavMesh->removeTile(TileRef, &TileData.TileRawData->RawData, &TileData.TileDataSize);
 					}
 					else
 					{
 						// In the editor we have a copy of tile data so just release tile in navmesh
-						NavMesh->removeTile(TileRef, nullptr, nullptr);
+						DetourNavMesh->removeTile(TileRef, nullptr, nullptr);
 					}
 
 					if (ActiveTiles)
@@ -400,25 +434,25 @@ TArray<uint32> URecastNavMeshDataChunk::DetachTiles(FPImplRecastNavMesh& NavMesh
 						ActiveTiles->RemoveSwap(FIntPoint(TileData.X, TileData.Y));	
 					}
 						
-					Result.Add(NavMesh->decodePolyIdTile(TileRef));
+					Result.Add(FNavTileRef(TileRef));
 				}
 
 				if (bIsDynamic)
 				{
 					// Remove any tile remaining
-					const int32 MaxTiles = NavMesh->getTileCountAt(TileData.X, TileData.Y);
+					const int32 MaxTiles = DetourNavMesh->getTileCountAt(TileData.X, TileData.Y);
 					if (MaxTiles > 0)
 					{
 						ExtraMeshTiles.SetNumZeroed(MaxTiles, false);
-						const int32 MeshTilesCount = NavMesh->getTilesAt(TileData.X, TileData.Y, ExtraMeshTiles.GetData(), MaxTiles);
+						const int32 MeshTilesCount = DetourNavMesh->getTilesAt(TileData.X, TileData.Y, ExtraMeshTiles.GetData(), MaxTiles);
 						for (int32 i = 0; i < MeshTilesCount; ++i)
 						{
 							const dtMeshTile* ExtraMeshTile = ExtraMeshTiles[i];
-							dtTileRef ExtraTileRef = NavMesh->getTileRef(ExtraMeshTile);
+							dtTileRef ExtraTileRef = DetourNavMesh->getTileRef(ExtraMeshTile);
 							if (ExtraTileRef)
 							{
-								NavMesh->removeTile(ExtraTileRef, nullptr, nullptr);
-								Result.Add(NavMesh->decodePolyIdTile(ExtraTileRef));
+								DetourNavMesh->removeTile(ExtraTileRef, nullptr, nullptr);
+								Result.Add(FNavTileRef(ExtraTileRef));
 							}
 						}
 					}
diff --git a/Engine/Source/Runtime/NavigationSystem/Private/NavMesh/RecastNavMeshGenerator.cpp b/Engine/Source/Runtime/NavigationSystem/Private/NavMesh/RecastNavMeshGenerator.cpp
--- a/Engine/Source/Runtime/NavigationSystem/Private/NavMesh/RecastNavMeshGenerator.cpp
+++ b/Engine/Source/Runtime/NavigationSystem/Private/NavMesh/RecastNavMeshGenerator.cpp
@@ -4630,7 +4630,7 @@ void FRecastNavMeshGenerator::EnsureBuildCompletion()
 	do 
 	{
 		const int32 NumTasksToProcess = (bDoAsyncDataGathering ? 1 : MaxTileGeneratorTasks) - RunningDirtyTiles.Num();
-		ProcessTileTasks(NumTasksToProcess);
+		ProcessTileTasksAndGetUpdatedTiles(NumTasksToProcess);
 		
 		// Block until tasks are finished
 		for (FRunningTileElement& Element : RunningDirtyTiles)
@@ -4688,15 +4688,15 @@ void FRecastNavMeshGenerator::TickAsyncBuild(float DeltaSeconds)
 	const bool bDoAsyncDataGathering = GatherGeometryOnGameThread() == false;
 
 	const int32 NumTasksToSubmit = (bDoAsyncDataGathering ? 1 : MaxTileGeneratorTasks) - NumRunningTasks;
-	TArray<uint32> UpdatedTileIndices = ProcessTileTasks(NumTasksToSubmit);
+	TArray<FNavTileRef> UpdatedTileRefs = ProcessTileTasksAndGetUpdatedTiles(NumTasksToSubmit);
 			
-	if (UpdatedTileIndices.Num() > 0)
+	if (UpdatedTileRefs.Num() > 0)
 	{
 		{
 			QUICK_SCOPE_CYCLE_COUNTER(STAT_RecastNavMeshGenerator_OnNavMeshTilesUpdated);
 
 			// Invalidate active paths that go through regenerated tiles
-			DestNavMesh->OnNavMeshTilesUpdated(UpdatedTileIndices);
+			DestNavMesh->OnNavMeshTilesUpdated(UpdatedTileRefs);
 		}
 
 		bRequestDrawingUpdate = true;
@@ -4704,11 +4704,24 @@ void FRecastNavMeshGenerator::TickAsyncBuild(float DeltaSeconds)
 #if	WITH_EDITOR
 		// Store completed tiles with timestamps to have ability to distinguish during debug draw
 		const double Timestamp = FPlatformTime::Seconds();
-		RecentlyBuiltTiles.Reserve(RecentlyBuiltTiles.Num() + UpdatedTileIndices.Num());
-		for (uint32 TiledIdx : UpdatedTileIndices)
+		RecentlyBuiltTiles.Reserve(RecentlyBuiltTiles.Num() + UpdatedTileRefs.Num());
+		for (const FNavTileRef TileRef : UpdatedTileRefs)
 		{
+			UE_SUPPRESS(LogNavigation, VeryVerbose,
+			{
+				if (DestNavMesh->GetRecastNavMeshImpl())
+				{
+					if (const dtNavMesh* DetourMesh = DestNavMesh->GetRecastNavMeshImpl()->GetRecastMesh())
+					{
+						const uint32 TileIndex = DetourMesh->decodePolyIdTile((dtTileRef)TileRef);
+						const uint32 Salt = DetourMesh->decodePolyIdSalt((dtTileRef)TileRef);
+						UE_LOG(LogNavigation, VeryVerbose, TEXT("%s Adding to RecentlyBuiltTiles TileId: %d Salt: %d TileRef: 0x%llx"), ANSI_TO_TCHAR(__FUNCTION__), TileIndex, Salt, (dtTileRef)TileRef);
+					}
+				}
+			});
+			
 			FTileTimestamp TileTimestamp;
-			TileTimestamp.TileIdx = TiledIdx;
+			TileTimestamp.NavTileRef = TileRef;
 			TileTimestamp.Timestamp = Timestamp;
 			RecentlyBuiltTiles.Add(TileTimestamp);
 		}
@@ -4823,7 +4836,7 @@ void FRecastNavMeshGenerator::ResetTimeSlicedTileGeneratorSync()
 	SyncTimeSlicedData.ProcessTileTasksSyncState = EProcessTileTasksSyncTimeSlicedState::Init;
 	SyncTimeSlicedData.UpdatedTilesCache.Reset();
 	SyncTimeSlicedData.OldLayerTileIdMapCached.Reset();
-	SyncTimeSlicedData.ResultTileIndicesCached.Reset();
+	SyncTimeSlicedData.ResultTileRefsCached.Reset();
 	SyncTimeSlicedData.AddGeneratedTilesState = EAddGeneratedTilesTimeSlicedState::Init;
 	SyncTimeSlicedData.AddGenTilesLayerIndex = 0;
 }
@@ -4833,7 +4846,7 @@ void FRecastNavMeshGenerator::RemoveTiles(const TArray<FIntPoint>& Tiles)
 {
 	for (const FIntPoint& TileXY : Tiles)
 	{
-		RemoveTileLayers(TileXY.X, TileXY.Y);
+		RemoveTileLayersAndGetUpdatedTiles(TileXY.X, TileXY.Y);
 
 		if (PendingDirtyTiles.Num() > 0)
 		{
@@ -4942,10 +4955,19 @@ namespace RecastTileVersionHelper
 	}
 }
 
+// Deprecated
 TArray<uint32> FRecastNavMeshGenerator::RemoveTileLayers(const int32 TileX, const int32 TileY, TMap<int32, dtPolyRef>* OldLayerTileIdMap)
+{
+	const TArray<FNavTileRef>& TileRefs = RemoveTileLayersAndGetUpdatedTiles(TileX, TileY, OldLayerTileIdMap);
+	TArray<uint32> TileIds;
+	FNavTileRef::DeprecatedGetTileIdsFromNavTileRefs(DestNavMesh->GetRecastNavMeshImpl(), TileRefs, TileIds);
+	return TileIds;
+}
+
+TArray<FNavTileRef> FRecastNavMeshGenerator::RemoveTileLayersAndGetUpdatedTiles(const int32 TileX, const int32 TileY, TMap<int32, dtPolyRef>* OldLayerTileIdMap)
 {
 	dtNavMesh* DetourMesh = DestNavMesh->GetRecastNavMeshImpl()->GetRecastMesh();
-	TArray<uint32> UpdatedIndices;
+	TArray<FNavTileRef> UpdatedIndices;
 	
 	if (DetourMesh != nullptr && DetourMesh->isEmpty() == false)
 	{
@@ -4963,13 +4985,13 @@ TArray<uint32> FRecastNavMeshGenerator::RemoveTileLayers(const int32 TileX, cons
 				dtPolyRef TileRef = DetourMesh->getTileRef(Tiles[i]);
 
 				NumActiveTiles--;
-				UE_LOG(LogNavigation, VeryVerbose, TEXT("%s> Tile (%d,%d:%d), removing TileRef: 0x%llx (active count: %d)"),
-					*DestNavMesh->GetName(), TileX, TileY, LayerIndex, TileRef, NumActiveTiles);
+				DestNavMesh->LogRecastTile(ANSI_TO_TCHAR(__FUNCTION__), FName(""), FName("removing"), *DetourMesh, TileX, TileY, LayerIndex, TileRef);
 
 				DetourMesh->removeTile(TileRef, nullptr, nullptr);
 
-				uint32 TileId = RecastTileVersionHelper::GetUpdatedTileId(TileRef, DetourMesh);
-				UpdatedIndices.AddUnique(TileId);
+				UpdatedIndices.AddUnique(FNavTileRef(TileRef));
+
+				RecastTileVersionHelper::GetUpdatedTileId(TileRef, DetourMesh);	// Updates TileRef
 
 				if (OldLayerTileIdMap)
 				{
@@ -5005,7 +5027,7 @@ FRecastNavMeshGenerator::FSyncTimeSlicedData::FSyncTimeSlicedData()
 {
 }
 
-void FRecastNavMeshGenerator::AddGeneratedTileLayer(int32 LayerIndex, FRecastTileGenerator& TileGenerator, const TMap<int32, dtPolyRef>& OldLayerTileIdMap, TArray<uint32>& OutResultTileIndices)
+void FRecastNavMeshGenerator::AddGeneratedTileLayer(int32 LayerIndex, FRecastTileGenerator& TileGenerator, const TMap<int32, dtPolyRef>& OldLayerTileIdMap, TArray<FNavTileRef>& OutResultTileRefs)
 {
 	SCOPE_CYCLE_COUNTER(STAT_Navigation_RecastAddGeneratedTileLayer);
 
@@ -5032,13 +5054,13 @@ void FRecastNavMeshGenerator::AddGeneratedTileLayer(int32 LayerIndex, FRecastTil
 		if (OldTileRef)
 		{
 			NumActiveTiles--;
-			UE_LOG(LogNavigation, VeryVerbose, TEXT("%s> Tile (%d,%d:%d), removing TileRef: 0x%llx (active count: %d)"),
-				*DestNavMesh->GetName(), TileX, TileY, LayerIndex, OldTileRef, NumActiveTiles);
+			DestNavMesh->LogRecastTile(ANSI_TO_TCHAR(__FUNCTION__), FName(""), FName("removing"), *DetourMesh, TileX, TileY, LayerIndex, OldTileRef);
 
 			DetourMesh->removeTile(OldTileRef, nullptr, nullptr);
 
-			const uint32 TileId = RecastTileVersionHelper::GetUpdatedTileId(OldTileRef, DetourMesh);
-			OutResultTileIndices.AddUnique(TileId);
+			OutResultTileRefs.AddUnique(FNavTileRef(OldTileRef));
+			
+			RecastTileVersionHelper::GetUpdatedTileId(OldTileRef, DetourMesh);	// Updates OldTileRef
 		}
 		else
 		{
@@ -5074,11 +5096,10 @@ void FRecastNavMeshGenerator::AddGeneratedTileLayer(int32 LayerIndex, FRecastTil
 			}
 			else
 			{
-				OutResultTileIndices.AddUnique(DetourMesh->decodePolyIdTile(ResultTileRef));
+				OutResultTileRefs.AddUnique(FNavTileRef(ResultTileRef));
 				NumActiveTiles++;
 
-				UE_LOG(LogNavigation, VeryVerbose, TEXT("%s> Tile (%d,%d:%d), added TileRef: 0x%llx (active count: %d)"),
-					*DestNavMesh->GetName(), TileX, TileY, LayerIndex, ResultTileRef, NumActiveTiles);
+				DestNavMesh->LogRecastTile(ANSI_TO_TCHAR(__FUNCTION__), FName(""), FName("added generated"), *DetourMesh, TileX, TileY, LayerIndex, ResultTileRef);
 
 				{
 					// NavMesh took the ownership of generated data, so we don't need to deallocate it
@@ -5091,8 +5112,8 @@ void FRecastNavMeshGenerator::AddGeneratedTileLayer(int32 LayerIndex, FRecastTil
 	{
 		// remove the layer since it ended up empty
 		DetourMesh->removeTile(OldTileRef, nullptr, nullptr);
-		const uint32 TileId = RecastTileVersionHelper::GetUpdatedTileId(OldTileRef, DetourMesh);
-		OutResultTileIndices.AddUnique(TileId);
+
+		OutResultTileRefs.AddUnique(FNavTileRef(OldTileRef));
 	}
 }
 
@@ -5146,7 +5167,18 @@ void FRecastNavMeshGenerator::LogDirtyAreas(
 }
 #endif
 
+// Deprecated 
 ETimeSliceWorkResult FRecastNavMeshGenerator::AddGeneratedTilesTimeSliced(FRecastTileGenerator& TileGenerator, TArray<uint32>& OutResultTileIndices)
+{
+	TArray<FNavTileRef> OutTileRefs;
+	const ETimeSliceWorkResult Result = AddGeneratedTilesTimeSliced(TileGenerator, OutTileRefs);
+
+	TArray<uint32> TileIds;
+	FNavTileRef::DeprecatedGetTileIdsFromNavTileRefs(DestNavMesh->GetRecastNavMeshImpl(), OutTileRefs, TileIds);
+	return Result;
+}
+
+ETimeSliceWorkResult FRecastNavMeshGenerator::AddGeneratedTilesTimeSliced(FRecastTileGenerator& TileGenerator, TArray<FNavTileRef>& OutResultTileRefs)
 {
 	SCOPE_CYCLE_COUNTER(STAT_Navigation_RecastAddGeneratedTiles);
 
@@ -5163,15 +5195,15 @@ ETimeSliceWorkResult FRecastNavMeshGenerator::AddGeneratedTilesTimeSliced(FRecas
 	{
 	case EAddGeneratedTilesTimeSlicedState::Init:
 	{
-		SyncTimeSlicedData.ResultTileIndicesCached.Reset();
-		SyncTimeSlicedData.ResultTileIndicesCached.Reserve(TileLayers.Num());
+		SyncTimeSlicedData.ResultTileRefsCached.Reset();
+		SyncTimeSlicedData.ResultTileRefsCached.Reserve(TileLayers.Num());
 		SyncTimeSlicedData.OldLayerTileIdMapCached.Reset();
 		SyncTimeSlicedData.OldLayerTileIdMapCached.Reserve(TileLayers.Num());
 		SyncTimeSlicedData.AddGenTilesLayerIndex = TileGenerator.GetDirtyLayersMask().Find(true);
 		if (TileGenerator.IsFullyRegenerated())
 		{
 			// remove all layers
-			SyncTimeSlicedData.ResultTileIndicesCached = RemoveTileLayers(TileX, TileY, &SyncTimeSlicedData.OldLayerTileIdMapCached);
+			SyncTimeSlicedData.ResultTileRefsCached = RemoveTileLayersAndGetUpdatedTiles(TileX, TileY, &SyncTimeSlicedData.OldLayerTileIdMapCached);
 		}
 
 		SyncTimeSlicedData.AddGeneratedTilesState = EAddGeneratedTilesTimeSlicedState::AddTiles;
@@ -5192,7 +5224,7 @@ ETimeSliceWorkResult FRecastNavMeshGenerator::AddGeneratedTilesTimeSliced(FRecas
 						break;
 					}
 
-					AddGeneratedTileLayer(SyncTimeSlicedData.AddGenTilesLayerIndex, TileGenerator, SyncTimeSlicedData.OldLayerTileIdMapCached, SyncTimeSlicedData.ResultTileIndicesCached);
+					AddGeneratedTileLayer(SyncTimeSlicedData.AddGenTilesLayerIndex, TileGenerator, SyncTimeSlicedData.OldLayerTileIdMapCached, SyncTimeSlicedData.ResultTileRefsCached);
 
 					SyncTimeSlicedData.TimeSliceManager->GetTimeSlicer().TestTimeSliceFinished();
 				}
@@ -5218,25 +5250,34 @@ ETimeSliceWorkResult FRecastNavMeshGenerator::AddGeneratedTilesTimeSliced(FRecas
 		SyncTimeSlicedData.AddGenTilesLayerIndex = 0;
 		SyncTimeSlicedData.AddGeneratedTilesState = EAddGeneratedTilesTimeSlicedState::Init;
 
-		OutResultTileIndices = MoveTemp(SyncTimeSlicedData.ResultTileIndicesCached);
+		OutResultTileRefs = MoveTemp(SyncTimeSlicedData.ResultTileRefsCached);
 	}
 
 	return WorkResult;
 }
 
+// Deprecated
 TArray<uint32> FRecastNavMeshGenerator::AddGeneratedTiles(FRecastTileGenerator& TileGenerator)
+{
+	const TArray<FNavTileRef>& TileRefs = AddGeneratedTilesAndGetUpdatedTiles(TileGenerator);
+	TArray<uint32> TileIds;
+	FNavTileRef::DeprecatedGetTileIdsFromNavTileRefs(DestNavMesh->GetRecastNavMeshImpl(), TileRefs, TileIds);
+	return TileIds;
+}
+
+TArray<FNavTileRef> FRecastNavMeshGenerator::AddGeneratedTilesAndGetUpdatedTiles(FRecastTileGenerator& TileGenerator)
 {
 	SCOPE_CYCLE_COUNTER(STAT_Navigation_RecastAddGeneratedTiles);
 
 	TMap<int32, dtPolyRef> OldLayerTileIdMap;
-	TArray<uint32> ResultTileIndices;
+	TArray<FNavTileRef> ResultTileRefs;
 	const int32 TileX = TileGenerator.GetTileX();
 	const int32 TileY = TileGenerator.GetTileY();
 
 	if (TileGenerator.IsFullyRegenerated())
 	{
 		// remove all layers
-		ResultTileIndices = RemoveTileLayers(TileX, TileY, &OldLayerTileIdMap);
+		ResultTileRefs = RemoveTileLayersAndGetUpdatedTiles(TileX, TileY, &OldLayerTileIdMap);
 	}
 
 	dtNavMesh* DetourMesh = DestNavMesh->GetRecastNavMeshImpl()->GetRecastMesh();
@@ -5247,18 +5288,18 @@ TArray<uint32> FRecastNavMeshGenerator::AddGeneratedTiles(FRecastTileGenerator&
 		&& FirstDirtyTileIndex != INDEX_NONE)
 	{
 		TArray<FNavMeshTileData> TileLayers = TileGenerator.GetNavigationData();
-		ResultTileIndices.Reserve(TileLayers.Num());
+		ResultTileRefs.Reserve(TileLayers.Num());
 
 		for (int32 LayerIndex = FirstDirtyTileIndex; LayerIndex < TileGenerator.GetDirtyLayersMask().Num(); ++LayerIndex)
 		{
 			if (TileGenerator.IsLayerChanged(LayerIndex))
 			{
-				AddGeneratedTileLayer(LayerIndex, TileGenerator, OldLayerTileIdMap, ResultTileIndices);
+				AddGeneratedTileLayer(LayerIndex, TileGenerator, OldLayerTileIdMap, ResultTileRefs);
 			}
 		}
 	}
 
-	return ResultTileIndices;
+	return ResultTileRefs;
 }
 
 void FRecastNavMeshGenerator::DiscardCurrentBuildingTasks()
@@ -5657,13 +5698,21 @@ TSharedRef<FRecastTileGenerator> FRecastNavMeshGenerator::CreateTileGenerator(co
 	return ConstuctTileGeneratorImpl<FRecastTileGenerator>(Coord, DirtyAreas);
 }
 
+// Deprecated
 void FRecastNavMeshGenerator::RemoveLayers(const FIntPoint& Tile, TArray<uint32>& UpdatedTiles)
+{
+	TArray<FNavTileRef> TileRefs;
+	RemoveLayers(Tile, TileRefs);
+	FNavTileRef::DeprecatedGetTileIdsFromNavTileRefs(DestNavMesh->GetRecastNavMeshImpl(), TileRefs, UpdatedTiles);
+}
+	
+void FRecastNavMeshGenerator::RemoveLayers(const FIntPoint& Tile, TArray<FNavTileRef>& UpdatedTiles)
 {
 	SCOPE_CYCLE_COUNTER(STAT_Navigation_RemoveLayers);
 
 	// If there is nothing to generate remove all tiles from navmesh at specified grid coordinates
 	UpdatedTiles.Append(
-		RemoveTileLayers(Tile.X, Tile.Y)
+		RemoveTileLayersAndGetUpdatedTiles(Tile.X, Tile.Y)
 	);
 	DestNavMesh->MarkEmptyTileCacheLayers(Tile.X, Tile.Y);
 }
@@ -5691,11 +5740,20 @@ void FRecastNavMeshGenerator::StoreDebugData(const FRecastTileGenerator& TileGen
 #endif
 
 #if RECAST_ASYNC_REBUILDING
+// Deprecated	
 TArray<uint32> FRecastNavMeshGenerator::ProcessTileTasksAsync(const int32 NumTasksToProcess)
+{
+	const TArray<FNavTileRef>& TileRefs = ProcessTileTasksAsyncAndGetUpdatedTiles(NumTasksToProcess);
+	TArray<uint32> TileIds;
+	FNavTileRef::DeprecatedGetTileIdsFromNavTileRefs(DestNavMesh->GetRecastNavMeshImpl(), TileRefs, TileIds);
+	return TileIds;
+}
+
+TArray<FNavTileRef> FRecastNavMeshGenerator::ProcessTileTasksAsyncAndGetUpdatedTiles(const int32 NumTasksToProcess)
 {
 	QUICK_SCOPE_CYCLE_COUNTER(STAT_RecastNavMeshGenerator_ProcessTileTasksAsync);
 
-	TArray<uint32> UpdatedTiles;
+	TArray<FNavTileRef> UpdatedTiles;
 	const bool bGameStaticNavMesh = IsGameStaticNavMesh(DestNavMesh);
 
 	int32 NumProcessedTasks = 0;
@@ -5764,8 +5822,8 @@ TArray<uint32> FRecastNavMeshGenerator::ProcessTileTasksAsync(const int32 NumTas
 			if (!Element.bShouldDiscard)
 			{
 				FRecastTileGenerator& TileGenerator = *(Element.AsyncTask->GetTask().TileGenerator);
-				TArray<uint32> UpdatedTileIndices = AddGeneratedTiles(TileGenerator);
-				UpdatedTiles.Append(UpdatedTileIndices);
+				TArray<FNavTileRef> UpdatedTileRefs = AddGeneratedTilesAndGetUpdatedTiles(TileGenerator);
+				UpdatedTiles.Append(UpdatedTileRefs);
 			
 				StoreCompressedTileCacheLayers(TileGenerator, Element.Coord.X, Element.Coord.Y);
 
@@ -5807,14 +5865,23 @@ TSharedRef<FRecastTileGenerator> FRecastNavMeshGenerator::CreateTileGeneratorFro
 	return TileGenerator;
 }
 
+// Deprecated
 TArray<uint32> FRecastNavMeshGenerator::ProcessTileTasksSyncTimeSliced()
+{
+	const TArray<FNavTileRef>& TileRefs = ProcessTileTasksSyncTimeSlicedAndGetUpdatedTiles();
+	TArray<uint32> TileIds;
+	FNavTileRef::GetTileIdsFromNavTileRefs(DestNavMesh->GetRecastNavMeshImpl(), TileRefs, TileIds);
+	return TileIds;
+}
+
+TArray<FNavTileRef> FRecastNavMeshGenerator::ProcessTileTasksSyncTimeSlicedAndGetUpdatedTiles()
 {
 	QUICK_SCOPE_CYCLE_COUNTER(STAT_RecastNavMeshGenerator_ProcessTileTasksSyncTimeSliced);
 	CSV_SCOPED_TIMING_STAT(NAVREGEN, ProcessTileTasksSyncTimeSliced);
 
 	check(SyncTimeSlicedData.TimeSliceManager);
 
-	TArray<uint32> UpdatedTiles;
+	TArray<FNavTileRef> UpdatedTiles;
 	double TimeStartProcessingTileThisFrame = 0.;
 
 	auto HasWorkToDo = [this]()
@@ -5990,14 +6057,23 @@ TArray<uint32> FRecastNavMeshGenerator::ProcessTileTasksSyncTimeSliced()
 	return EndFunction(false /* bCalcTileRegenDuration */, bHadWorktoDo);
 }
 
-//this code path is approx 10% faster than ProcessTileTasksSyncTimeSliced, however it spikes far worse for most use cases.
+// Deprecated
 TArray<uint32> FRecastNavMeshGenerator::ProcessTileTasksSync(const int32 NumTasksToProcess)
+{
+	const TArray<FNavTileRef>& TileRefs = ProcessTileTasksSyncAndGetUpdatedTiles(NumTasksToProcess);
+	TArray<uint32> TileIds;
+	FNavTileRef::GetTileIdsFromNavTileRefs(DestNavMesh->GetRecastNavMeshImpl(), TileRefs, TileIds);
+	return TileIds;
+}
+	
+//this code path is approx 10% faster than ProcessTileTasksSyncTimeSliced, however it spikes far worse for most use cases.
+TArray<FNavTileRef> FRecastNavMeshGenerator::ProcessTileTasksSyncAndGetUpdatedTiles(const int32 NumTasksToProcess)
 {
 	QUICK_SCOPE_CYCLE_COUNTER(STAT_RecastNavMeshGenerator_ProcessTileTasksSync);
 
 	const bool bGameStaticNavMesh = IsGameStaticNavMesh(DestNavMesh);
 	int32 NumProcessedTasks = 0;
-	TArray<uint32> UpdatedTiles;
+	TArray<FNavTileRef> UpdatedTiles;
 	FIntPoint TileLocation;
 
 	// Submit pending tile elements
@@ -6014,7 +6090,7 @@ TArray<uint32> FRecastNavMeshGenerator::ProcessTileTasksSync(const int32 NumTask
 		{
 			TileGeneratorRef.DoWork();
 
-			UpdatedTiles = AddGeneratedTiles(TileGeneratorRef);
+			UpdatedTiles = AddGeneratedTilesAndGetUpdatedTiles(TileGeneratorRef);
 
 			StoreCompressedTileCacheLayers(TileGeneratorRef, TileLocation.X, TileLocation.Y);
 		}
@@ -6037,14 +6113,22 @@ TArray<uint32> FRecastNavMeshGenerator::ProcessTileTasksSync(const int32 NumTask
 #endif
 
 TArray<uint32> FRecastNavMeshGenerator::ProcessTileTasks(const int32 NumTasksToProcess)
+{
+	const TArray<FNavTileRef>& TileRefs = ProcessTileTasksAndGetUpdatedTiles(NumTasksToProcess);
+	TArray<uint32> TileIds;
+	FNavTileRef::DeprecatedGetTileIdsFromNavTileRefs(DestNavMesh->GetRecastNavMeshImpl(), TileRefs, TileIds);
+	return TileIds;
+}
+
+TArray<FNavTileRef> FRecastNavMeshGenerator::ProcessTileTasksAndGetUpdatedTiles(const int32 NumTasksToProcess)
 {
 	QUICK_SCOPE_CYCLE_COUNTER(STAT_RecastNavMeshGenerator_ProcessTileTasks);
 
 	const bool bHasTasksAtStart = GetNumRemaningBuildTasks() > 0;
-	TArray<uint32> UpdatedTiles;
+	TArray<FNavTileRef> UpdatedTiles;
 
 #if RECAST_ASYNC_REBUILDING
-	UpdatedTiles = ProcessTileTasksAsync(NumTasksToProcess);
+	UpdatedTiles = ProcessTileTasksAsyncAndGetUpdatedTiles(NumTasksToProcess);
 #else
 	if (SyncTimeSlicedData.TimeSliceManager)
 	{
@@ -6069,11 +6153,11 @@ TArray<uint32> FRecastNavMeshGenerator::ProcessTileTasks(const int32 NumTasksToP
 
 	if (SyncTimeSlicedData.bTimeSliceRegenActive)
 	{
-		UpdatedTiles = ProcessTileTasksSyncTimeSliced();
+		UpdatedTiles = ProcessTileTasksSyncTimeSlicedAndGetUpdatedTiles();
 	}
 	else
 	{
-		UpdatedTiles = ProcessTileTasksSync(NumTasksToProcess);
+		UpdatedTiles = ProcessTileTasksSyncAndGetUpdatedTiles(NumTasksToProcess);
 	}
 #endif
 
@@ -6253,14 +6337,14 @@ bool FRecastNavMeshGenerator::IsTimeSliceRegenActive() const
 	return SyncTimeSlicedData.bTimeSliceRegenActive;
 }
 
-bool FRecastNavMeshGenerator::IsTileChanged(int32 TileIdx) const
+bool FRecastNavMeshGenerator::IsTileChanged(const FNavTileRef InTileRef) const
 {
 #if WITH_EDITOR	
 	// Check recently built tiles
-	if (TileIdx > 0)
+	if (InTileRef.IsValid())
 	{
 		FTileTimestamp TileTimestamp;
-		TileTimestamp.TileIdx = static_cast<uint32>(TileIdx);
+		TileTimestamp.NavTileRef = InTileRef;
 		if (RecentlyBuiltTiles.Contains(TileTimestamp))
 		{
 			return true;
diff --git a/Engine/Source/Runtime/NavigationSystem/Public/NavMesh/RecastNavMesh.h b/Engine/Source/Runtime/NavigationSystem/Public/NavMesh/RecastNavMesh.h
--- a/Engine/Source/Runtime/NavigationSystem/Public/NavMesh/RecastNavMesh.h
+++ b/Engine/Source/Runtime/NavigationSystem/Public/NavMesh/RecastNavMesh.h
@@ -243,6 +243,37 @@ struct FRecastDebugGeometry
 	uint32 NAVIGATIONSYSTEM_API GetAllocatedSize() const;
 };
 
+struct FNavTileRef
+{
+	FNavTileRef() {}
+	explicit FNavTileRef(const uint64 InTileRef) : TileRef(InTileRef) {}
+
+	explicit operator uint64() const { return TileRef; }
+
+	bool operator==(const FNavTileRef InRef) const { return TileRef == (uint64)InRef; }
+	bool operator!=(const FNavTileRef InRef) const { return TileRef != (uint64)InRef; }
+
+	bool IsValid() const { return TileRef != (uint64)FNavTileRef(); }
+
+	/** Those 2 functions are used for backward compatibility of the following deprecated functions in FRecastNavMeshGenerator and ARecastNavMesh:
+	*	  RemoveTileLayers
+	*     AddGeneratedTilesTimeSliced
+	*     AddGeneratedTiles
+	*     RemoveLayers
+	*     ProcessTileTasksAsync
+	*     ProcessTileTasks
+	*     AttachTiles
+	*     DetachTiles
+	*     OnNavMeshTilesUpdated
+	*     InvalidateAffectedPaths
+	*   They will be removed with the deprecated methods */
+	static void NAVIGATIONSYSTEM_API DeprecatedGetTileIdsFromNavTileRefs(const FPImplRecastNavMesh* RecastNavMeshImpl, const TArray<FNavTileRef>& InTileRefs, TArray<uint32>& OutTileIds);
+	static void NAVIGATIONSYSTEM_API DeprecatedMakeTileRefsFromTileIds(const FPImplRecastNavMesh* RecastNavMeshImpl, const TArray<uint32>& InTileIds, TArray<FNavTileRef>& OutTileRefs);
+	
+private:	
+	uint64 TileRef = 0;
+};
+
 struct FNavPoly
 {
 	NavNodeRef Ref;
@@ -929,6 +960,8 @@ public:
 
 	const TArray<FIntPoint>& GetActiveTiles() const;
 	TArray<FIntPoint>& GetActiveTiles(); 
+
+	void LogRecastTile(const TCHAR* Caller, const FName& Prefix, const FName& OperationName, const dtNavMesh& DetourMesh, const int32 TileX, const int32 TileY, const int32 LayerIndex, const uint64 TileRef) const;
 	
 protected:
 	/** Serialization helper. */
@@ -990,8 +1023,12 @@ public:
 	void RequestDrawingUpdate(bool bForce = false);
 
 	/** called after regenerating tiles */
+	UE_DEPRECATED(5.1, "Use new version with FNavTileRef")
 	virtual void OnNavMeshTilesUpdated(const TArray<uint32>& ChangedTiles);
 
+	/** called after regenerating tiles */
+	virtual void OnNavMeshTilesUpdated(const TArray<FNavTileRef>& ChangedTiles);
+
 	/** Event from generator that navmesh build has finished */
 	virtual void OnNavMeshGenerationFinished();
 
@@ -1252,8 +1289,12 @@ protected:
 	void UpdatePolyRefBitsPreview();
 	
 	/** Invalidates active paths that go through changed tiles  */
+	UE_DEPRECATED(5.1, "Use new version with FNavTileRef")
 	void InvalidateAffectedPaths(const TArray<uint32>& ChangedTiles);
 
+	/** Invalidates active paths that go through changed tiles  */
+	void InvalidateAffectedPaths(const TArray<FNavTileRef>& ChangedTiles);
+
 	/** created a new FRecastNavMeshGenerator instance. Overrider to supply your
 	 *	own extentions. Note: needs to derive from FRecastNavMeshGenerator */
 	virtual FRecastNavMeshGenerator* CreateGeneratorInstance();
@@ -1264,6 +1305,7 @@ private:
 	friend struct FRecastGraphWrapper;
 	friend FRecastNavMeshGenerator;
 	friend class FPImplRecastNavMesh;
+	friend class URecastNavMeshDataChunk;
 	// destroys FPImplRecastNavMesh instance if it has been created 
 	void DestroyRecastPImpl();
 	// @todo docuement
diff --git a/Engine/Source/Runtime/NavigationSystem/Public/NavMesh/RecastNavMeshDataChunk.h b/Engine/Source/Runtime/NavigationSystem/Public/NavMesh/RecastNavMeshDataChunk.h
--- a/Engine/Source/Runtime/NavigationSystem/Public/NavMesh/RecastNavMeshDataChunk.h
+++ b/Engine/Source/Runtime/NavigationSystem/Public/NavMesh/RecastNavMeshDataChunk.h
@@ -6,6 +6,7 @@
 #include "UObject/ObjectMacros.h"
 #include "AI/Navigation/NavigationTypes.h"
 #include "AI/Navigation/NavigationDataChunk.h"
+#include "NavMesh/RecastNavMesh.h"
 #include "RecastNavMeshDataChunk.generated.h"
 
 class FPImplRecastNavMesh;
@@ -66,17 +67,33 @@ class NAVIGATIONSYSTEM_API URecastNavMeshDataChunk : public UNavigationDataChunk
 	//~ End UObject Interface
 
 	/** Attaches tiles to specified navmesh, transferring tile ownership to navmesh */
+	UE_DEPRECATED(5.1, "Use overload using ARecastNavMesh& instead")
 	TArray<uint32> AttachTiles(FPImplRecastNavMesh& NavMeshImpl);
 
 	/** Attaches tiles to specified navmesh */
+	UE_DEPRECATED(5.1, "Use overload using ARecastNavMesh& instead")
 	TArray<uint32> AttachTiles(FPImplRecastNavMesh& NavMeshImpl, const bool bKeepCopyOfData, const bool bKeepCopyOfCacheData);
 
 	/** Detaches tiles from specified navmesh, taking tile ownership */
+	UE_DEPRECATED(5.1, "Use overload using ARecastNavMesh& instead")
 	TArray<uint32> DetachTiles(FPImplRecastNavMesh& NavMeshImpl);
 
 	/** Detaches tiles from specified navmesh */
+	UE_DEPRECATED(5.1, "Use overload using ARecastNavMesh& instead")
 	TArray<uint32> DetachTiles(FPImplRecastNavMesh& NavMeshImpl, const bool bTakeDataOwnership, const bool bTakeCacheDataOwnership);
 
+	/** Attaches tiles to specified navmesh, transferring tile ownership to navmesh */
+	TArray<FNavTileRef> AttachTiles(ARecastNavMesh& NavMesh);
+
+	/** Attaches tiles to specified navmesh */
+	TArray<FNavTileRef> AttachTiles(ARecastNavMesh& NavMesh, const bool bKeepCopyOfData, const bool bKeepCopyOfCacheData);
+
+	/** Detaches tiles from specified navmesh, taking tile ownership */
+	TArray<FNavTileRef> DetachTiles(ARecastNavMesh& NavMesh);
+
+	/** Detaches tiles from specified navmesh */
+	TArray<FNavTileRef> DetachTiles(ARecastNavMesh& NavMesh, const bool bTakeDataOwnership, const bool bTakeCacheDataOwnership);
+
 	/** 
 	 * Experimental: Moves tiles data on the xy plane by the offset (in tile coordinates) and rotation (in degree).
 	 * @param NavMeshImpl		Recast navmesh implementation.
diff --git a/Engine/Source/Runtime/NavigationSystem/Public/NavMesh/RecastNavMeshGenerator.h b/Engine/Source/Runtime/NavigationSystem/Public/NavMesh/RecastNavMeshGenerator.h
--- a/Engine/Source/Runtime/NavigationSystem/Public/NavMesh/RecastNavMeshGenerator.h
+++ b/Engine/Source/Runtime/NavigationSystem/Public/NavMesh/RecastNavMeshGenerator.h
@@ -581,12 +581,12 @@ struct FRunningTileElement
 
 struct FTileTimestamp
 {
-	uint32 TileIdx;
+	FNavTileRef NavTileRef;
 	double Timestamp;
 	
 	bool operator == (const FTileTimestamp& Other) const
 	{
-		return TileIdx == Other.TileIdx;
+		return NavTileRef == Other.NavTileRef;
 	}
 };
 
@@ -643,7 +643,11 @@ public:
 	virtual int32 GetNumRunningBuildTasks() const override;
 
 	/** Checks if a given tile is being build or has just finished building */
+	UE_DEPRECATED(5.1, "Use new version with FNavTileRef")
 	bool IsTileChanged(int32 TileIdx) const;
+
+	/** Checks if a given tile is being build or has just finished building */
+	bool IsTileChanged(const FNavTileRef InTileRef) const;
 		
 	FORCEINLINE uint32 GetVersion() const { return Version; }
 
@@ -740,7 +744,10 @@ protected:
 	/** Marks all tiles overlapping with InclusionBounds dirty (via MarkDirtyTiles). */
 	bool MarkNavBoundsDirty();
 
+	UE_DEPRECATED(5.1, "Use new version with FNavTileRef")
 	void RemoveLayers(const FIntPoint& Tile, TArray<uint32>& UpdatedTiles);
+
+	void RemoveLayers(const FIntPoint& Tile, TArray<FNavTileRef>& UpdatedTiles);
 	
 	void StoreCompressedTileCacheLayers(const FRecastTileGenerator& TileGenerator, int32 TileX, int32 TileY);
 
@@ -750,28 +757,56 @@ protected:
 
 #if RECAST_ASYNC_REBUILDING
 	/** Processes pending tile generation tasks Async*/
+	UE_DEPRECATED(5.1, "Use ProcessTileTasksAsyncAndGetUpdatedTiles instead")
 	TArray<uint32> ProcessTileTasksAsync(const int32 NumTasksToProcess);
+
+	/** Processes pending tile generation tasks Async*/
+	TArray<FNavTileRef> ProcessTileTasksAsyncAndGetUpdatedTiles(const int32 NumTasksToProcess);
 #else
 	TSharedRef<FRecastTileGenerator> CreateTileGeneratorFromPendingElement(FIntPoint &OutTileLocation);
+
 	/** Processes pending tile generation tasks Sync with option for time slicing currently an experimental feature. */
+	UE_DEPRECATED(5.1, "Use ProcessTileTasksSyncTimeSlicedAndGetUpdatedTiles instead")
 	TArray<uint32> ProcessTileTasksSyncTimeSliced();
+	UE_DEPRECATED(5.1, "Use ProcessTileTasksSyncAndGetUpdatedTiles instead")
 	TArray<uint32> ProcessTileTasksSync(const int32 NumTasksToProcess);
+
+	/** Processes pending tile generation tasks Sync with option for time slicing currently an experimental feature. */
+	TArray<FNavTileRef> ProcessTileTasksSyncTimeSlicedAndGetUpdatedTiles();
+	TArray<FNavTileRef> ProcessTileTasksSyncAndGetUpdatedTiles(const int32 NumTasksToProcess);
 #endif
 	/** Processes pending tile generation tasks */
+	UE_DEPRECATED(5.1, "Use ProcessTileTasksAndGetUpdatedTiles instead")
 	TArray<uint32> ProcessTileTasks(const int32 NumTasksToProcess);
 
+	/** Processes pending tile generation tasks */
+	TArray<FNavTileRef> ProcessTileTasksAndGetUpdatedTiles(const int32 NumTasksToProcess);
+
 	void ResetTimeSlicedTileGeneratorSync();
 
 public:
 	/** Adds generated tiles to NavMesh, replacing old ones, uses time slicing returns Failed if any layer failed */
+	UE_DEPRECATED(5.1, "Use new version with FNavTileRef")
 	ETimeSliceWorkResult AddGeneratedTilesTimeSliced(FRecastTileGenerator& TileGenerator, TArray<uint32>& OutResultTileIndices);
+
+	/** Adds generated tiles to NavMesh, replacing old ones, uses time slicing returns Failed if any layer failed */
+	ETimeSliceWorkResult AddGeneratedTilesTimeSliced(FRecastTileGenerator& TileGenerator, TArray<FNavTileRef>& OutResultTileRefs);
+
 	/** Adds generated tiles to NavMesh, replacing old ones */
+	UE_DEPRECATED(5.1, "Use AddGeneratedTilesAndGetUpdatedTiles instead")
 	TArray<uint32> AddGeneratedTiles(FRecastTileGenerator& TileGenerator);
 
+	/** Adds generated tiles to NavMesh, replacing old ones */
+	TArray<FNavTileRef> AddGeneratedTilesAndGetUpdatedTiles(FRecastTileGenerator& TileGenerator);
+
 public:
 	/** Removes all tiles at specified grid location */
+	UE_DEPRECATED(5.1, "Use RemoveTileLayersAndGetUpdatedTiles instead")
 	TArray<uint32> RemoveTileLayers(const int32 TileX, const int32 TileY, TMap<int32, dtPolyRef>* OldLayerTileIdMap = nullptr);
 
+	/** Removes all tiles at specified grid location */
+	TArray<FNavTileRef> RemoveTileLayersAndGetUpdatedTiles(const int32 TileX, const int32 TileY, TMap<int32, dtPolyRef>* OldLayerTileIdMap = nullptr);
+
 	void RemoveTiles(const TArray<FIntPoint>& Tiles);
 	void ReAddTiles(const TArray<FIntPoint>& Tiles);
 
@@ -808,8 +843,11 @@ protected:
 	//----------------------------------------------------------------------//
 	virtual uint32 LogMemUsed() const override;
 
+	UE_DEPRECATED(5.1, "Use new version with FNavTileRef")
 	void AddGeneratedTileLayer(int32 LayerIndex, FRecastTileGenerator& TileGenerator, const TMap<int32, dtPolyRef>& OldLayerTileIdMap, TArray<uint32>& OutResultTileIndices);
 
+	void AddGeneratedTileLayer(int32 LayerIndex, FRecastTileGenerator& TileGenerator, const TMap<int32, dtPolyRef>& OldLayerTileIdMap, TArray<FNavTileRef>& OutResultTileRefs);
+
 #if !UE_BUILD_SHIPPING
 	/** Data struct used by 'LogDirtyAreas' that contains all the information regarding the areas that are being dirtied, per dirtied tile. */
 	struct FNavigationDirtyAreaPerTileDebugInformation
@@ -895,10 +933,10 @@ protected:
 		EProcessTileTasksSyncTimeSlicedState ProcessTileTasksSyncState;
 
 		/** Used by ProcessTileTasksSyncTimeSliced */
-		TArray<uint32> UpdatedTilesCache;
+		TArray<FNavTileRef> UpdatedTilesCache;
 
 		/** Used by AddGeneratedTilesTimeSliced */
-		TArray<uint32> ResultTileIndicesCached;
+		TArray<FNavTileRef> ResultTileRefsCached;
 
 		/** Used by AddGeneratedTilesTimeSliced */
 		TMap<int32, dtPolyRef> OldLayerTileIdMapCached;
