Build fix for 'GetTileIdsFromNavTileRefs': is not a member of 'FNavTileRef'
#rnx

diff --git a/Engine/Source/Runtime/NavigationSystem/Private/NavMesh/RecastNavMeshGenerator.cpp b/Engine/Source/Runtime/NavigationSystem/Private/NavMesh/RecastNavMeshGenerator.cpp
--- a/Engine/Source/Runtime/NavigationSystem/Private/NavMesh/RecastNavMeshGenerator.cpp
+++ b/Engine/Source/Runtime/NavigationSystem/Private/NavMesh/RecastNavMeshGenerator.cpp
@@ -5870,7 +5870,7 @@ TArray<uint32> FRecastNavMeshGenerator::ProcessTileTasksSyncTimeSliced()
 {
 	const TArray<FNavTileRef>& TileRefs = ProcessTileTasksSyncTimeSlicedAndGetUpdatedTiles();
 	TArray<uint32> TileIds;
-	FNavTileRef::GetTileIdsFromNavTileRefs(DestNavMesh->GetRecastNavMeshImpl(), TileRefs, TileIds);
+	FNavTileRef::DeprecatedGetTileIdsFromNavTileRefs(DestNavMesh->GetRecastNavMeshImpl(), TileRefs, TileIds);
 	return TileIds;
 }
 
@@ -6062,7 +6062,7 @@ TArray<uint32> FRecastNavMeshGenerator::ProcessTileTasksSync(const int32 NumTask
 {
 	const TArray<FNavTileRef>& TileRefs = ProcessTileTasksSyncAndGetUpdatedTiles(NumTasksToProcess);
 	TArray<uint32> TileIds;
-	FNavTileRef::GetTileIdsFromNavTileRefs(DestNavMesh->GetRecastNavMeshImpl(), TileRefs, TileIds);
+	FNavTileRef::DeprecatedGetTileIdsFromNavTileRefs(DestNavMesh->GetRecastNavMeshImpl(), TileRefs, TileIds);
 	return TileIds;
 }
 	
