PR #9020: Fixed potential crash in ARecastNavMesh::InvalidateAffectedPaths (Contributed by [REDACTED_PERSON])

Added same pattern to ANavigationData::PurgeUnusedPaths

diff --git a/Engine/Source/Runtime/NavigationSystem/Private/NavMesh/RecastNavMesh.cpp b/Engine/Source/Runtime/NavigationSystem/Private/NavMesh/RecastNavMesh.cpp
--- a/Engine/Source/Runtime/NavigationSystem/Private/NavMesh/RecastNavMesh.cpp
+++ b/Engine/Source/Runtime/NavigationSystem/Private/NavMesh/RecastNavMesh.cpp
@@ -2196,9 +2196,9 @@ void ARecastNavMesh::InvalidateAffectedPaths(const TArray<FNavTileRef>& ChangedT
 	{
 		FScopeLock PathLock(&ActivePathsLock);
 
-		FNavPathWeakPtr* WeakPathPtr = (ActivePaths.GetData() + PathsCount - 1);
-		for (int32 PathIndex = PathsCount - 1; PathIndex >= 0; --PathIndex, --WeakPathPtr)
+		for (int32 PathIndex = PathsCount - 1; PathIndex >= 0; --PathIndex)
 		{
+			FNavPathWeakPtr* WeakPathPtr = &ActivePaths[PathIndex];
 			FNavPathSharedPtr SharedPath = WeakPathPtr->Pin();
 			if (WeakPathPtr->IsValid() == false)
 			{
diff --git a/Engine/Source/Runtime/NavigationSystem/Private/NavigationData.cpp b/Engine/Source/Runtime/NavigationSystem/Private/NavigationData.cpp
--- a/Engine/Source/Runtime/NavigationSystem/Private/NavigationData.cpp
+++ b/Engine/Source/Runtime/NavigationSystem/Private/NavigationData.cpp
@@ -425,12 +425,12 @@ void ANavigationData::PurgeUnusedPaths()
 	FScopeLock PathLock(&ActivePathsLock);
 
 	const int32 Count = ActivePaths.Num();
-	FNavPathWeakPtr* WeakPathPtr = (ActivePaths.GetData() + Count - 1);
-	for (int32 i = Count - 1; i >= 0; --i, --WeakPathPtr)
+	for (int32 PathIndex = Count - 1; PathIndex >= 0; --PathIndex)
 	{
+		FNavPathWeakPtr* WeakPathPtr = &ActivePaths[PathIndex];
 		if (WeakPathPtr->IsValid() == false)
 		{
-			ActivePaths.RemoveAtSwap(i, 1, /*bAllowShrinking=*/false);
+			ActivePaths.RemoveAtSwap(PathIndex, 1, /*bAllowShrinking=*/false);
 		}
 	}
 }
