BuildFix: missing file for CL20702220
add accessor to TimeSlicer Remaining Duration

diff --git a/Engine/Source/Runtime/NavigationSystem/Public/NavigationSystem.h b/Engine/Source/Runtime/NavigationSystem/Public/NavigationSystem.h
--- a/Engine/Source/Runtime/NavigationSystem/Public/NavigationSystem.h
+++ b/Engine/Source/Runtime/NavigationSystem/Public/NavigationSystem.h
@@ -122,6 +122,7 @@ public:
 
 	//* Returns the cached result of calling TestTimeSliceFinished, false by default */
 	bool IsTimeSliceFinishedCached() const { return bTimeSliceFinishedCached; }
+	double GetRemainingDuration() const { return RemainingDuration; }
 	double GetRemainingDurationFraction() const { return OriginalDuration > 0. ? RemainingDuration / OriginalDuration : 0.; }
 
 protected:
