[Navmesh] Misc: merge namespace and remove deprecated log
#rnx

diff --git a/Engine/Source/Runtime/NavigationSystem/Private/NavMesh/RecastNavMesh.cpp b/Engine/Source/Runtime/NavigationSystem/Private/NavMesh/RecastNavMesh.cpp
--- a/Engine/Source/Runtime/NavigationSystem/Private/NavMesh/RecastNavMesh.cpp
+++ b/Engine/Source/Runtime/NavigationSystem/Private/NavMesh/RecastNavMesh.cpp
@@ -47,7 +47,7 @@
 
 #endif // WITH_RECAST
 
-namespace 
+namespace UE::NavMesh::Private
 {
 	// Max tile size in voxels. Larger than this tiles will start to get slow to build.
 	constexpr int32 ArbitraryMaxTileSizeVoxels = 1024;
@@ -76,10 +76,6 @@ namespace
 		return FMath::Clamp(CellSize, ArbitraryMinCellSize, ArbitraryMaxCellSize);
 	}
 
-}
-
-namespace UE::NavMesh::Private
-{
 	FNavTileRef GetTileRefFromPolyRef(const dtNavMesh& DetourMesh, const NavNodeRef PolyRef)
 	{
 		unsigned int Salt = 0;
@@ -2877,8 +2873,8 @@ void ARecastNavMesh::PostEditChangeProperty(FPropertyChangedEvent& PropertyChang
 			}
 			else if (PropName == GET_MEMBER_NAME_CHECKED(ARecastNavMesh, TileSizeUU))
 			{
-				CellSize = GetClampedCellSize(CellSize);
-				TileSizeUU = GetClampedTileSizeUU(TileSizeUU, CellSize, AgentRadius);
+				CellSize = UE::NavMesh::Private::GetClampedCellSize(CellSize);
+				TileSizeUU = UE::NavMesh::Private::GetClampedTileSizeUU(TileSizeUU, CellSize, AgentRadius);
 						
 				// Match cell size to tile size.
 				CellSize = TileSizeUU / FMath::TruncToInt(TileSizeUU / CellSize);
@@ -2888,8 +2884,8 @@ void ARecastNavMesh::PostEditChangeProperty(FPropertyChangedEvent& PropertyChang
 			}
 			else if (PropName == GET_MEMBER_NAME_CHECKED(ARecastNavMesh, CellSize))
 			{
-				CellSize = GetClampedCellSize(CellSize);
-				TileSizeUU = GetClampedTileSizeUU(TileSizeUU, CellSize, AgentRadius);
+				CellSize = UE::NavMesh::Private::GetClampedCellSize(CellSize);
+				TileSizeUU = UE::NavMesh::Private::GetClampedTileSizeUU(TileSizeUU, CellSize, AgentRadius);
 				
 				// Match tile size to cell size.
 				TileSizeUU = CellSize * FMath::TruncToInt(TileSizeUU / CellSize);
diff --git a/Engine/Source/Runtime/NavigationSystem/Private/NavigationSystem.cpp b/Engine/Source/Runtime/NavigationSystem/Private/NavigationSystem.cpp
--- a/Engine/Source/Runtime/NavigationSystem/Private/NavigationSystem.cpp
+++ b/Engine/Source/Runtime/NavigationSystem/Private/NavigationSystem.cpp
@@ -3427,7 +3427,6 @@ void UNavigationSystemV1::Build()
 	TRACE_CPUPROFILER_EVENT_SCOPE(UNavigationSystemV1::Build);
 	
 	UE_LOG(LogNavigationDataBuild, Display, TEXT("UNavigationSystemV1::Build started..."));
-	UE_LOG(LogNavigationDataBuild, Display, TEXT("   Building navigation data using Chaos physics."));
 
 	UWorld* World = GetWorld();
 	if (!World)
