[WP Dynamic Navmesh] Fix missing tile in base navmesh when building WP navmesh from editor

diff --git a/Engine/Source/Runtime/NavigationSystem/Private/NavMesh/RecastNavMeshGenerator.cpp b/Engine/Source/Runtime/NavigationSystem/Private/NavMesh/RecastNavMeshGenerator.cpp
--- a/Engine/Source/Runtime/NavigationSystem/Private/NavMesh/RecastNavMeshGenerator.cpp
+++ b/Engine/Source/Runtime/NavigationSystem/Private/NavMesh/RecastNavMeshGenerator.cpp
@@ -5116,6 +5116,12 @@ void FRecastNavMeshGenerator::AddGeneratedTileLayer(int32 LayerIndex, FRecastTil
 	}
 }
 
+bool FRecastNavMeshGenerator::IsAllowedToAddTileLayers(const FIntPoint Tile) const
+{
+	check(DestNavMesh);
+	return !DestNavMesh->IsWorldPartitionedDynamicNavmesh() || IsInActiveSet(Tile) || (GetWorld() && !GetWorld()->IsGameWorld());
+}
+
 #if !UE_BUILD_SHIPPING
 void FRecastNavMeshGenerator::LogDirtyAreas(
 	const TMap<FPendingTileElement, TArray<FNavigationDirtyAreaPerTileDebugInformation>>& DirtyAreasDebuggingInformation) const
@@ -5210,7 +5216,7 @@ ETimeSliceWorkResult FRecastNavMeshGenerator::AddGeneratedTilesTimeSliced(FRecas
 	case EAddGeneratedTilesTimeSlicedState::AddTiles:
 	{
 		if (DetourMesh != nullptr
-			&& (!DestNavMesh->IsWorldPartitionedDynamicNavmesh() || IsInActiveSet(FIntPoint(TileX, TileY)))
+			&& IsAllowedToAddTileLayers(FIntPoint(TileX, TileY))
 			&& SyncTimeSlicedData.AddGenTilesLayerIndex != INDEX_NONE)
 		{
 			for (; SyncTimeSlicedData.AddGenTilesLayerIndex < TileGenerator.GetDirtyLayersMask().Num(); ++SyncTimeSlicedData.AddGenTilesLayerIndex)
@@ -5283,7 +5289,7 @@ TArray<FNavTileRef> FRecastNavMeshGenerator::AddGeneratedTilesAndGetUpdatedTiles
 	const int32 FirstDirtyTileIndex = TileGenerator.GetDirtyLayersMask().Find(true);
 
 	if (DetourMesh != nullptr
-		&& (!DestNavMesh->IsWorldPartitionedDynamicNavmesh() || IsInActiveSet(FIntPoint(TileX, TileY)))
+		&& IsAllowedToAddTileLayers(FIntPoint(TileX, TileY))
 		&& FirstDirtyTileIndex != INDEX_NONE)
 	{
 		TArray<FNavMeshTileData> TileLayers = TileGenerator.GetNavigationData();
diff --git a/Engine/Source/Runtime/NavigationSystem/Public/NavMesh/RecastNavMeshGenerator.h b/Engine/Source/Runtime/NavigationSystem/Public/NavMesh/RecastNavMeshGenerator.h
--- a/Engine/Source/Runtime/NavigationSystem/Public/NavMesh/RecastNavMeshGenerator.h
+++ b/Engine/Source/Runtime/NavigationSystem/Public/NavMesh/RecastNavMeshGenerator.h
@@ -850,6 +850,7 @@ protected:
 	UE_DEPRECATED(5.1, "Use new version with FNavTileRef")
 	void AddGeneratedTileLayer(int32 LayerIndex, FRecastTileGenerator& TileGenerator, const TMap<int32, dtPolyRef>& OldLayerTileIdMap, TArray<uint32>& OutResultTileIndices);
 
+	bool IsAllowedToAddTileLayers(const FIntPoint Tile) const;
 	void AddGeneratedTileLayer(int32 LayerIndex, FRecastTileGenerator& TileGenerator, const TMap<int32, dtPolyRef>& OldLayerTileIdMap, TArray<FNavTileRef>& OutResultTileRefs);
 
 #if !UE_BUILD_SHIPPING
