Added debugging code to better track down slow time sliced sections of code.

diff --git a/Engine/Source/Runtime/NavigationSystem/Private/NavMesh/RecastNavMeshGenerator.cpp b/Engine/Source/Runtime/NavigationSystem/Private/NavMesh/RecastNavMeshGenerator.cpp
--- a/Engine/Source/Runtime/NavigationSystem/Private/NavMesh/RecastNavMeshGenerator.cpp
+++ b/Engine/Source/Runtime/NavigationSystem/Private/NavMesh/RecastNavMeshGenerator.cpp
@@ -1879,6 +1879,8 @@ ETimeSliceWorkResult FRecastTileGenerator::GatherGeometryFromSourcesTimeSliced()
 
 		GatherNavigationDataGeometry(ElementData, *NavSys, NavDataConfig, bUpdateGeometry);
 		
+		MARK_TIMESLICE_SECTION_DEBUG(TimeSliceManager->GetTimeSlicer(), GatherGeometryFromSources);
+
 		if (TimeSliceManager->GetTimeSlicer().TestTimeSliceFinished())
 		{
 			return ETimeSliceWorkResult::CallAgainNextTimeSlice;
@@ -2520,6 +2522,8 @@ ETimeSliceWorkResult FRecastTileGenerator::RasterizeGeometryRecastTimeSliced(FNa
 
 		RasterizeGeomRecastState = ERasterizeGeomRecastTimeSlicedState::RasterizeTriangles;
 
+		MARK_TIMESLICE_SECTION_DEBUG(TimeSliceManager->GetTimeSlicer(), MarkWalkableTriangles);
+
 		if (TimeSliceManager->GetTimeSlicer().TestTimeSliceFinished())
 		{
 			return ETimeSliceWorkResult::CallAgainNextTimeSlice;
@@ -2542,6 +2546,8 @@ ETimeSliceWorkResult FRecastTileGenerator::RasterizeGeometryRecastTimeSliced(FNa
 		//reset this so next call we start by marking walkable triangles
 		RasterizeGeomRecastState = ERasterizeGeomRecastTimeSlicedState::MarkWalkableTriangles;
 
+		MARK_TIMESLICE_SECTION_DEBUG(TimeSliceManager->GetTimeSlicer(), RasterizeTriangles);
+
 		TimeSliceManager->GetTimeSlicer().TestTimeSliceFinished();
 	}
 	break;
@@ -2647,6 +2653,8 @@ ETimeSliceWorkResult FRecastTileGenerator::RasterizeGeometryTimeSliced(FNavMeshB
 
 		RasterizeGeomState = ERasterizeGeomTimeSlicedState::RasterizeGeometryRecast;
 
+		MARK_TIMESLICE_SECTION_DEBUG(TimeSliceManager->GetTimeSlicer(), RasterizeGeometryTransformCoords);
+
 		if (TimeSliceManager->GetTimeSlicer().TestTimeSliceFinished())
 		{
 			return ETimeSliceWorkResult::CallAgainNextTimeSlice;
@@ -3044,6 +3052,8 @@ ETimeSliceWorkResult FRecastTileGenerator::GenerateCompressedLayersTimeSliced(FN
 
 		GenCompressedLayersTimeSlicedState = EGenerateCompressedLayersTimeSliced::RasterizeTriangles;
 
+		MARK_TIMESLICE_SECTION_DEBUG(TimeSliceManager->GetTimeSlicer(), CreateHeightField);
+
 		if (TimeSliceManager->GetTimeSlicer().TestTimeSliceFinished())
 		{
 			return ETimeSliceWorkResult::CallAgainNextTimeSlice;
@@ -3086,6 +3096,8 @@ ETimeSliceWorkResult FRecastTileGenerator::GenerateCompressedLayersTimeSliced(FN
 		{
 			ApplyVoxelFilter(RasterContext->SolidHF, TileConfig.walkableRadius);
 
+			MARK_TIMESLICE_SECTION_DEBUG(TimeSliceManager->GetTimeSlicer(), VoxelFilter);
+
 			if (TimeSliceManager->GetTimeSlicer().TestTimeSliceFinished())
 			{
 				return ETimeSliceWorkResult::CallAgainNextTimeSlice;
@@ -3098,6 +3110,8 @@ ETimeSliceWorkResult FRecastTileGenerator::GenerateCompressedLayersTimeSliced(FN
 
 		GenCompressedLayersTimeSlicedState = EGenerateCompressedLayersTimeSliced::CompactHeightField;
 
+		MARK_TIMESLICE_SECTION_DEBUG(TimeSliceManager->GetTimeSlicer(), RecastFilterLedgeSpans);
+
 		if (TimeSliceManager->GetTimeSlicer().TestTimeSliceFinished())
 		{
 			return ETimeSliceWorkResult::CallAgainNextTimeSlice;
@@ -3115,6 +3129,8 @@ ETimeSliceWorkResult FRecastTileGenerator::GenerateCompressedLayersTimeSliced(FN
 
 		GenCompressedLayersTimeSlicedState = EGenerateCompressedLayersTimeSliced::ErodeWalkable;
 
+		MARK_TIMESLICE_SECTION_DEBUG(TimeSliceManager->GetTimeSlicer(), CompactHeightField);
+
 		if (TimeSliceManager->GetTimeSlicer().TestTimeSliceFinished())
 		{
 			return ETimeSliceWorkResult::CallAgainNextTimeSlice;
@@ -3132,6 +3148,8 @@ ETimeSliceWorkResult FRecastTileGenerator::GenerateCompressedLayersTimeSliced(FN
 
 		GenCompressedLayersTimeSlicedState = EGenerateCompressedLayersTimeSliced::BuildLayers;
 
+		MARK_TIMESLICE_SECTION_DEBUG(TimeSliceManager->GetTimeSlicer(), ErodeWalkable);
+
 		if (TimeSliceManager->GetTimeSlicer().TestTimeSliceFinished())
 		{
 			return ETimeSliceWorkResult::CallAgainNextTimeSlice;
@@ -3141,6 +3159,8 @@ ETimeSliceWorkResult FRecastTileGenerator::GenerateCompressedLayersTimeSliced(FN
 	{
 		const bool bRecastBuildLayers = RecastBuildLayers(BuildContext, *RasterContext);
 
+		MARK_TIMESLICE_SECTION_DEBUG(TimeSliceManager->GetTimeSlicer(), BuildLayers);
+
 		//this could have done a fair amount of work either way so check time slice
 		TimeSliceManager->GetTimeSlicer().TestTimeSliceFinished();
 
@@ -3164,6 +3184,8 @@ ETimeSliceWorkResult FRecastTileGenerator::GenerateCompressedLayersTimeSliced(FN
 
 		const bool bRecastBuildTileCache = RecastBuildTileCache(BuildContext, *RasterContext);
 	
+		MARK_TIMESLICE_SECTION_DEBUG(TimeSliceManager->GetTimeSlicer(), BuildTileCache);
+
 		//this could have done a fair amount of work either way so check time slice
 		TimeSliceManager->GetTimeSlicer().TestTimeSliceFinished();
 
@@ -3658,6 +3680,8 @@ ETimeSliceWorkResult FRecastTileGenerator::GenerateNavigationDataTimeSliced(FNav
 
 			const bool bGenDataLayer = GenerateNavigationDataLayer(BuildContext, TileCompressor, *GenNavDataTimeSlicedAllocator, *GenNavDataTimeSlicedGenerationContext, GenNavDataLayerTimeSlicedIdx);
 
+			MARK_TIMESLICE_SECTION_DEBUG(TimeSliceManager->GetTimeSlicer(), GenerateLayers);
+
 			//carry on iterating but don't do any more work if the time slice is finished (as we may not need to in which case we can avoid calling this function again)
 			TimeSliceManager->GetTimeSlicer().TestTimeSliceFinished();
 
@@ -5231,6 +5255,8 @@ ETimeSliceWorkResult FRecastNavMeshGenerator::AddGeneratedTilesTimeSliced(FRecas
 
 					AddGeneratedTileLayer(SyncTimeSlicedData.AddGenTilesLayerIndex, TileGenerator, SyncTimeSlicedData.OldLayerTileIdMapCached, SyncTimeSlicedData.ResultTileRefsCached);
 
+					MARK_TIMESLICE_SECTION_DEBUG(SyncTimeSlicedData.TimeSliceManager->GetTimeSlicer(), AddTiles);
+
 					SyncTimeSlicedData.TimeSliceManager->GetTimeSlicer().TestTimeSliceFinished();
 				}
 			}
@@ -5955,6 +5981,16 @@ TArray<FNavTileRef> FRecastNavMeshGenerator::ProcessTileTasksSyncTimeSlicedAndGe
 
 				SyncTimeSlicedData.TileGeneratorSync = CreateTileGeneratorFromPendingElement(TileLocation, NextPendingDirtyTileIndex);
 
+				check(SyncTimeSlicedData.TileGeneratorSync);
+
+#if ALLOW_TIME_SLICE_DEBUG
+				UNavigationSystemV1* NavSys = FNavigationSystem::GetCurrent<UNavigationSystemV1>(GetWorld());
+
+				check(NavSys);
+
+				SyncTimeSlicedData.TimeSliceManager->GetTimeSlicer().SetCurrentTileDebug(*NavSys, SyncTimeSlicedData.TimeSliceManager->GetNavDataIdx(), TileLocation, SyncTimeSlicedData.TileGeneratorSync->GetTileBB());
+#endif
+
 				SyncTimeSlicedData.CurrentTileRegenDuration = 0.;
 
 				if (SyncTimeSlicedData.TileGeneratorSync->HasDataToBuild())
@@ -5971,6 +6007,8 @@ TArray<FNavTileRef> FRecastNavMeshGenerator::ProcessTileTasksSyncTimeSlicedAndGe
 					}
 				}
 
+				MARK_TIMESLICE_SECTION_DEBUG(SyncTimeSlicedData.TimeSliceManager->GetTimeSlicer(), CreateTileGenerator);
+
 				if (SyncTimeSlicedData.TimeSliceManager->GetTimeSlicer().TestTimeSliceFinished())
 				{
 					return EndFunction(true /* bCalcTileRegenDuration */, bHadWorkToDo);
@@ -6048,6 +6086,8 @@ TArray<FNavTileRef> FRecastNavMeshGenerator::ProcessTileTasksSyncTimeSlicedAndGe
 
 				SyncTimeSlicedData.CurrentTileRegenDuration = 0.;
 				
+				MARK_TIMESLICE_SECTION_DEBUG(SyncTimeSlicedData.TimeSliceManager->GetTimeSlicer(), FinishTile);
+
 				//test time slice 
 				if (SyncTimeSlicedData.TimeSliceManager->GetTimeSlicer().TestTimeSliceFinished())
 				{
diff --git a/Engine/Source/Runtime/NavigationSystem/Private/NavigationSystem.cpp b/Engine/Source/Runtime/NavigationSystem/Private/NavigationSystem.cpp
--- a/Engine/Source/Runtime/NavigationSystem/Private/NavigationSystem.cpp
+++ b/Engine/Source/Runtime/NavigationSystem/Private/NavigationSystem.cpp
@@ -268,11 +268,44 @@ void FNavRegenTimeSlicer::EndTimeSliceAndAdjustDuration()
 	RemainingDuration = FMath::Max(RemainingDuration - (TimeLastTested - StartTime), 0.);
 }
 
+#if ALLOW_TIME_SLICE_DEBUG
+void FNavRegenTimeSlicer::SetCurrentTileDebug(UNavigationSystemV1& InNavSystem, const int InNavDataIdx, const FIntPoint& InTile, const FBox& InTileBox) const
+{
+	NavSystemDebug = &InNavSystem;
+	TileDebug = InTile;
+	TileBoxDebug = InTileBox;
+	NaxDataIdxDebug = InNavDataIdx;
+}
+#endif // ALLOW_TIME_SLICE_DEBUG
+
 bool FNavRegenTimeSlicer::TestTimeSliceFinished() const
 {
 	ensureMsgf(!bTimeSliceFinishedCached, TEXT("Testing time slice is finished when we have already confirmed that!"));
 
-	TimeLastTested = FPlatformTime::Seconds();
+	const double Time = FPlatformTime::Seconds();
+
+#if ALLOW_TIME_SLICE_DEBUG
+	check(NavSystemDebug.Get());
+	check(NavSystemDebug->NavDataSet.IsValidIndex(NaxDataIdxDebug));
+
+	ARecastNavMesh* NavData = Cast<ARecastNavMesh>(NavSystemDebug->NavDataSet[NaxDataIdxDebug]);
+	check(NavData);
+
+	const double TimeSinceLastTested = Time - TimeLastTested;
+	if (TimeSinceLastTested >= NavData->TimeSliceLongDurationDebug)
+	{
+		const FVector Pos = TileBoxDebug.GetCenter();
+
+		// I'd quite like to make this a Warning, but it would be too frequently logged as things stand.
+		UE_LOG(LogNavigation, Verbose, TEXT("Nav mesh data: %s, tile at %d, %d, coordinate %f, %f, %f: %s is taking %f secs to partially regenerate!"), *NavData->GetName(), TileDebug.X, TileDebug.Y, Pos.X, Pos.Y, Pos.Z, *SectionNameDebug.ToString(), TimeSinceLastTested);
+		UE_VLOG_BOX(NavData, LogNavigation, Verbose, TileBoxDebug, FColor::Red, TEXT("Nav mesh data : %s,  tile at %d, %d, %s is taking %f secs to partially regenerate!"), *NavData->GetName(), TileDebug.X, TileDebug.Y, *SectionNameDebug.ToString(), TimeSinceLastTested);
+	}
+
+	// Reset SectionDebugName
+	SectionNameDebug = FNavigationSystem::TimeSliceDefaultSectionNameDebug;
+#endif // ALLOW_TIME_SLICE_DEBUG
+
+	TimeLastTested = Time;
 
 	bTimeSliceFinishedCached = (TimeLastTested - StartTime) >= RemainingDuration;
 	return bTimeSliceFinishedCached;
diff --git a/Engine/Source/Runtime/NavigationSystem/Public/NavMesh/RecastNavMesh.h b/Engine/Source/Runtime/NavigationSystem/Public/NavMesh/RecastNavMesh.h
--- a/Engine/Source/Runtime/NavigationSystem/Public/NavMesh/RecastNavMesh.h
+++ b/Engine/Source/Runtime/NavigationSystem/Public/NavMesh/RecastNavMesh.h
@@ -815,6 +815,10 @@ class NAVIGATIONSYSTEM_API ARecastNavMesh : public ANavigationData
 	UPROPERTY(config)
 	uint32 bAllowNavLinkAsPathEnd : 1;
 
+	/** If a single time sliced section of navmesh regen code exceeds this duration then it will trigger debug logging */
+	UPROPERTY(EditAnywhere, Category = TimeSlicing, config, AdvancedDisplay)
+	double TimeSliceLongDurationDebug = 0.002;
+
 private:
 	/** Cache rasterized voxels instead of just collision vertices/indices in navigation octree */
 	UPROPERTY(config)
diff --git a/Engine/Source/Runtime/NavigationSystem/Public/NavMesh/RecastNavMeshGenerator.h b/Engine/Source/Runtime/NavigationSystem/Public/NavMesh/RecastNavMeshGenerator.h
--- a/Engine/Source/Runtime/NavigationSystem/Public/NavMesh/RecastNavMeshGenerator.h
+++ b/Engine/Source/Runtime/NavigationSystem/Public/NavMesh/RecastNavMeshGenerator.h
@@ -282,6 +282,7 @@ public:
 
 	FORCEINLINE int32 GetTileX() const { return TileX; }
 	FORCEINLINE int32 GetTileY() const { return TileY; }
+	FORCEINLINE const FBox& GetTileBB() const { return TileBB; }
 	/** Whether specified layer was updated */
 	FORCEINLINE bool IsLayerChanged(int32 LayerIdx) const { return DirtyLayers[LayerIdx]; }
 	FORCEINLINE const TBitArray<>& GetDirtyLayersMask() const { return DirtyLayers; }
diff --git a/Engine/Source/Runtime/NavigationSystem/Public/NavigationSystem.h b/Engine/Source/Runtime/NavigationSystem/Public/NavigationSystem.h
--- a/Engine/Source/Runtime/NavigationSystem/Public/NavigationSystem.h
+++ b/Engine/Source/Runtime/NavigationSystem/Public/NavigationSystem.h
@@ -40,6 +40,12 @@ class UNavigationSystemModuleConfig;
 struct FNavigationRelevantData;
 struct FNavigationOctreeElement;
 
+#if !UE_BUILD_SHIPPING
+#define ALLOW_TIME_SLICE_DEBUG 1
+#else
+#define ALLOW_TIME_SLICE_DEBUG 0
+#endif
+
 /** delegate to let interested parties know that new nav area class has been registered */
 DECLARE_MULTICAST_DELEGATE_OneParam(FOnNavAreaChanged, const UClass* /*AreaClass*/);
 
@@ -76,6 +82,10 @@ namespace FNavigationSystem
 	bool NAVIGATIONSYSTEM_API ShouldLoadNavigationOnClient(ANavigationData& NavData);
 
 	void NAVIGATIONSYSTEM_API MakeAllComponentsNeverAffectNav(AActor& Actor);
+
+#if ALLOW_TIME_SLICE_DEBUG
+	const FName TimeSliceDefaultSectionNameDebug= FName(TEXT("DefaultSection"));
+#endif // ALLOW_TIME_SLICE_DEBUG
 }
 
 struct FNavigationSystemExec: public FSelfRegisteringExec
@@ -97,6 +107,15 @@ namespace ENavigationBuildLock
 }
 
 
+// Use this just  before a call to TestTimeSliceFinished() to allow the time slice logging to use a debug name for the section of code being timed.
+#if ALLOW_TIME_SLICE_DEBUG
+#define MARK_TIMESLICE_SECTION_DEBUG(TIME_SLICER, TIME_SLICE_FNAME) \
+static const FName TIME_SLICE_FNAME(TEXT(#TIME_SLICE_FNAME)); \
+TIME_SLICER.SetSectionNameDebug(TIME_SLICE_FNAME);
+#else
+#define MARK_TIMESLICE_SECTION_DEBUG(TIME_SLICER, TIME_SLICE_FNAME) ;
+#endif
+
 class NAVIGATIONSYSTEM_API FNavRegenTimeSlicer
 {
 public:
@@ -125,12 +144,34 @@ public:
 	double GetRemainingDuration() const { return RemainingDuration; }
 	double GetRemainingDurationFraction() const { return OriginalDuration > 0. ? RemainingDuration / OriginalDuration : 0.; }
 
+
+#if ALLOW_TIME_SLICE_DEBUG
+	void SetCurrentTileDebug(UNavigationSystemV1& InNavSystem, const int InNavDataIdx, const FIntPoint& InTile, const FBox& InTileBox) const;
+
+	/**
+	 *  Sets the debug name for a time sliced section of code.
+	 *  Do not call this directly use MARK_TIMESLICE_SECTION_DEBUG
+	 */
+	void SetSectionNameDebug(FName InSectionNameDebug) const
+	{
+		SectionNameDebug = InSectionNameDebug;
+	}
+#endif // ALLOW_TIME_SLICE_DEBUG
+
 protected:
 	double OriginalDuration = 0.;
 	double RemainingDuration = 0.;
 	double StartTime = 0.;
 	mutable double TimeLastTested = 0.;
 	mutable bool bTimeSliceFinishedCached = false;
+
+#if ALLOW_TIME_SLICE_DEBUG
+	mutable TWeakObjectPtr<UNavigationSystemV1> NavSystemDebug;
+	mutable int NaxDataIdxDebug = 0;
+	mutable FIntPoint TileDebug;
+	mutable FBox TileBoxDebug;
+	mutable FName SectionNameDebug = FNavigationSystem::TimeSliceDefaultSectionNameDebug;
+#endif
 };
 
 class NAVIGATIONSYSTEM_API FNavRegenTimeSliceManager
