[Navmesh] Add warnings when tiles could not be added because the tile limit has been reached

diff --git a/Engine/Source/Runtime/NavigationSystem/Private/NavMesh/PImplRecastNavMesh.cpp b/Engine/Source/Runtime/NavigationSystem/Private/NavMesh/PImplRecastNavMesh.cpp
--- a/Engine/Source/Runtime/NavigationSystem/Private/NavMesh/PImplRecastNavMesh.cpp
+++ b/Engine/Source/Runtime/NavigationSystem/Private/NavMesh/PImplRecastNavMesh.cpp
@@ -741,7 +741,12 @@ void FPImplRecastNavMesh::Serialize( FArchive& ArWrapped, int32 NavMeshVersion )
 				if (TileData != NULL)
 				{
 					dtMeshHeader* const TileHeader = (dtMeshHeader*)TileData;
-					DetourNavMesh->addTile(TileData, TileDataSize, DT_TILE_FREE_DATA, TileRef, NULL);
+					Status = DetourNavMesh->addTile(TileData, TileDataSize, DT_TILE_FREE_DATA, TileRef, NULL);
+					if (dtStatusDetail(Status, DT_OUT_OF_MEMORY))
+					{
+						UE_LOG(LogNavigation, Warning, TEXT("%s Failed to add tile (%d,%d:%d), %d tile limit reached in %s."),
+							ANSI_TO_TCHAR(__FUNCTION__), TileHeader->x, TileHeader->y, TileHeader->layer, DetourNavMesh->getMaxTiles(), *NavMeshOwner->GetFullName());
+					}
 
 					// Serialize compressed tile cache layer
 					uint8* ComressedTileData = nullptr;
diff --git a/Engine/Source/Runtime/NavigationSystem/Private/NavMesh/RecastNavMeshDataChunk.cpp b/Engine/Source/Runtime/NavigationSystem/Private/NavMesh/RecastNavMeshDataChunk.cpp
--- a/Engine/Source/Runtime/NavigationSystem/Private/NavMesh/RecastNavMeshDataChunk.cpp
+++ b/Engine/Source/Runtime/NavigationSystem/Private/NavMesh/RecastNavMeshDataChunk.cpp
@@ -297,6 +297,12 @@ TArray<FNavTileRef> URecastNavMeshDataChunk::AttachTiles(ARecastNavMesh& NavMesh
 
 				if (dtStatusFailed(status))
 				{
+					if (dtStatusDetail(status, DT_OUT_OF_MEMORY))
+					{
+						UE_LOG(LogNavigation, Warning, TEXT("%s> Failed to add tile (%d,%d:%d), %d tile limit reached! (from: %s)"),
+							*NavMesh.GetName(), Header->x, Header->y, Header->layer, DetourNavMesh->getMaxTiles(), ANSI_TO_TCHAR(__FUNCTION__));
+					}
+					
 					continue;
 				}
 				else
diff --git a/Engine/Source/Runtime/NavigationSystem/Private/NavMesh/RecastNavMeshGenerator.cpp b/Engine/Source/Runtime/NavigationSystem/Private/NavMesh/RecastNavMeshGenerator.cpp
--- a/Engine/Source/Runtime/NavigationSystem/Private/NavMesh/RecastNavMeshGenerator.cpp
+++ b/Engine/Source/Runtime/NavigationSystem/Private/NavMesh/RecastNavMeshGenerator.cpp
@@ -5113,8 +5113,8 @@ void FRecastNavMeshGenerator::AddGeneratedTileLayer(int32 LayerIndex, FRecastTil
 			{
 				if (dtStatusDetail(status, DT_OUT_OF_MEMORY))
 				{
-					UE_LOG(LogNavigation, Error, TEXT("%s> Tile (%d,%d:%d), tile limit reached!! (%d)"),
-						*DestNavMesh->GetName(), TileX, TileY, LayerIndex, DetourMesh->getMaxTiles());
+					UE_LOG(LogNavigation, Error, TEXT("%s> Failed to add tile (%d,%d:%d), %d tile limit reached! (from %s)"),
+						*DestNavMesh->GetName(), TileX, TileY, LayerIndex, DetourMesh->getMaxTiles(), ANSI_TO_TCHAR(__FUNCTION__));
 				}
 			}
 			else
