Time slicing recast ledge span filtering.

diff --git a/Engine/Source/Runtime/NavigationSystem/Private/NavMesh/RecastNavMeshGenerator.cpp b/Engine/Source/Runtime/NavigationSystem/Private/NavMesh/RecastNavMeshGenerator.cpp
--- a/Engine/Source/Runtime/NavigationSystem/Private/NavMesh/RecastNavMeshGenerator.cpp
+++ b/Engine/Source/Runtime/NavigationSystem/Private/NavMesh/RecastNavMeshGenerator.cpp
@@ -1618,6 +1618,8 @@ FRecastTileGenerator::FRecastTileGenerator(FRecastNavMeshGenerator& ParentGenera
 
 	RasterizeGeomRecastState = ERasterizeGeomRecastTimeSlicedState::MarkWalkableTriangles;
 	RasterizeGeomState = ERasterizeGeomTimeSlicedState::RasterizeGeometryTransformCoords;
+	GenerateRecastFilterState = EGenerateRecastFilterTimeSlicedState::FilterLowHangingWalkableObstacles;
+	GenRecastFilterLedgeSpansYStart = 0;
 	DoWorkTimeSlicedState = EDoWorkTimeSlicedState::GatherGeometryFromSources;
 	GenerateTileTimeSlicedState = EGenerateTileTimeSlicedState::GenerateCompressedLayers;
 
@@ -1626,6 +1628,9 @@ FRecastTileGenerator::FRecastTileGenerator(FRecastNavMeshGenerator& ParentGenera
 	GenCompressedLayersTimeSlicedState = EGenerateCompressedLayersTimeSliced::Init;
 	RasterizeTrianglesTimeSlicedRawGeomIdx = 0;
 	RasterizeTrianglesTimeSlicedInstTransformIdx = 0;
+
+	check(ParentGenerator.GetOwner());
+	TileTimeSliceSettings.FilterLedgeSpansMaxYProcess = ParentGenerator.GetOwner()->TimeSliceFilterLedgeSpansMaxYProcess;
 }
 
 FRecastTileGenerator::~FRecastTileGenerator()
@@ -2820,6 +2825,87 @@ void FRecastTileGenerator::GenerateRecastFilter(FNavMeshBuildContext& BuildConte
 	}
 }
 
+ETimeSliceWorkResult FRecastTileGenerator::GenerateRecastFilterTimeSliced(FNavMeshBuildContext& BuildContext, FTileRasterizationContext& RasterContext)
+{
+	SCOPE_CYCLE_COUNTER(STAT_Navigation_RecastFilter)
+
+	check(TimeSliceManager);
+
+	ETimeSliceWorkResult WorkResult = ETimeSliceWorkResult::Succeeded;
+
+	switch (GenerateRecastFilterState)
+	{
+	case EGenerateRecastFilterTimeSlicedState::FilterLowHangingWalkableObstacles:
+	{
+		// Once all geometry is rasterized, we do initial pass of filtering to
+		// remove unwanted overhangs caused by the conservative rasterization
+		// as well as filter spans where the character cannot possibly stand.
+		rcFilterLowHangingWalkableObstacles(&BuildContext, TileConfig.walkableClimb, *RasterContext.SolidHF);
+		GenerateRecastFilterState = EGenerateRecastFilterTimeSlicedState::FilterLedgeSpans;
+	}// fall through to next state
+	case EGenerateRecastFilterTimeSlicedState::FilterLedgeSpans:
+	{
+		SCOPE_CYCLE_COUNTER(STAT_Navigation_FilterLedgeSpans)
+
+		bool DoIter = true;
+
+		do
+		{
+			rcFilterLedgeSpans(&BuildContext, TileConfig.walkableHeight, TileConfig.walkableClimb, GenRecastFilterLedgeSpansYStart, TileTimeSliceSettings.FilterLedgeSpansMaxYProcess, *RasterContext.SolidHF);
+
+			GenRecastFilterLedgeSpansYStart += TileTimeSliceSettings.FilterLedgeSpansMaxYProcess;
+
+			if (GenRecastFilterLedgeSpansYStart >= RasterContext.SolidHF->height)
+			{
+				GenRecastFilterLedgeSpansYStart = 0;
+				DoIter = false;
+				GenerateRecastFilterState = EGenerateRecastFilterTimeSlicedState::FilterWalkableLowHeightSpans;
+			}
+
+			MARK_TIMESLICE_SECTION_DEBUG(TimeSliceManager->GetTimeSlicer(), FilterLedgeSpans);
+
+			// Only FilterLedge Spans has been found to be slow so we only actually test the timeslice here for this function
+			if (TimeSliceManager->GetTimeSlicer().TestTimeSliceFinished())
+			{
+				return ETimeSliceWorkResult::CallAgainNextTimeSlice;
+			}
+		} while (DoIter);
+	}// fall through to next state
+	case EGenerateRecastFilterTimeSlicedState::FilterWalkableLowHeightSpans:
+	{
+		// TileConfig.walkableHeight is set to 1 when marking low spans, calculate real value for filtering
+		const int32 FilterWalkableHeight = FMath::CeilToInt(TileConfig.AgentHeight / TileConfig.ch);
+
+		if (!TileConfig.bMarkLowHeightAreas)
+		{
+			rcFilterWalkableLowHeightSpans(&BuildContext, TileConfig.walkableHeight, *RasterContext.SolidHF);
+		}
+		else if (TileConfig.bFilterLowSpanFromTileCache)
+		{
+			// TODO: investigate if creating detailed 2D map from active modifiers is cheap enough
+			// for now, switch on presence of those modifiers, will save memory as long as they are sparse (should be)
+			if (TileConfig.bFilterLowSpanSequences && bHasLowAreaModifiers)
+			{
+				rcFilterWalkableLowHeightSpansSequences(&BuildContext, FilterWalkableHeight, *RasterContext.SolidHF);
+			}
+			else
+			{
+				rcFilterWalkableLowHeightSpans(&BuildContext, FilterWalkableHeight, *RasterContext.SolidHF);
+			}
+		}
+		GenerateRecastFilterState = EGenerateRecastFilterTimeSlicedState::FilterLowHangingWalkableObstacles;
+	}
+	break;
+	default:
+	{
+		ensureMsgf(false, TEXT("unhandled EGenerateRecastFilterTimeSlicedState"));
+		return ETimeSliceWorkResult::Failed;
+	}
+	}
+
+	return ETimeSliceWorkResult::Succeeded;
+}
+
 bool FRecastTileGenerator::BuildCompactHeightField(FNavMeshBuildContext& BuildContext, FTileRasterizationContext& RasterContext)
 {
 	SCOPE_CYCLE_COUNTER(STAT_Navigation_RecastBuildCompactHeightField);
@@ -3106,13 +3192,15 @@ ETimeSliceWorkResult FRecastTileGenerator::GenerateCompressedLayersTimeSliced(FN
 	}// fall through to next state
 	case EGenerateCompressedLayersTimeSliced::RecastFilter:
 	{
-		GenerateRecastFilter(BuildContext, *RasterContext);
-
-		GenCompressedLayersTimeSlicedState = EGenerateCompressedLayersTimeSliced::CompactHeightField;
+		const ETimeSliceWorkResult WorkResult = GenerateRecastFilterTimeSliced(BuildContext, *RasterContext);
 
-		MARK_TIMESLICE_SECTION_DEBUG(TimeSliceManager->GetTimeSlicer(), RecastFilterLedgeSpans);
+		// Non timesliced code this is based on did not care about success or failure here.
+		if (WorkResult != ETimeSliceWorkResult::CallAgainNextTimeSlice)
+		{
+			GenCompressedLayersTimeSlicedState = EGenerateCompressedLayersTimeSliced::CompactHeightField;
+		}
 
-		if (TimeSliceManager->GetTimeSlicer().TestTimeSliceFinished())
+		if (TimeSliceManager->GetTimeSlicer().IsTimeSliceFinishedCached())
 		{
 			return ETimeSliceWorkResult::CallAgainNextTimeSlice;
 		}
diff --git a/Engine/Source/Runtime/NavigationSystem/Public/NavMesh/RecastNavMesh.h b/Engine/Source/Runtime/NavigationSystem/Public/NavMesh/RecastNavMesh.h
--- a/Engine/Source/Runtime/NavigationSystem/Public/NavMesh/RecastNavMesh.h
+++ b/Engine/Source/Runtime/NavigationSystem/Public/NavMesh/RecastNavMesh.h
@@ -815,6 +815,10 @@ class NAVIGATIONSYSTEM_API ARecastNavMesh : public ANavigationData
 	UPROPERTY(config)
 	uint32 bAllowNavLinkAsPathEnd : 1;
 
+	/** The maximum number of y coords to process when time slicing filter ledge spans during navmesh regeneration. */
+	UPROPERTY(EditAnywhere, Category = TimeSlicing, config, AdvancedDisplay, meta = (ClampMin = "1", UIMin = "1"))
+	int32 TimeSliceFilterLedgeSpansMaxYProcess = 13;
+
 	/** If a single time sliced section of navmesh regen code exceeds this duration then it will trigger debug logging */
 	UPROPERTY(EditAnywhere, Category = TimeSlicing, config, AdvancedDisplay)
 	double TimeSliceLongDurationDebug = 0.002;
diff --git a/Engine/Source/Runtime/NavigationSystem/Public/NavMesh/RecastNavMeshGenerator.h b/Engine/Source/Runtime/NavigationSystem/Public/NavMesh/RecastNavMeshGenerator.h
--- a/Engine/Source/Runtime/NavigationSystem/Public/NavMesh/RecastNavMeshGenerator.h
+++ b/Engine/Source/Runtime/NavigationSystem/Public/NavMesh/RecastNavMeshGenerator.h
@@ -241,6 +241,13 @@ enum class ERasterizeGeomTimeSlicedState : uint8
 	RasterizeGeometryRecast,
 };
 
+enum class EGenerateRecastFilterTimeSlicedState : uint8
+{
+	FilterLowHangingWalkableObstacles,
+	FilterLedgeSpans,
+	FilterWalkableLowHeightSpans,
+};
+
 enum class EDoWorkTimeSlicedState : uint8
 {
 	Invalid,
@@ -262,6 +269,11 @@ enum class EGenerateNavDataTimeSlicedState : uint8
 	GenerateLayers,
 };
 
+struct FRecastTileTimeSliceSettings
+{
+	int32 FilterLedgeSpansMaxYProcess = 13;
+};
+
 /**
  * Class handling generation of a single tile, caching data that can speed up subsequent tile generations
  */
@@ -367,6 +379,7 @@ protected:
 	UE_DEPRECATED(5.0, "Call the version of this function where Coords are now a TArray of FReals!")
 	void RasterizeGeometry(FNavMeshBuildContext& BuildContext, const TArray<float>& Coords, const TArray<int32>& Indices, const FTransform& LocalToWorld, const rcRasterizationFlags RasterizationFlags, FTileRasterizationContext& RasterContext);
 	void GenerateRecastFilter(FNavMeshBuildContext& BuildContext, FTileRasterizationContext& RasterContext);
+	ETimeSliceWorkResult GenerateRecastFilterTimeSliced(FNavMeshBuildContext& BuildContext, FTileRasterizationContext& RasterContext);
 	bool BuildCompactHeightField(FNavMeshBuildContext& BuildContext, FTileRasterizationContext& RasterContext);
 	bool RecastErodeWalkable(FNavMeshBuildContext& BuildContext, FTileRasterizationContext& RasterContext);
 	bool RecastBuildLayers(FNavMeshBuildContext& BuildContext, FTileRasterizationContext& RasterContext);
@@ -426,6 +439,8 @@ protected:
 	/** Start time slicing variables */
 	ERasterizeGeomRecastTimeSlicedState RasterizeGeomRecastState;
 	ERasterizeGeomTimeSlicedState RasterizeGeomState;
+	EGenerateRecastFilterTimeSlicedState GenerateRecastFilterState;
+	int32 GenRecastFilterLedgeSpansYStart;
 	EDoWorkTimeSlicedState DoWorkTimeSlicedState;
 	EGenerateTileTimeSlicedState GenerateTileTimeSlicedState;
 
@@ -440,6 +455,8 @@ protected:
 	TUniquePtr<struct FTileCacheAllocator> GenNavDataTimeSlicedAllocator;
 	TUniquePtr<struct FTileGenerationContext> GenNavDataTimeSlicedGenerationContext;
 	TUniquePtr<struct FTileRasterizationContext> GenCompressedlayersTimeSlicedRasterContext;
+
+	FRecastTileTimeSliceSettings TileTimeSliceSettings;
 	/** End time slicing variables */
 
 	int32 TileX;
diff --git a/Engine/Source/Runtime/Navmesh/Private/Recast/RecastFilter.cpp b/Engine/Source/Runtime/Navmesh/Private/Recast/RecastFilter.cpp
--- a/Engine/Source/Runtime/Navmesh/Private/Recast/RecastFilter.cpp
+++ b/Engine/Source/Runtime/Navmesh/Private/Recast/RecastFilter.cpp
@@ -72,6 +72,89 @@ void rcFilterLowHangingWalkableObstacles(rcContext* ctx, const int walkableClimb
 	ctx->stopTimer(RC_TIMER_FILTER_LOW_OBSTACLES);
 }
 
+void rcFilterLedgeSpansImp(rcContext* ctx, const int walkableHeight, const int walkableClimb, const int filterLedgeSpansAtY,
+	rcHeightfield& solid)
+{
+	rcAssert(ctx);
+
+	const int w = solid.width;
+	const int h = solid.height;
+	const int MAX_HEIGHT = 0xffff;
+
+	// Mark border spans.
+	for (int x = 0; x < w; ++x)
+	{
+		for (rcSpan* s = solid.spans[x + filterLedgeSpansAtY*w]; s; s = s->next)
+		{
+			// Skip non walkable spans.
+			if (s->data.area == RC_NULL_AREA)
+				continue;
+
+			const int bot = (int)(s->data.smax);
+			const int top = s->next ? (int)(s->next->data.smin) : MAX_HEIGHT;
+
+			// Find neighbours minimum height.
+			int minh = MAX_HEIGHT;
+
+			// Min and max height of accessible neighbours.
+			int asmin = s->data.smax;
+			int asmax = s->data.smax;
+
+			for (int dir = 0; dir < 4; ++dir)
+			{
+				int dx = x + rcGetDirOffsetX(dir);
+				int dy = filterLedgeSpansAtY + rcGetDirOffsetY(dir);
+				// Skip neighbours which are out of bounds.
+				if (dx < 0 || dy < 0 || dx >= w || dy >= h)
+				{
+					minh = rcMin(minh, -walkableClimb - bot);
+					continue;
+				}
+
+				// From minus infinity to the first span.
+				rcSpan* ns = solid.spans[dx + dy*w];
+				int nbot = -walkableClimb;
+				int ntop = ns ? (int)ns->data.smin : MAX_HEIGHT;
+				// Skip neightbour if the gap between the spans is too small.
+				if (rcMin(top, ntop) - rcMax(bot, nbot) > walkableHeight)
+					minh = rcMin(minh, nbot - bot);
+
+				// Rest of the spans.
+				for (ns = solid.spans[dx + dy*w]; ns; ns = ns->next)
+				{
+					nbot = (int)ns->data.smax;
+					ntop = ns->next ? (int)ns->next->data.smin : MAX_HEIGHT;
+					// Skip neightbour if the gap between the spans is too small.
+					if (rcMin(top, ntop) - rcMax(bot, nbot) > walkableHeight)
+					{
+						minh = rcMin(minh, nbot - bot);
+
+						// Find min/max accessible neighbour height. 
+						if (rcAbs(nbot - bot) <= walkableClimb)
+						{
+							if (nbot < asmin) asmin = nbot;
+							if (nbot > asmax) asmax = nbot;
+						}
+
+					}
+				}
+			}
+
+			// The current span is close to a ledge if the drop to any
+			// neighbour span is less than the walkableClimb.
+			if (minh < -walkableClimb)
+				s->data.area = RC_NULL_AREA;
+
+			// If the difference between all neighbours is too large,
+			// we are at steep slope, mark the span as ledge.
+			if ((asmax - asmin) > walkableClimb)
+			{
+				s->data.area = RC_NULL_AREA;
+			}
+		}
+	}
+}
+
 /// @par
 ///
 /// A ledge is a span with one or more neighbors whose maximum is further away than @p walkableClimb
@@ -96,81 +179,31 @@ void rcFilterLedgeSpans(rcContext* ctx, const int walkableHeight, const int walk
 	// Mark border spans.
 	for (int y = 0; y < h; ++y)
 	{
-		for (int x = 0; x < w; ++x)
-		{
-			for (rcSpan* s = solid.spans[x + y*w]; s; s = s->next)
-			{
-				// Skip non walkable spans.
-				if (s->data.area == RC_NULL_AREA)
-					continue;
-				
-				const int bot = (int)(s->data.smax);
-				const int top = s->next ? (int)(s->next->data.smin) : MAX_HEIGHT;
-				
-				// Find neighbours minimum height.
-				int minh = MAX_HEIGHT;
+		rcFilterLedgeSpansImp(ctx, walkableHeight, walkableClimb, y, solid);
+	}
+	
+	ctx->stopTimer(RC_TIMER_FILTER_BORDER);
+}	
 
-				// Min and max height of accessible neighbours.
-				int asmin = s->data.smax;
-				int asmax = s->data.smax;
 
-				for (int dir = 0; dir < 4; ++dir)
-				{
-					int dx = x + rcGetDirOffsetX(dir);
-					int dy = y + rcGetDirOffsetY(dir);
-					// Skip neighbours which are out of bounds.
-					if (dx < 0 || dy < 0 || dx >= w || dy >= h)
-					{
-						minh = rcMin(minh, -walkableClimb - bot);
-						continue;
-					}
+/// @see rcHeightfield, rcConfig
+void rcFilterLedgeSpans(rcContext* ctx, const int walkableHeight, const int walkableClimb, const int yStart, const int maxYProcess,
+	rcHeightfield& solid)
+{
+	rcAssert(ctx);
 
-					// From minus infinity to the first span.
-					rcSpan* ns = solid.spans[dx + dy*w];
-					int nbot = -walkableClimb;
-					int ntop = ns ? (int)ns->data.smin : MAX_HEIGHT;
-					// Skip neightbour if the gap between the spans is too small.
-					if (rcMin(top,ntop) - rcMax(bot,nbot) > walkableHeight)
-						minh = rcMin(minh, nbot - bot);
-					
-					// Rest of the spans.
-					for (ns = solid.spans[dx + dy*w]; ns; ns = ns->next)
-					{
-						nbot = (int)ns->data.smax;
-						ntop = ns->next ? (int)ns->next->data.smin : MAX_HEIGHT;
-						// Skip neightbour if the gap between the spans is too small.
-						if (rcMin(top,ntop) - rcMax(bot,nbot) > walkableHeight)
-						{
-							minh = rcMin(minh, nbot - bot);
-						
-							// Find min/max accessible neighbour height. 
-							if (rcAbs(nbot - bot) <= walkableClimb)
-							{
-								if (nbot < asmin) asmin = nbot;
-								if (nbot > asmax) asmax = nbot;
-							}
-							
-						}
-					}
-				}
-				
-				// The current span is close to a ledge if the drop to any
-				// neighbour span is less than the walkableClimb.
-				if (minh < -walkableClimb)
-					s->data.area = RC_NULL_AREA;
-					
-				// If the difference between all neighbours is too large,
-				// we are at steep slope, mark the span as ledge.
-				if ((asmax - asmin) > walkableClimb)
-				{
-					s->data.area = RC_NULL_AREA;
-				}
-			}
-		}
+	ctx->startTimer(RC_TIMER_FILTER_BORDER);
+
+	const int w = solid.width;
+	const int h = rcMin(yStart + maxYProcess, solid.height);
+
+	for (int y = yStart; y < h; ++y)
+	{
+		rcFilterLedgeSpansImp(ctx, walkableHeight, walkableClimb, y, solid);
 	}
-	
+
 	ctx->stopTimer(RC_TIMER_FILTER_BORDER);
-}	
+}
 
 /// @par
 ///
diff --git a/Engine/Source/Runtime/Navmesh/Public/Recast/Recast.h b/Engine/Source/Runtime/Navmesh/Public/Recast/Recast.h
--- a/Engine/Source/Runtime/Navmesh/Public/Recast/Recast.h
+++ b/Engine/Source/Runtime/Navmesh/Public/Recast/Recast.h
@@ -1058,6 +1058,19 @@ NAVMESH_API void rcRasterizeTriangles(rcContext* ctx, const rcReal* verts, const
 ///  @param[in,out]	solid			A fully built heightfield.  (All spans have been added.)
 NAVMESH_API void rcFilterLowHangingWalkableObstacles(rcContext* ctx, const int walkableClimb, rcHeightfield& solid);
 
+/// Marks spans that are ledges as not-walkable, by a number of y coords at a time.
+///  @ingroup recast
+///  @param[in,out]	ctx				The build context to use during the operation.
+///  @param[in]		walkableHeight	Minimum floor to 'ceiling' height that will still allow the floor area to 
+///  								be considered walkable. [Limit: >= 3] [Units: vx]
+///  @param[in]		walkableClimb	Maximum ledge height that is considered to still be traversable. 
+///  								[Limit: >=0] [Units: vx]
+///  @param[in,out]	solid			A fully built heightfield.  (All spans have been added.)
+///  @param[in]		yStart			y coord to start at
+///  @param[in]		maxYProcess	    Max y coords to process (yStart + maxYProcess can be more than solid.height and will be capped to solid.height)
+NAVMESH_API void rcFilterLedgeSpans(rcContext* ctx, const int walkableHeight,
+						const int walkableClimb, const int yStart, const int maxYProcess, rcHeightfield& solid);
+
 /// Marks spans that are ledges as not-walkable. 
 ///  @ingroup recast
 ///  @param[in,out]	ctx				The build context to use during the operation.
