[Navmesh] Rename debug display property bHeightfieldSolidPostRadiusFiltering to bHeightfieldSolidPostInclusionBoundsFiltering and add some property descriptions.

diff --git a/Engine/Config/BaseEngine.ini b/Engine/Config/BaseEngine.ini
--- a/Engine/Config/BaseEngine.ini
+++ b/Engine/Config/BaseEngine.ini
@@ -1092,6 +1092,8 @@ DefaultWorldPartitionHLODLayer=/Engine/Maps/Templates/HLODs/HLODLayer_Instanced.
 +PropertyRedirects=(OldName="RectLightComponent.SetSourceTexture.bNewValue", NewName="NewValue")
 +PropertyRedirects=(OldName="RectLightComponent.SetSourceWidth.bNewValue", NewName="NewValue")
 
++PropertyRedirects=(OldName="RecastNavMeshTileGenerationDebug.bHeightfieldSolidPostRadiusFiltering", NewName="bHeightfieldSolidPostInclusionBoundsFiltering")
+
 [CoreUObject.Metadata]
 ; Note: UnrealHeaderTool should be rerun after making changes to MetadataRedirects to catch any keys specified in class headers
 +MetadataRedirects=(OldKey="K2Protected", NewKey="BlueprintProtected")
diff --git a/Engine/Source/Runtime/NavigationSystem/Private/NavMesh/RecastNavMesh.cpp b/Engine/Source/Runtime/NavigationSystem/Private/NavMesh/RecastNavMesh.cpp
--- a/Engine/Source/Runtime/NavigationSystem/Private/NavMesh/RecastNavMesh.cpp
+++ b/Engine/Source/Runtime/NavigationSystem/Private/NavMesh/RecastNavMesh.cpp
@@ -436,7 +436,7 @@ FRecastNavMeshTileGenerationDebug::FRecastNavMeshTileGenerationDebug()
 	bEnabled = false;
 	TileCoordinate = FIntVector::ZeroValue;
 	bHeightfieldSolidFromRasterization = false;
-	bHeightfieldSolidPostRadiusFiltering = false;
+	bHeightfieldSolidPostInclusionBoundsFiltering = false;
 	bHeightfieldSolidPostHeightFiltering = false;
 	bCompactHeightfield = false;
 	bCompactHeightfieldEroded = false;
diff --git a/Engine/Source/Runtime/NavigationSystem/Private/NavMesh/RecastNavMeshGenerator.cpp b/Engine/Source/Runtime/NavigationSystem/Private/NavMesh/RecastNavMeshGenerator.cpp
--- a/Engine/Source/Runtime/NavigationSystem/Private/NavMesh/RecastNavMeshGenerator.cpp
+++ b/Engine/Source/Runtime/NavigationSystem/Private/NavMesh/RecastNavMeshGenerator.cpp
@@ -3330,7 +3330,7 @@ bool FRecastTileGenerator::GenerateCompressedLayers(FNavMeshBuildContext& BuildC
 	}
 
 #if RECAST_INTERNAL_DEBUG_DATA
-	if (IsTileDebugActive() && TileDebugSettings.bHeightfieldSolidPostRadiusFiltering)
+	if (IsTileDebugActive() && TileDebugSettings.bHeightfieldSolidPostInclusionBoundsFiltering)
 	{
 		duDebugDrawHeightfieldSolid(&BuildContext.InternalDebugData, *RasterContext.SolidHF);
 	}
diff --git a/Engine/Source/Runtime/NavigationSystem/Public/NavMesh/RecastNavMesh.h b/Engine/Source/Runtime/NavigationSystem/Public/NavMesh/RecastNavMesh.h
--- a/Engine/Source/Runtime/NavigationSystem/Public/NavMesh/RecastNavMesh.h
+++ b/Engine/Source/Runtime/NavigationSystem/Public/NavMesh/RecastNavMesh.h
@@ -434,7 +434,7 @@ struct NAVIGATIONSYSTEM_API FRecastNavMeshTileGenerationDebug
 	uint32 bHeightfieldSolidFromRasterization : 1;
 
 	UPROPERTY(EditAnywhere, Category = Debug)
-	uint32 bHeightfieldSolidPostRadiusFiltering : 1;
+	uint32 bHeightfieldSolidPostInclusionBoundsFiltering : 1;
 
 	UPROPERTY(EditAnywhere, Category = Debug)
 	uint32 bHeightfieldSolidPostHeightFiltering : 1;
@@ -629,10 +629,11 @@ class NAVIGATIONSYSTEM_API ARecastNavMesh : public ANavigationData
 	UPROPERTY(EditAnywhere, Category=Display)
 	uint32 bDrawClusters:1;
 
-	/** Draw edges of every navmesh's triangle */
+	/** Draw octree used to store navigation relevant actors */
 	UPROPERTY(EditAnywhere, Category = Display)
 	uint32 bDrawOctree : 1;
 
+	/** Draw octree used to store navigation relevant actors with the elements bounds */
 	UPROPERTY(EditAnywhere, Category = Display, meta=(editcondition = "bDrawOctree"))
 	uint32 bDrawOctreeDetails : 1;
 
