[WP navmesh] Allow navmesh tile pool to be resizable when built through the UWorldPartitionNavigationDataBuilder (ignoring bFixedTilePoolSize)

diff --git a/Engine/Source/Editor/UnrealEd/Private/WorldPartition/WorldPartitionNavigationDataBuilder.cpp b/Engine/Source/Editor/UnrealEd/Private/WorldPartition/WorldPartitionNavigationDataBuilder.cpp
--- a/Engine/Source/Editor/UnrealEd/Private/WorldPartition/WorldPartitionNavigationDataBuilder.cpp
+++ b/Engine/Source/Editor/UnrealEd/Private/WorldPartition/WorldPartitionNavigationDataBuilder.cpp
@@ -289,13 +289,13 @@ bool UWorldPartitionNavigationDataBuilder::GenerateNavigationData(UWorldPartitio
 
 	static int32 CallCount = 0;
 	CallCount++;
-	UE_LOG(LogWorldPartitionNavigationDataBuilder, Log, TEXT("%i. GenerateNavigationData for LoadedBounds %s"), CallCount, *LoadedBounds.ToString());
+	UE_LOG(LogWorldPartitionNavigationDataBuilder, Log, TEXT("Iteration %i. GenerateNavigationData for LoadedBounds %s"), CallCount, *LoadedBounds.ToString());
 
 	UWorld* World = WorldPartition->World;
 
 	// Generate navmesh
-	// Make sure navigation is added and initialized in EditorMode
-	FNavigationSystem::AddNavigationSystemToWorld(*World, FNavigationSystemRunMode::EditorMode);
+	// Make sure navigation is added and initialized in EditorWorldPartitionBuildMode
+	FNavigationSystem::AddNavigationSystemToWorld(*World, FNavigationSystemRunMode::EditorWorldPartitionBuildMode);
 
 	UNavigationSystemBase* NavSystem = World->GetNavigationSystem();
 	if (NavSystem == nullptr)
diff --git a/Engine/Source/Runtime/Engine/Classes/AI/NavigationSystemBase.h b/Engine/Source/Runtime/Engine/Classes/AI/NavigationSystemBase.h
--- a/Engine/Source/Runtime/Engine/Classes/AI/NavigationSystemBase.h
+++ b/Engine/Source/Runtime/Engine/Classes/AI/NavigationSystemBase.h
@@ -34,6 +34,7 @@ enum class FNavigationSystemRunMode : uint8
 	SimulationMode,
 	PIEMode,
 	InferFromWorldMode,
+	EditorWorldPartitionBuildMode
 };
 
 namespace ENavigationLockReason
@@ -104,6 +105,8 @@ namespace FNavigationSystem
 	/** Discards all navigation data chunks in all sub-levels */
 	ENGINE_API void DiscardNavigationDataChunks(UWorld& InWorld);
 
+	ENGINE_API bool IsEditorRunMode(const FNavigationSystemRunMode Mode);
+	
 	template<typename TNavSys>
 	FORCEINLINE TNavSys* GetCurrent(UWorld* World)
 	{
diff --git a/Engine/Source/Runtime/Engine/Private/AI/NavigationSystemBase.cpp b/Engine/Source/Runtime/Engine/Private/AI/NavigationSystemBase.cpp
--- a/Engine/Source/Runtime/Engine/Private/AI/NavigationSystemBase.cpp
+++ b/Engine/Source/Runtime/Engine/Private/AI/NavigationSystemBase.cpp
@@ -67,6 +67,11 @@ namespace FNavigationSystem
 		}
 	}
 
+	bool IsEditorRunMode(const FNavigationSystemRunMode Mode)
+	{
+		return Mode == FNavigationSystemRunMode::EditorMode || Mode == FNavigationSystemRunMode::EditorWorldPartitionBuildMode;
+	}
+
 	void AddNavigationSystemToWorld(UWorld& WorldOwner, const FNavigationSystemRunMode RunMode, UNavigationSystemConfig* NavigationSystemConfig, const bool bInitializeForWorld, const bool bOverridePreviousNavSys)
 	{
 		TRACE_CPUPROFILER_EVENT_SCOPE(FNavigationSystem::AddNavigationSystemToWorld);
@@ -98,7 +103,7 @@ namespace FNavigationSystem
 			{
 				WorldOwner.GetNavigationSystem()->InitializeForWorld(WorldOwner, ResolvedRunMode);
 			}
-			else if (ResolvedRunMode == FNavigationSystemRunMode::EditorMode)
+			else if (IsEditorRunMode(ResolvedRunMode))
 			{
 				DiscardNavigationDataChunks(WorldOwner);
 			}
diff --git a/Engine/Source/Runtime/NavigationSystem/Private/NavMesh/RecastNavMesh.cpp b/Engine/Source/Runtime/NavigationSystem/Private/NavMesh/RecastNavMesh.cpp
--- a/Engine/Source/Runtime/NavigationSystem/Private/NavMesh/RecastNavMesh.cpp
+++ b/Engine/Source/Runtime/NavigationSystem/Private/NavMesh/RecastNavMesh.cpp
@@ -1241,6 +1241,16 @@ int32 ARecastNavMesh::GetCompressedTileCacheSize()
 
 bool ARecastNavMesh::IsResizable() const
 {
+	// Ignore bFixedTilePoolSize when in EditorWorldPartitionBuildMode
+	if (const UWorld* World = GetWorld())
+	{
+		const UNavigationSystemV1* NavSys = FNavigationSystem::GetCurrent<UNavigationSystemV1>(World);
+		if (NavSys && NavSys->GetRunMode() == FNavigationSystemRunMode::EditorWorldPartitionBuildMode)
+		{
+			return true;
+		}
+	}
+	
 	return !bFixedTilePoolSize;
 }
 
diff --git a/Engine/Source/Runtime/NavigationSystem/Private/NavSystemConfigOverride.cpp b/Engine/Source/Runtime/NavigationSystem/Private/NavSystemConfigOverride.cpp
--- a/Engine/Source/Runtime/NavigationSystem/Private/NavSystemConfigOverride.cpp
+++ b/Engine/Source/Runtime/NavigationSystem/Private/NavSystemConfigOverride.cpp
@@ -124,7 +124,7 @@ void ANavSystemConfigOverride::OverrideNavSystem()
 				: FNavigationSystemRunMode::GameMode)
 			;
 
-		if (RunMode == FNavigationSystemRunMode::EditorMode)
+		if (FNavigationSystem::IsEditorRunMode(RunMode))
 		{
 			FNavigationSystem::AddNavigationSystemToWorld(*World, RunMode, NavigationSystemConfig, /*bInitializeForWorld=*/false, /*bOverridePreviousNavSys=*/true);
 #if WITH_EDITOR
diff --git a/Engine/Source/Runtime/NavigationSystem/Private/NavigationSystem.cpp b/Engine/Source/Runtime/NavigationSystem/Private/NavigationSystem.cpp
--- a/Engine/Source/Runtime/NavigationSystem/Private/NavigationSystem.cpp
+++ b/Engine/Source/Runtime/NavigationSystem/Private/NavigationSystem.cpp
@@ -933,7 +933,7 @@ void UNavigationSystemV1::OnWorldInitDone(FNavigationSystemRunMode Mode)
 			}
 		}
 
-		if (OperationMode == FNavigationSystemRunMode::EditorMode)
+		if (FNavigationSystem::IsEditorRunMode(OperationMode))
 		{
 			RemoveNavigationBuildLock(InitialNavBuildingLockFlags, ELockRemovalRebuildAction::RebuildIfNotInEditor);
 		}
@@ -994,7 +994,7 @@ void UNavigationSystemV1::OnWorldInitDone(FNavigationSystemRunMode Mode)
 			}
 		}
 
-		if (OperationMode == FNavigationSystemRunMode::EditorMode)
+		if (FNavigationSystem::IsEditorRunMode(OperationMode))
 		{
 			// don't lock navigation building in editor
 			RemoveNavigationBuildLock(InitialNavBuildingLockFlags, ELockRemovalRebuildAction::RebuildIfNotInEditor);
@@ -1022,7 +1022,7 @@ void UNavigationSystemV1::OnWorldInitDone(FNavigationSystemRunMode Mode)
 	}
 
 #if	WITH_EDITOR
-	if (Mode == FNavigationSystemRunMode::EditorMode)
+	if (FNavigationSystem::IsEditorRunMode(Mode))
 	{
 		// make sure this static get applied to this instance
 		bNavigationAutoUpdateEnabled = !bNavigationAutoUpdateEnabled; 
@@ -1336,7 +1336,7 @@ void UNavigationSystemV1::AddReferencedObjects(UObject* InThis, FReferenceCollec
 	Collector.AddReferencedObject(CrowdManager, InThis);
 
 	// don't reference NavAreaClasses in editor (unless PIE is active)
-	if (This->OperationMode != FNavigationSystemRunMode::EditorMode)
+	if (!FNavigationSystem::IsEditorRunMode(This->OperationMode))
 	{
 		Collector.AddReferencedObjects(This->NavAreaClasses, InThis);
 	}
@@ -1354,7 +1354,7 @@ void UNavigationSystemV1::SetNavigationAutoUpdateEnabled(bool bNewEnable, UNavig
 		{
 			const bool bCurrentIsEnabled = NavSystem->GetIsAutoUpdateEnabled();
 			NavSystem->DefaultDirtyAreasController.bCanAccumulateDirtyAreas = bCurrentIsEnabled
-				|| (NavSystem->OperationMode != FNavigationSystemRunMode::EditorMode && NavSystem->OperationMode != FNavigationSystemRunMode::InvalidMode);
+				|| (!FNavigationSystem::IsEditorRunMode(NavSystem->OperationMode) && NavSystem->OperationMode != FNavigationSystemRunMode::InvalidMode);
 
 			if (bCurrentIsEnabled)
 			{
@@ -3481,7 +3481,7 @@ void UNavigationSystemV1::Build()
 
 	if (bAutoCreateNavigationData == true
 #if WITH_EDITOR
-		|| OperationMode == FNavigationSystemRunMode::EditorMode
+		|| FNavigationSystem::IsEditorRunMode(OperationMode)
 #endif // WITH_EDITOR
 		)
 	{
@@ -3794,7 +3794,7 @@ void UNavigationSystemV1::RemoveNavigationBuildLock(uint8 Flags, const ELockRemo
 		DefaultDirtyAreasController.OnNavigationBuildUnlocked();
 
 		const bool bRebuild = 
-			(RebuildAction == ELockRemovalRebuildAction::RebuildIfNotInEditor && (OperationMode != FNavigationSystemRunMode::EditorMode)) || 
+			(RebuildAction == ELockRemovalRebuildAction::RebuildIfNotInEditor && !FNavigationSystem::IsEditorRunMode(OperationMode)) || 
 			(RebuildAction == ELockRemovalRebuildAction::Rebuild);
 		
 		if (bRebuild)
@@ -3954,7 +3954,7 @@ void UNavigationSystemV1::OnLevelAddedToWorld(ULevel* InLevel, UWorld* InWorld)
 	}
 
 #if WITH_EDITOR
-	if (OperationMode == FNavigationSystemRunMode::EditorMode)
+	if (FNavigationSystem::IsEditorRunMode(OperationMode))
 	{
 		// see if there are any unregistered yet valid nav data instances
 		// In general we register navdata on its PostLoad, but in some cases
@@ -4769,7 +4769,7 @@ void UNavigationSystemV1::GetOnScreenMessages(TMultiMap<FCoreDelegates::EOnScree
 	const bool bIsNavigationAutoUpdateEnabled = true;
 #endif
 	if (IsNavigationDirty()
-		&& ((OperationMode == FNavigationSystemRunMode::EditorMode && !bIsNavigationAutoUpdateEnabled)
+		&& ((FNavigationSystem::IsEditorRunMode(OperationMode) && !bIsNavigationAutoUpdateEnabled)
 			|| !SupportsNavigationGeneration() || !CanRebuildDirtyNavigation()))
 	{
 		OutMessages.Add(FCoreDelegates::EOnScreenMessageSeverity::Error
