Allow ANavigationData to support external packaging
Prevent discard of sublevel navdata when using OFPA
Avoid Build Paths to dirty the Level package when Level uses OFPA
Beautify NavData actor labels when using OFPA

diff --git a/Engine/Source/Editor/UnrealEd/Private/EditorBuildUtils.cpp b/Engine/Source/Editor/UnrealEd/Private/EditorBuildUtils.cpp
--- a/Engine/Source/Editor/UnrealEd/Private/EditorBuildUtils.cpp
+++ b/Engine/Source/Editor/UnrealEd/Private/EditorBuildUtils.cpp
@@ -419,6 +419,12 @@ bool FEditorBuildUtils::EditorBuild( UWorld* InWorld, FName Id, const bool bAllo
 			const FScopedBusyCursor BusyCursor;
 
 			TriggerNavigationBuilder(InWorld, Id);
+
+			// No need to dirty the world package if it uses external actors
+			if (InWorld->PersistentLevel->IsUsingExternalActors())
+			{
+				bDirtyPersistentLevel = false;
+			}
 		}
 	}
 	else if (CustomBuildTypes.Contains(Id))
diff --git a/Engine/Source/Runtime/NavigationSystem/Private/NavMesh/RecastNavMesh.cpp b/Engine/Source/Runtime/NavigationSystem/Private/NavMesh/RecastNavMesh.cpp
--- a/Engine/Source/Runtime/NavigationSystem/Private/NavMesh/RecastNavMesh.cpp
+++ b/Engine/Source/Runtime/NavigationSystem/Private/NavMesh/RecastNavMesh.cpp
@@ -602,7 +602,7 @@ void ARecastNavMesh::CheckToDiscardSubLevelNavData(const UNavigationSystemBase&
 	{
 		// Get rid of instances saved within levels that are streamed-in
 		if ((GEngine->IsSettingUpPlayWorld() == false) // this is a @HACK
-			&& (OwningWorld->GetOutermost() != GetOutermost())
+			&& (OwningWorld->PersistentLevel != GetLevel())
 			// If we are cooking, then let them all pass.
 			// They will be handled at load-time when running.
 			&& (IsRunningCommandlet() == false))
diff --git a/Engine/Source/Runtime/NavigationSystem/Private/NavigationSystem.cpp b/Engine/Source/Runtime/NavigationSystem/Private/NavigationSystem.cpp
--- a/Engine/Source/Runtime/NavigationSystem/Private/NavigationSystem.cpp
+++ b/Engine/Source/Runtime/NavigationSystem/Private/NavigationSystem.cpp
@@ -3733,7 +3733,14 @@ ANavigationData* UNavigationSystemV1::CreateNavigationDataInstanceInLevel(const
 #if WITH_EDITOR
 			if (World->WorldType == EWorldType::Editor)
 			{
-				const bool bMarkDirty = false;
+				FString ActorLabel = StrName;
+				if (Instance->IsPackageExternal())
+				{
+					// When using external package, don't rely on actor's name to generate a label as it contains a unique actor identifier which obfuscates the label
+					ActorLabel = FString::Printf(TEXT("%s-%s"), *(Instance->GetClass()->GetFName().GetPlainNameString()), *(NavConfig.Name.ToString()));
+				}
+				
+				constexpr bool bMarkDirty = false;
 				Instance->SetActorLabel(StrName, bMarkDirty);
 			}
 #endif // WITH_EDITOR
diff --git a/Engine/Source/Runtime/NavigationSystem/Public/NavigationData.h b/Engine/Source/Runtime/NavigationSystem/Public/NavigationData.h
--- a/Engine/Source/Runtime/NavigationSystem/Public/NavigationData.h
+++ b/Engine/Source/Runtime/NavigationSystem/Public/NavigationData.h
@@ -556,7 +556,6 @@ public:
 #if WITH_EDITOR
 	virtual void RerunConstructionScripts() override;
 	virtual void PostEditUndo() override;
-	virtual bool SupportsExternalPackaging() const override { return false; }
 	bool IsBuildingOnLoad() const { return bIsBuildingOnLoad; }
 	void SetIsBuildingOnLoad(bool bValue) { bIsBuildingOnLoad = bValue; }
 #endif // WITH_EDITOR
