Check that navigable areas have recast tiles indices that can fit in to an int32. An error log is output if this occurs.

diff --git a/Engine/Source/Runtime/NavigationSystem/Private/NavMesh/RecastNavMesh.cpp b/Engine/Source/Runtime/NavigationSystem/Private/NavMesh/RecastNavMesh.cpp
--- a/Engine/Source/Runtime/NavigationSystem/Private/NavMesh/RecastNavMesh.cpp
+++ b/Engine/Source/Runtime/NavigationSystem/Private/NavMesh/RecastNavMesh.cpp
@@ -1065,6 +1065,20 @@ bool ARecastNavMesh::GetNavMeshTileXY(const FVector& Point, int32& OutX, int32&
 	return RecastNavMeshImpl && RecastNavMeshImpl->GetNavMeshTileXY(Point, OutX, OutY);
 }
 
+bool ARecastNavMesh::CheckTileIndicesInValidRange(const FVector& Point, bool& bOutInRange) const
+{
+	const dtNavMesh* const DetourNavMesh = RecastNavMeshImpl ? RecastNavMeshImpl->DetourNavMesh : nullptr;
+	const bool bValidMesh = DetourNavMesh != nullptr;
+
+	if (bValidMesh)
+	{
+		const FVector RecastPt = Unreal2RecastPoint(Point);
+		bOutInRange = DetourNavMesh->isTileLocInValidRange(&RecastPt.X);
+	}
+
+	return bValidMesh;
+}
+
 void ARecastNavMesh::GetNavMeshTilesAt(int32 TileX, int32 TileY, TArray<int32>& Indices) const
 {
 	if (RecastNavMeshImpl)
diff --git a/Engine/Source/Runtime/NavigationSystem/Private/NavMesh/RecastNavMeshGenerator.cpp b/Engine/Source/Runtime/NavigationSystem/Private/NavMesh/RecastNavMeshGenerator.cpp
--- a/Engine/Source/Runtime/NavigationSystem/Private/NavMesh/RecastNavMeshGenerator.cpp
+++ b/Engine/Source/Runtime/NavigationSystem/Private/NavMesh/RecastNavMeshGenerator.cpp
@@ -4323,7 +4323,27 @@ FString FRecastTileGenerator::GetReferencerName() const
 	return TEXT("FRecastTileGenerator");
 }
 
-static int32 CalculateMaxTilesCount(const TNavStatArray<FBox>& NavigableAreas, FVector::FReal TileSizeInWorldUnits, FVector::FReal AvgLayersPerGridCell, const uint32 NavMeshVersion)
+namespace UE::NavMesh::Private
+{
+void CheckTileIndicesInValidRange(const TNavStatArray<FBox>& NavigableAreas, const ARecastNavMesh& NavMesh)
+{
+	auto CheckTileHelper = [&NavMesh](const FVector& Pos) {
+		// If there are no NavigableAreas then there won't be any indices for them so default true.
+		bool bIndiciesFitInInt32 = false;
+
+		ensure(NavMesh.CheckTileIndicesInValidRange(Pos, bIndiciesFitInInt32));
+
+		UE_CLOG(!bIndiciesFitInInt32, LogNavigation, Error, TEXT("Magnitude of Recast tile indicies are too large to fit in an int32 for NavigableAreas extent %s for %s"),*Pos.ToString(), *GetFullNameSafe(&NavMesh));
+	};
+
+	for (const FBox& AreaBounds : NavigableAreas)
+	{
+		CheckTileHelper(AreaBounds.Min);
+		CheckTileHelper(AreaBounds.Max);
+	}
+}
+
+int32 CalculateMaxTilesCount(const TNavStatArray<FBox>& NavigableAreas, FVector::FReal TileSizeInWorldUnits, FVector::FReal AvgLayersPerGridCell, const uint32 NavMeshVersion)
 {
 	int32 GridCellsCount = 0;
 	for (const FBox& AreaBounds : NavigableAreas)
@@ -4350,6 +4370,7 @@ static int32 CalculateMaxTilesCount(const TNavStatArray<FBox>& NavigableAreas, F
 	
 	return FMath::CeilToInt(GridCellsCount * AvgLayersPerGridCell);
 }
+} // UE::NavMesh::Private
 
 // Whether navmesh is static, does not support rebuild from geometry
 static bool IsGameStaticNavMesh(ARecastNavMesh* InNavMesh)
@@ -4537,6 +4558,8 @@ void FRecastNavMeshGenerator::Init()
 							, SavedNavParams->maxTiles, FMath::CeilToInt(FMath::Log2(static_cast<float>(SavedNavParams->maxTiles)))
 							, MaxTiles, FMath::CeilToInt(FMath::Log2(static_cast<float>(MaxTiles))));
 					}
+
+					UE::NavMesh::Private::CheckTileIndicesInValidRange(InclusionBounds, *DestNavMesh);
 				}
 			}
 		};
@@ -4648,6 +4671,8 @@ bool FRecastNavMeshGenerator::ConstructTiledNavMesh()
 				bSuccess = true;
 				NumActiveTiles = GetTilesCountHelper(DetourMesh);
 				DestNavMesh->GetRecastNavMeshImpl()->SetRecastMesh(DetourMesh);
+
+				UE::NavMesh::Private::CheckTileIndicesInValidRange(InclusionBounds, *DestNavMesh);
 			}
 		}
 
@@ -4689,7 +4714,7 @@ void FRecastNavMeshGenerator::CalcNavMeshProperties(int32& MaxTiles, int32& MaxP
 	int32 MaxRequestedTiles = 0;
 	if (DestNavMesh->IsResizable())
 	{
-		MaxRequestedTiles = CalculateMaxTilesCount(InclusionBounds, Config.tileSize * Config.cs, AvgLayersPerTile, DestNavMesh->NavMeshVersion);
+		MaxRequestedTiles = UE::NavMesh::Private::CalculateMaxTilesCount(InclusionBounds, Config.tileSize * Config.cs, AvgLayersPerTile, DestNavMesh->NavMeshVersion);
 	}
 	else
 	{
@@ -4847,13 +4872,15 @@ void FRecastNavMeshGenerator::TickAsyncBuild(float DeltaSeconds)
 
 void FRecastNavMeshGenerator::OnNavigationBoundsChanged()
 {
+	check(DestNavMesh);
+
 	UpdateNavigationBounds();
 	
 	dtNavMesh* DetourMesh = DestNavMesh->GetRecastNavMeshImpl() ? DestNavMesh->GetRecastNavMeshImpl()->GetRecastMesh() : nullptr;
 	if (!IsGameStaticNavMesh(DestNavMesh) && DestNavMesh->IsResizable() && DetourMesh)
 	{
 		// Check whether Navmesh size needs to be changed
-		int32 MaxRequestedTiles = CalculateMaxTilesCount(InclusionBounds, Config.tileSize * Config.cs, AvgLayersPerTile, DestNavMesh->NavMeshVersion);
+		const int32 MaxRequestedTiles = UE::NavMesh::Private::CalculateMaxTilesCount(InclusionBounds, Config.tileSize * Config.cs, AvgLayersPerTile, DestNavMesh->NavMeshVersion);
 		if (DetourMesh->getMaxTiles() != MaxRequestedTiles)
 		{
 			// Destroy current NavMesh
@@ -4873,6 +4900,8 @@ void FRecastNavMeshGenerator::OnNavigationBoundsChanged()
 				RebuildDirtyAreas(AsDirtyAreas);
 			}
 		}
+
+		UE::NavMesh::Private::CheckTileIndicesInValidRange(InclusionBounds, *DestNavMesh);
 	}
 }
 
diff --git a/Engine/Source/Runtime/NavigationSystem/Public/NavMesh/RecastNavMesh.h b/Engine/Source/Runtime/NavigationSystem/Public/NavMesh/RecastNavMesh.h
--- a/Engine/Source/Runtime/NavigationSystem/Public/NavMesh/RecastNavMesh.h
+++ b/Engine/Source/Runtime/NavigationSystem/Public/NavMesh/RecastNavMesh.h
@@ -996,6 +996,9 @@ public:
 	/** Retrieves XY coordinates of tile specified by position */
 	bool GetNavMeshTileXY(const FVector& Point, int32& OutX, int32& OutY) const;
 
+	/** Checks the supplied Points tile indicies can fit in the range of an int32 */
+	bool CheckTileIndicesInValidRange(const FVector& Point, bool& bOutInRange) const;
+
 	/** Retrieves all tile indices at matching XY coordinates */
 	void GetNavMeshTilesAt(int32 TileX, int32 TileY, TArray<int32>& Indices) const;
 
diff --git a/Engine/Source/Runtime/Navmesh/Private/Detour/DetourNavMesh.cpp b/Engine/Source/Runtime/Navmesh/Private/Detour/DetourNavMesh.cpp
--- a/Engine/Source/Runtime/Navmesh/Private/Detour/DetourNavMesh.cpp
+++ b/Engine/Source/Runtime/Navmesh/Private/Detour/DetourNavMesh.cpp
@@ -2089,10 +2089,38 @@ const dtMeshTile* dtNavMesh::getTile(int i) const
 	return &m_tiles[i];
 }
 
+bool dtNavMesh::isTileLocInValidRange(const dtReal tx, const dtReal ty) const
+{
+	return (tx >= (dtReal)std::numeric_limits<int>::min()) &&
+		(tx <= (dtReal)std::numeric_limits<int>::max()) &&
+		(ty >= (dtReal)std::numeric_limits<int>::min()) &&
+		(ty <= (dtReal)std::numeric_limits<int>::max());
+}
+
+void dtNavMesh::calcTileLoc(const dtReal* pos, dtReal* tx, dtReal* ty) const
+{
+	*tx = dtFloor((pos[0] - m_orig[0]) / m_tileWidth);
+	*ty = dtFloor((pos[2] - m_orig[2]) / m_tileHeight);
+}
+
 void dtNavMesh::calcTileLoc(const dtReal* pos, int* tx, int* ty) const
 {
-	*tx = (int)dtFloor((pos[0]-m_orig[0]) / m_tileWidth);
-	*ty = (int)dtFloor((pos[2]-m_orig[2]) / m_tileHeight);
+	dtReal txReal = 0.;
+	dtReal tyReal = 0.;
+
+	calcTileLoc(pos, &txReal, &tyReal);
+	
+	*tx = (int)txReal;
+	*ty = (int)tyReal;
+}
+
+bool dtNavMesh::isTileLocInValidRange(const dtReal* pos) const
+{
+	dtReal tx = 0.;
+	dtReal ty = 0.;
+
+	calcTileLoc(pos, &tx, &ty);
+	return isTileLocInValidRange(tx, ty);
 }
 
 dtStatus dtNavMesh::getTileAndPolyByRef(const dtPolyRef ref, const dtMeshTile** tile, const dtPoly** poly) const
diff --git a/Engine/Source/Runtime/Navmesh/Public/Detour/DetourNavMesh.h b/Engine/Source/Runtime/Navmesh/Public/Detour/DetourNavMesh.h
--- a/Engine/Source/Runtime/Navmesh/Public/Detour/DetourNavMesh.h
+++ b/Engine/Source/Runtime/Navmesh/Public/Detour/DetourNavMesh.h
@@ -531,6 +531,12 @@ public:
 	///  @param[out]	ty		The tile's y-location. (x, y)
 	void calcTileLoc(const dtReal* pos, int* tx, int* ty) const;
 
+
+	/// Calculates whether the tile grid location for the specified world position
+	/// can fit in the tile indices type (currently an int)
+	///  @param[in]	pos			The world position for the query. [(x, y, z)]
+	bool isTileLocInValidRange(const dtReal* pos) const;
+
 	/// Gets the tile at the specified grid location.
 	///  @param[in]	x		The tile's x-location. (x, y, layer)
 	///  @param[in]	y		The tile's y-location. (x, y, layer)
@@ -918,7 +924,15 @@ private:
 	/// Returns closest point on polygon.
 	void closestPointOnPolyInTile(const dtMeshTile* tile, unsigned int ip,
 								  const dtReal* pos, dtReal* closest) const;
-	
+
+private:
+	/// Calculates whether the tile grid location can fit in the tile indices type (currently an int)
+	///  @param[in]	tx			Tile X coord
+	///  @param[in]	ty			Tile Y coord
+	bool isTileLocInValidRange(const dtReal tx, const dtReal ty) const;
+	void calcTileLoc(const dtReal* pos, dtReal* tx, dtReal* ty) const;
+
+public:
 	dtNavMeshParams m_params;			///< Current initialization params. TODO: do not store this info twice.
 	dtReal m_orig[3];					///< Origin of the tile (0,0)
 	dtReal m_tileWidth, m_tileHeight;	///< Dimensions of each tile.
