Fix navmesh actor to use the beautified actor label
#rnx

diff --git a/Engine/Source/Runtime/NavigationSystem/Private/NavigationSystem.cpp b/Engine/Source/Runtime/NavigationSystem/Private/NavigationSystem.cpp
--- a/Engine/Source/Runtime/NavigationSystem/Private/NavigationSystem.cpp
+++ b/Engine/Source/Runtime/NavigationSystem/Private/NavigationSystem.cpp
@@ -3741,7 +3741,7 @@ ANavigationData* UNavigationSystemV1::CreateNavigationDataInstanceInLevel(const
 				}
 				
 				constexpr bool bMarkDirty = false;
-				Instance->SetActorLabel(StrName, bMarkDirty);
+				Instance->SetActorLabel(ActorLabel, bMarkDirty);
 			}
 #endif // WITH_EDITOR
 		}
