[Navmesh] Fix assert in FNavRegenTimeSlicer::TestTimeSliceFinished()

diff --git a/Engine/Source/Runtime/NavigationSystem/Private/NavigationSystem.cpp b/Engine/Source/Runtime/NavigationSystem/Private/NavigationSystem.cpp
--- a/Engine/Source/Runtime/NavigationSystem/Private/NavigationSystem.cpp
+++ b/Engine/Source/Runtime/NavigationSystem/Private/NavigationSystem.cpp
@@ -287,18 +287,18 @@ bool FNavRegenTimeSlicer::TestTimeSliceFinished() const
 #if ALLOW_TIME_SLICE_DEBUG
 	check(NavSystemDebug.Get());
 	check(NavSystemDebug->NavDataSet.IsValidIndex(NaxDataIdxDebug));
-
-	ARecastNavMesh* NavData = Cast<ARecastNavMesh>(NavSystemDebug->NavDataSet[NaxDataIdxDebug]);
-	check(NavData);
-
-	const double TimeSinceLastTested = Time - TimeLastTested;
-	if (TimeSinceLastTested >= NavData->TimeSliceLongDurationDebug)
+	
+	if (const ARecastNavMesh* NavData = Cast<ARecastNavMesh>(NavSystemDebug->NavDataSet[NaxDataIdxDebug]))
 	{
-		const FVector Pos = TileBoxDebug.GetCenter();
+		const double TimeSinceLastTested = Time - TimeLastTested;
+		if (TimeSinceLastTested >= NavData->TimeSliceLongDurationDebug)
+		{
+			const FVector Pos = TileBoxDebug.GetCenter();
 
-		// I'd quite like to make this a Warning, but it would be too frequently logged as things stand.
-		UE_LOG(LogNavigation, Verbose, TEXT("Nav mesh data: %s, tile at %d, %d, coordinate %f, %f, %f: %s is taking %f secs to partially regenerate!"), *NavData->GetName(), TileDebug.X, TileDebug.Y, Pos.X, Pos.Y, Pos.Z, *SectionNameDebug.ToString(), TimeSinceLastTested);
-		UE_VLOG_BOX(NavData, LogNavigation, Verbose, TileBoxDebug, FColor::Red, TEXT("Nav mesh data : %s,  tile at %d, %d, %s is taking %f secs to partially regenerate!"), *NavData->GetName(), TileDebug.X, TileDebug.Y, *SectionNameDebug.ToString(), TimeSinceLastTested);
+			// I'd quite like to make this a Warning, but it would be too frequently logged as things stand.
+			UE_LOG(LogNavigation, Verbose, TEXT("Nav mesh data: %s, tile at %d, %d, coordinate %f, %f, %f: %s is taking %f secs to partially regenerate!"), *NavData->GetName(), TileDebug.X, TileDebug.Y, Pos.X, Pos.Y, Pos.Z, *SectionNameDebug.ToString(), TimeSinceLastTested);
+			UE_VLOG_BOX(NavData, LogNavigation, Verbose, TileBoxDebug, FColor::Red, TEXT("Nav mesh data : %s,  tile at %d, %d, %s is taking %f secs to partially regenerate!"), *NavData->GetName(), TileDebug.X, TileDebug.Y, *SectionNameDebug.ToString(), TimeSinceLastTested);
+		}
 	}
 
 	// Reset SectionDebugName
