[Navmesh] When building a standard navmesh in a partitioned world, make sure the navigable world is loaded before building.

diff --git a/Engine/Source/Runtime/NavigationSystem/Private/NavMesh/RecastNavMesh.cpp b/Engine/Source/Runtime/NavigationSystem/Private/NavMesh/RecastNavMesh.cpp
--- a/Engine/Source/Runtime/NavigationSystem/Private/NavMesh/RecastNavMesh.cpp
+++ b/Engine/Source/Runtime/NavigationSystem/Private/NavMesh/RecastNavMesh.cpp
@@ -25,6 +25,10 @@
 
 #if WITH_EDITOR
 #include "ObjectEditorUtils.h"
+#include "WorldPartition/WorldPartition.h"
+#include "WorldPartition/WorldPartitionEditorLoaderAdapter.h"
+#include "WorldPartition/LoaderAdapter/LoaderAdapterShape.h"
+#include "LevelInstance/LevelInstanceSubsystem.h"
 #endif
 
 #if WITH_RECAST
@@ -548,6 +552,34 @@ void ARecastNavMesh::UpdateNavMeshDrawing()
 #endif // UE_BUILD_SHIPPING
 }
 
+void ARecastNavMesh::LoadBeforeGeneratorRebuild()
+{
+#if WITH_EDITOR
+	// If it's not a world partitioned navmesh but it's in a partitioned world, we need to make sure the navigable world is loaded before building the navmesh.
+	if (!bIsWorldPartitioned)
+	{
+		UWorld* World = GetWorld();
+		if (World && World->IsPartitionedWorld())
+		{
+			const UNavigationSystemV1* NavSys = FNavigationSystem::GetCurrent<UNavigationSystemV1>(World);
+			if (NavSys)
+			{
+				UWorldPartition* WorldPartition = World->GetWorldPartition();
+				check(WorldPartition);
+				
+				const FBox Bounds = NavSys->GetNavigableWorldBounds();
+				UWorldPartitionEditorLoaderAdapter* EditorLoaderAdapter = WorldPartition->CreateEditorLoaderAdapter<FLoaderAdapterShape>(World, Bounds, TEXT("Navigable World"));
+				EditorLoaderAdapter->GetLoaderAdapter()->SetUserCreated(false);
+				EditorLoaderAdapter->GetLoaderAdapter()->Load();
+
+				// Make sure level instances are loaded.
+				World->BlockTillLevelStreamingCompleted();
+			}
+		}
+	}
+#endif //WITH_EDITOR
+}
+
 void ARecastNavMesh::CleanUp()
 {
 	Super::CleanUp();
diff --git a/Engine/Source/Runtime/NavigationSystem/Private/NavigationData.cpp b/Engine/Source/Runtime/NavigationSystem/Private/NavigationData.cpp
--- a/Engine/Source/Runtime/NavigationSystem/Private/NavigationData.cpp
+++ b/Engine/Source/Runtime/NavigationSystem/Private/NavigationData.cpp
@@ -526,6 +526,10 @@ void ANavigationData::ConditionalConstructGenerator()
 
 void ANavigationData::RebuildAll()
 {
+	const double LoadTime = FPlatformTime::Seconds();
+	LoadBeforeGeneratorRebuild();
+	UE_LOG(LogNavigationDataBuild, Display, TEXT("   %s load time: %.2fs"), ANSI_TO_TCHAR(__FUNCTION__), (FPlatformTime::Seconds() - LoadTime));
+	
 	ConditionalConstructGenerator(); //recreate generator
 	
 	if (NavDataGenerator.IsValid())
@@ -537,7 +541,9 @@ void ANavigationData::RebuildAll()
 		}
 #endif // WITH_EDITOR
 
+		const double BuildTime = FPlatformTime::Seconds();
 		NavDataGenerator->RebuildAll();
+		UE_LOG(LogNavigationDataBuild, Display, TEXT("   %s NavDataGenerator->RebuildAll() time: %.2fs"), ANSI_TO_TCHAR(__FUNCTION__), (FPlatformTime::Seconds() - BuildTime));
 	}
 }
 
diff --git a/Engine/Source/Runtime/NavigationSystem/Public/NavMesh/RecastNavMesh.h b/Engine/Source/Runtime/NavigationSystem/Public/NavMesh/RecastNavMesh.h
--- a/Engine/Source/Runtime/NavigationSystem/Public/NavMesh/RecastNavMesh.h
+++ b/Engine/Source/Runtime/NavigationSystem/Public/NavMesh/RecastNavMesh.h
@@ -921,6 +921,8 @@ public:
 	/** RecastNavMesh instances are dynamically spawned and should not be coppied */
 	virtual bool ShouldExport() override { return false; }
 #endif
+
+	virtual void LoadBeforeGeneratorRebuild() override;
 	
 	virtual void CleanUp() override;
 
diff --git a/Engine/Source/Runtime/NavigationSystem/Public/NavigationData.h b/Engine/Source/Runtime/NavigationSystem/Public/NavigationData.h
--- a/Engine/Source/Runtime/NavigationSystem/Public/NavigationData.h
+++ b/Engine/Source/Runtime/NavigationSystem/Public/NavigationData.h
@@ -621,6 +621,9 @@ protected:
 public:
 	/** Creates new generator in case navigation supports it */
 	virtual void ConditionalConstructGenerator();
+
+	/** Any loading before NavDataGenerator->RebuildAll() */
+	virtual void LoadBeforeGeneratorRebuild() {}
 	
 	/** Triggers rebuild in case navigation supports it */
 	virtual void RebuildAll();
