[Navigation] Make ANavSystemConfigOverride not spatially loaded

diff --git a/Engine/Source/Runtime/NavigationSystem/Private/NavSystemConfigOverride.cpp b/Engine/Source/Runtime/NavigationSystem/Private/NavSystemConfigOverride.cpp
--- a/Engine/Source/Runtime/NavigationSystem/Private/NavSystemConfigOverride.cpp
+++ b/Engine/Source/Runtime/NavigationSystem/Private/NavSystemConfigOverride.cpp
@@ -23,6 +23,8 @@ ANavSystemConfigOverride::ANavSystemConfigOverride(const FObjectInitializer& Obj
 	RootComponent->Mobility = EComponentMobility::Static;
 
 #if WITH_EDITORONLY_DATA
+	bIsSpatiallyLoaded = false;
+	
 	SpriteComponent = CreateEditorOnlyDefaultSubobject<UBillboardComponent>(TEXT("Sprite"));
 
 	if (!IsRunningCommandlet())
