Proper fix for Server Assert at FNavRegenTimeSlicer::TestTimeSliceFinished upon initializing a session. Previous fix [REDACTED_HANDLE] did not fix the actual issue.

diff --git a/Engine/Source/Runtime/NavigationSystem/Private/NavMesh/RecastNavMeshGenerator.cpp b/Engine/Source/Runtime/NavigationSystem/Private/NavMesh/RecastNavMeshGenerator.cpp
--- a/Engine/Source/Runtime/NavigationSystem/Private/NavMesh/RecastNavMeshGenerator.cpp
+++ b/Engine/Source/Runtime/NavigationSystem/Private/NavMesh/RecastNavMeshGenerator.cpp
@@ -6066,13 +6066,41 @@ TArray<FNavTileRef> FRecastNavMeshGenerator::ProcessTileTasksSyncTimeSlicedAndGe
 			SyncTimeSlicedData.TimeSliceManager->GetTimeSlicer().EndTimeSliceAndAdjustDuration();
 		}
 
+#if ALLOW_TIME_SLICE_DEBUG
+		// Reset the debug function to make sure the captured variables can't be used when invalid.
+		// This is just bombproofing the code as TestTimeSliceFinished() should not be called until 
+		// ProcessTileTasksSyncTimeSlicedAndGetUpdatedTiles() is next called.
+		SyncTimeSlicedData.TimeSliceManager->GetTimeSlicer().DebugResetLongTimeSliceFunction();
+#endif // ALLOW_TIME_SLICE_DEBUG
+
 		return UpdatedTiles;
 	};
 
+#if ALLOW_TIME_SLICE_DEBUG
+	auto DebugLongTileRegenFunction = [this](FName DebugSectionName, double Duration)
+	{
+		const FRecastTileGenerator* const TileGenerator = SyncTimeSlicedData.TileGeneratorSync.Get();
+
+		// This shouldn't trigger but its fairly easy during development to accidentaly call TestTimeSliceFinished() when TileGenerator == nullptr.
+		if (ensure(TileGenerator))
+		{
+			const FVector Pos = SyncTimeSlicedData.TileGeneratorSync->GetTileBB().GetCenter();
+
+			check(DestNavMesh);
+
+			// I'd quite like to make this a Warning, but it would be too frequently logged as things stand.
+			UE_LOG(LogNavigation, Verbose, TEXT("Nav mesh data: %s, tile at %d, %d, coordinate %f, %f, %f, section %s is taking %f secs to partially regenerate!"), *DestNavMesh->GetName(), TileGenerator->GetTileX(), TileGenerator->GetTileY(), Pos.X, Pos.Y, Pos.Z, *DebugSectionName.ToString(), Duration);
+			UE_VLOG_BOX(DestNavMesh, LogNavigation, Verbose, TileGenerator->GetTileBB(), FColor::Red, TEXT("Nav mesh data : %s,  tile at %d, %d, section %s is taking %f secs to partially regenerate!"), *DestNavMesh->GetName(), TileGenerator->GetTileX(), TileGenerator->GetTileY(), *DebugSectionName.ToString(), Duration);
+		}
+	};
+
+	SyncTimeSlicedData.TimeSliceManager->GetTimeSlicer().DebugSetLongTimeSliceData(DebugLongTileRegenFunction, DestNavMesh->TimeSliceLongDurationDebug);
+#endif // ALLOW_TIME_SLICE_DEBUG
+
 	int32 NextPendingDirtyTileIndex = INDEX_NONE;
 	const bool bHadWorkToDo = HasWorkToDo(NextPendingDirtyTileIndex);
 
-	//only calculate the time slice and process tiles if we have work to do
+	// Only calculate the time slice and process tiles if we have work to do.
 	if (bHadWorkToDo)
 	{
 		SyncTimeSlicedData.TimeSliceManager->GetTimeSlicer().StartTimeSlice();
@@ -6100,14 +6128,6 @@ TArray<FNavTileRef> FRecastNavMeshGenerator::ProcessTileTasksSyncTimeSlicedAndGe
 
 				check(SyncTimeSlicedData.TileGeneratorSync);
 
-#if ALLOW_TIME_SLICE_DEBUG
-				UNavigationSystemV1* NavSys = FNavigationSystem::GetCurrent<UNavigationSystemV1>(GetWorld());
-
-				check(NavSys);
-
-				SyncTimeSlicedData.TimeSliceManager->GetTimeSlicer().SetCurrentTileDebug(*NavSys, SyncTimeSlicedData.TimeSliceManager->GetNavDataIdx(), TileLocation, SyncTimeSlicedData.TileGeneratorSync->GetTileBB());
-#endif
-
 				SyncTimeSlicedData.CurrentTileRegenDuration = 0.;
 
 				if (SyncTimeSlicedData.TileGeneratorSync->HasDataToBuild())
@@ -6133,6 +6153,8 @@ TArray<FNavTileRef> FRecastNavMeshGenerator::ProcessTileTasksSyncTimeSlicedAndGe
 			}
 			else
 			{
+				check(SyncTimeSlicedData.TileGeneratorSync);
+
 				TileLocation.X = SyncTimeSlicedData.TileGeneratorSync->GetTileX();
 				TileLocation.Y = SyncTimeSlicedData.TileGeneratorSync->GetTileY();
 			}
@@ -6195,7 +6217,6 @@ TArray<FNavTileRef> FRecastNavMeshGenerator::ProcessTileTasksSyncTimeSlicedAndGe
 			{
 				//reset state to Init for next tile to be processed
 				SyncTimeSlicedData.ProcessTileTasksSyncState = EProcessTileTasksSyncTimeSlicedState::Init;
-				SyncTimeSlicedData.TileGeneratorSync.Reset();
 
 				SyncTimeSlicedData.CurrentTileRegenDuration += (FPlatformTime::Seconds() - TimeStartProcessingTileThisFrame);
 
@@ -6206,7 +6227,13 @@ TArray<FNavTileRef> FRecastNavMeshGenerator::ProcessTileTasksSyncTimeSlicedAndGe
 				MARK_TIMESLICE_SECTION_DEBUG(SyncTimeSlicedData.TimeSliceManager->GetTimeSlicer(), FinishTile);
 
 				//test time slice 
-				if (SyncTimeSlicedData.TimeSliceManager->GetTimeSlicer().TestTimeSliceFinished())
+				const bool bTimeSliceFinished = SyncTimeSlicedData.TimeSliceManager->GetTimeSlicer().TestTimeSliceFinished();
+
+				//reset TileGeneratorSync after the last call to TestTimeSliceFinished for this tile, otherwise we may end up
+				//trying to access TileGeneratorSync from DebugLongTileRegenFunction()
+				SyncTimeSlicedData.TileGeneratorSync.Reset();
+
+				if (bTimeSliceFinished)
 				{
 					//we just calculated and set TileRegenDuration so no need to calculate it again
 					return EndFunction(false /* bCalcTileRegenDuration */, bHadWorkToDo);
diff --git a/Engine/Source/Runtime/NavigationSystem/Private/NavigationSystem.cpp b/Engine/Source/Runtime/NavigationSystem/Private/NavigationSystem.cpp
--- a/Engine/Source/Runtime/NavigationSystem/Private/NavigationSystem.cpp
+++ b/Engine/Source/Runtime/NavigationSystem/Private/NavigationSystem.cpp
@@ -269,13 +269,17 @@ void FNavRegenTimeSlicer::EndTimeSliceAndAdjustDuration()
 }
 
 #if ALLOW_TIME_SLICE_DEBUG
-void FNavRegenTimeSlicer::SetCurrentTileDebug(UNavigationSystemV1& InNavSystem, const int InNavDataIdx, const FIntPoint& InTile, const FBox& InTileBox) const
+void FNavRegenTimeSlicer::DebugSetLongTimeSliceData(TFunction<void(FName, double)> LongTimeSliceFunction, double LongTimeSliceDuration) const
 {
-	NavSystemDebug = &InNavSystem;
-	TileDebug = InTile;
-	TileBoxDebug = InTileBox;
-	NaxDataIdxDebug = InNavDataIdx;
+	DebugLongTimeSliceFunction = LongTimeSliceFunction;
+	DebugLongTimeSliceDuration = LongTimeSliceDuration;
 }
+
+void FNavRegenTimeSlicer::DebugResetLongTimeSliceFunction() const
+{
+	DebugLongTimeSliceFunction.Reset();
+}
+
 #endif // ALLOW_TIME_SLICE_DEBUG
 
 bool FNavRegenTimeSlicer::TestTimeSliceFinished() const
@@ -285,24 +289,17 @@ bool FNavRegenTimeSlicer::TestTimeSliceFinished() const
 	const double Time = FPlatformTime::Seconds();
 
 #if ALLOW_TIME_SLICE_DEBUG
-	check(NavSystemDebug.Get());
-	check(NavSystemDebug->NavDataSet.IsValidIndex(NaxDataIdxDebug));
-	
-	if (const ARecastNavMesh* NavData = Cast<ARecastNavMesh>(NavSystemDebug->NavDataSet[NaxDataIdxDebug]))
+	const double TimeSinceLastTested = Time - TimeLastTested;
+	if (TimeSinceLastTested >= DebugLongTimeSliceDuration)
 	{
-		const double TimeSinceLastTested = Time - TimeLastTested;
-		if (TimeSinceLastTested >= NavData->TimeSliceLongDurationDebug)
+		if (ensureMsgf(DebugLongTimeSliceFunction, TEXT("DebugLongTimeSliceFunction should be setup! Call DebugSetLongTimeSliceData() prior to TestTimeSliceFinished()!")))
 		{
-			const FVector Pos = TileBoxDebug.GetCenter();
-
-			// I'd quite like to make this a Warning, but it would be too frequently logged as things stand.
-			UE_LOG(LogNavigation, Verbose, TEXT("Nav mesh data: %s, tile at %d, %d, coordinate %f, %f, %f: %s is taking %f secs to partially regenerate!"), *NavData->GetName(), TileDebug.X, TileDebug.Y, Pos.X, Pos.Y, Pos.Z, *SectionNameDebug.ToString(), TimeSinceLastTested);
-			UE_VLOG_BOX(NavData, LogNavigation, Verbose, TileBoxDebug, FColor::Red, TEXT("Nav mesh data : %s,  tile at %d, %d, %s is taking %f secs to partially regenerate!"), *NavData->GetName(), TileDebug.X, TileDebug.Y, *SectionNameDebug.ToString(), TimeSinceLastTested);
+			DebugLongTimeSliceFunction(DebugSectionName, TimeSinceLastTested);
 		}
 	}
 
 	// Reset SectionDebugName
-	SectionNameDebug = FNavigationSystem::TimeSliceDefaultSectionNameDebug;
+	DebugSectionName = FNavigationSystem::DebugTimeSliceDefaultSectionName;
 #endif // ALLOW_TIME_SLICE_DEBUG
 
 	TimeLastTested = Time;
diff --git a/Engine/Source/Runtime/NavigationSystem/Public/NavigationSystem.h b/Engine/Source/Runtime/NavigationSystem/Public/NavigationSystem.h
--- a/Engine/Source/Runtime/NavigationSystem/Public/NavigationSystem.h
+++ b/Engine/Source/Runtime/NavigationSystem/Public/NavigationSystem.h
@@ -84,7 +84,7 @@ namespace FNavigationSystem
 	void NAVIGATIONSYSTEM_API MakeAllComponentsNeverAffectNav(AActor& Actor);
 
 #if ALLOW_TIME_SLICE_DEBUG
-	const FName TimeSliceDefaultSectionNameDebug= FName(TEXT("DefaultSection"));
+	const FName DebugTimeSliceDefaultSectionName = FName(TEXT("DefaultSection"));
 #endif // ALLOW_TIME_SLICE_DEBUG
 }
 
@@ -111,7 +111,7 @@ namespace ENavigationBuildLock
 #if ALLOW_TIME_SLICE_DEBUG
 #define MARK_TIMESLICE_SECTION_DEBUG(TIME_SLICER, TIME_SLICE_FNAME) \
 static const FName TIME_SLICE_FNAME(TEXT(#TIME_SLICE_FNAME)); \
-TIME_SLICER.SetSectionNameDebug(TIME_SLICE_FNAME);
+TIME_SLICER.DebugSetSectionName(TIME_SLICE_FNAME);
 #else
 #define MARK_TIMESLICE_SECTION_DEBUG(TIME_SLICER, TIME_SLICE_FNAME) ;
 #endif
@@ -146,15 +146,21 @@ public:
 
 
 #if ALLOW_TIME_SLICE_DEBUG
-	void SetCurrentTileDebug(UNavigationSystemV1& InNavSystem, const int InNavDataIdx, const FIntPoint& InTile, const FBox& InTileBox) const;
+	/*
+	 * Sets data used for debugging time slices that are taking too long to process.
+	 * @param LongTimeSliceFunction function called when a time slice is taking too long to process.
+	 * @param LongTimeSliceDuration when a time slice takes longer than this duration LongTimeSliceFunction will be called.
+	 */
+	void DebugSetLongTimeSliceData(TFunction<void(FName, double)> LongTimeSliceFunction, double LongTimeSliceDuration) const;
+	void DebugResetLongTimeSliceFunction() const;
 
 	/**
 	 *  Sets the debug name for a time sliced section of code.
 	 *  Do not call this directly use MARK_TIMESLICE_SECTION_DEBUG
 	 */
-	void SetSectionNameDebug(FName InSectionNameDebug) const
+	void DebugSetSectionName(FName InDebugSectionName) const
 	{
-		SectionNameDebug = InSectionNameDebug;
+		DebugSectionName = InDebugSectionName;
 	}
 #endif // ALLOW_TIME_SLICE_DEBUG
 
@@ -166,11 +172,9 @@ protected:
 	mutable bool bTimeSliceFinishedCached = false;
 
 #if ALLOW_TIME_SLICE_DEBUG
-	mutable TWeakObjectPtr<UNavigationSystemV1> NavSystemDebug;
-	mutable int NaxDataIdxDebug = 0;
-	mutable FIntPoint TileDebug;
-	mutable FBox TileBoxDebug;
-	mutable FName SectionNameDebug = FNavigationSystem::TimeSliceDefaultSectionNameDebug;
+	mutable TFunction<void(FName, double)> DebugLongTimeSliceFunction;
+	mutable double DebugLongTimeSliceDuration = 0.;
+	mutable FName DebugSectionName = FNavigationSystem::DebugTimeSliceDefaultSectionName;
 #endif
 };
 
