Attempted fix for crash on startup due to possible corrupt DDC of some NavCollisionData

diff --git a/Engine/Source/Runtime/NavigationSystem/Private/NavCollision.cpp b/Engine/Source/Runtime/NavigationSystem/Private/NavCollision.cpp
--- a/Engine/Source/Runtime/NavigationSystem/Private/NavCollision.cpp
+++ b/Engine/Source/Runtime/NavigationSystem/Private/NavCollision.cpp
@@ -90,7 +90,7 @@ public:
 
 	virtual const TCHAR* GetVersionString() const override
 	{
-		return TEXT("2983273BE5D6425391E8C0B3D1C0A55D");
+		return TEXT("982ABB61BCEA412C832BC21B87BF3538");
 	}
 
 	virtual FString GetPluginSpecificCacheKeySuffix() const override
