[Navmesh] Add log tracking FRecastNavMeshGenerator RebuildAll() time (independently of the loading time)
Change OnNavigationGenerationFinished() to display actor label and full name instead of package name.

diff --git a/Engine/Source/Runtime/NavigationSystem/Private/NavMesh/RecastNavMeshGenerator.cpp b/Engine/Source/Runtime/NavigationSystem/Private/NavMesh/RecastNavMeshGenerator.cpp
--- a/Engine/Source/Runtime/NavigationSystem/Private/NavMesh/RecastNavMeshGenerator.cpp
+++ b/Engine/Source/Runtime/NavigationSystem/Private/NavMesh/RecastNavMeshGenerator.cpp
@@ -4754,6 +4754,10 @@ bool FRecastNavMeshGenerator::RebuildAll()
 		// There are no navigation bounds to build, probably navmesh was resized and we just need to update debug draw
 		DestNavMesh->RequestDrawingUpdate();
 	}
+	else
+	{
+		RebuildAllStartTime = FPlatformTime::Seconds();
+	}
 
 	return true;
 }
@@ -6370,6 +6374,12 @@ TArray<FNavTileRef> FRecastNavMeshGenerator::ProcessTileTasksAndGetUpdatedTiles(
 	{
 		QUICK_SCOPE_CYCLE_COUNTER(STAT_RecastNavMeshGenerator_OnNavMeshGenerationFinished);
 
+		if (RebuildAllStartTime != 0)
+		{
+			UE_LOG(LogNavigationDataBuild, Display, TEXT("   %s build time: %.2fs"), ANSI_TO_TCHAR(__FUNCTION__), (FPlatformTime::Seconds() - RebuildAllStartTime));
+			RebuildAllStartTime = 0;
+		}
+		
 		DestNavMesh->OnNavMeshGenerationFinished();
 	}
 
diff --git a/Engine/Source/Runtime/NavigationSystem/Private/NavigationData.cpp b/Engine/Source/Runtime/NavigationSystem/Private/NavigationData.cpp
--- a/Engine/Source/Runtime/NavigationSystem/Private/NavigationData.cpp
+++ b/Engine/Source/Runtime/NavigationSystem/Private/NavigationData.cpp
@@ -541,9 +541,7 @@ void ANavigationData::RebuildAll()
 		}
 #endif // WITH_EDITOR
 
-		const double BuildTime = FPlatformTime::Seconds();
 		NavDataGenerator->RebuildAll();
-		UE_LOG(LogNavigationDataBuild, Display, TEXT("   %s NavDataGenerator->RebuildAll() time: %.2fs"), ANSI_TO_TCHAR(__FUNCTION__), (FPlatformTime::Seconds() - BuildTime));
 	}
 }
 
diff --git a/Engine/Source/Runtime/NavigationSystem/Private/NavigationSystem.cpp b/Engine/Source/Runtime/NavigationSystem/Private/NavigationSystem.cpp
--- a/Engine/Source/Runtime/NavigationSystem/Private/NavigationSystem.cpp
+++ b/Engine/Source/Runtime/NavigationSystem/Private/NavigationSystem.cpp
@@ -3896,7 +3896,7 @@ void UNavigationSystemV1::OnNavigationGenerationFinished(ANavigationData& NavDat
 #if WITH_EDITOR
 	if (GetWorld()->IsGameWorld() == false)
 	{
-		UE_LOG(LogNavigationDataBuild, Log, TEXT("Navigation data generation finished in %s."), *NavData.GetPackage()->GetName());
+		UE_LOG(LogNavigationDataBuild, Log, TEXT("Navigation data generation finished for %s (%s)."), *NavData.GetActorLabel(), *NavData.GetFullName());
 	}
 
 	// Reset bIsBuildingOnLoad
diff --git a/Engine/Source/Runtime/NavigationSystem/Public/NavMesh/RecastNavMeshGenerator.h b/Engine/Source/Runtime/NavigationSystem/Public/NavMesh/RecastNavMeshGenerator.h
--- a/Engine/Source/Runtime/NavigationSystem/Public/NavMesh/RecastNavMeshGenerator.h
+++ b/Engine/Source/Runtime/NavigationSystem/Public/NavMesh/RecastNavMeshGenerator.h
@@ -929,6 +929,8 @@ protected:
 
 	/** Use this if you don't want your tiles to start at (0,0,0) */
 	FVector RcNavMeshOrigin;
+
+	double RebuildAllStartTime = 0;
 	
 	uint32 bInitialized:1;
 
