Fixed NavMesh KPK error in when playing in a level with unsupported nav agents
#tests Entered a level that has some unsupported nav agents and made sure all the supported nav agents are still able to navigate. Checked that the modified code isn't run in levels that don't have any unsupported nav agents.

diff --git a/Engine/Source/Runtime/NavigationSystem/Private/NavigationSystem.cpp b/Engine/Source/Runtime/NavigationSystem/Private/NavigationSystem.cpp
--- a/Engine/Source/Runtime/NavigationSystem/Private/NavigationSystem.cpp
+++ b/Engine/Source/Runtime/NavigationSystem/Private/NavigationSystem.cpp
@@ -2241,7 +2241,7 @@ void UNavigationSystemV1::ProcessRegistrationCandidates()
 
 			if (Result != RegistrationSuccessful && Result != RegistrationFailed_DataPendingKill)
 			{
-				NavDataPtr->CleanUpAndMarkPendingKill();
+				NavDataPtr->Destroy();
 				if (NavDataPtr == MainNavData)
 				{
 					MainNavData = NULL;
