Fix crash in PIE on navmesh RebuildAll() using normal navmesh in a WP world
#rnx

diff --git a/Engine/Source/Runtime/NavigationSystem/Private/NavMesh/RecastNavMesh.cpp b/Engine/Source/Runtime/NavigationSystem/Private/NavMesh/RecastNavMesh.cpp
--- a/Engine/Source/Runtime/NavigationSystem/Private/NavMesh/RecastNavMesh.cpp
+++ b/Engine/Source/Runtime/NavigationSystem/Private/NavMesh/RecastNavMesh.cpp
@@ -559,7 +559,7 @@ void ARecastNavMesh::LoadBeforeGeneratorRebuild()
 	if (!bIsWorldPartitioned)
 	{
 		UWorld* World = GetWorld();
-		if (World && World->IsPartitionedWorld())
+		if (World && World->IsPartitionedWorld() && !World->IsGameWorld())
 		{
 			const UNavigationSystemV1* NavSys = FNavigationSystem::GetCurrent<UNavigationSystemV1>(World);
 			if (NavSys)
