Make sure AbstractNavData doesn't support external packaging since it's always transient.

#rnx

diff --git a/Engine/Source/Runtime/NavigationSystem/Public/AbstractNavData.h b/Engine/Source/Runtime/NavigationSystem/Public/AbstractNavData.h
--- a/Engine/Source/Runtime/NavigationSystem/Public/AbstractNavData.h
+++ b/Engine/Source/Runtime/NavigationSystem/Public/AbstractNavData.h
@@ -61,6 +61,12 @@ public:
 
 	virtual void PostLoad() override;
 
+#if WITH_EDITOR
+	// Begin AActor overrides
+	virtual bool SupportsExternalPackaging() const override { return false; }
+	// End AActor overrides
+#endif
+
 	// Begin ANavigationData overrides
 	virtual void BatchRaycast(TArray<FNavigationRaycastWork>& Workload, FSharedConstNavQueryFilter QueryFilter, const UObject* Querier = NULL) const override {};
 	virtual bool FindMoveAlongSurface(const FNavLocation& StartLocation, const FVector& TargetPosition, FNavLocation& OutLocation, FSharedConstNavQueryFilter Filter = NULL, const UObject* Querier = NULL) const override { return false;  };
