[Navmesh] Fix transformed convex collision generating wrong navmesh

diff --git a/Engine/Source/Runtime/NavigationSystem/Private/NavMesh/RecastNavMeshGenerator.cpp b/Engine/Source/Runtime/NavigationSystem/Private/NavMesh/RecastNavMeshGenerator.cpp
--- a/Engine/Source/Runtime/NavigationSystem/Private/NavMesh/RecastNavMeshGenerator.cpp
+++ b/Engine/Source/Runtime/NavigationSystem/Private/NavMesh/RecastNavMeshGenerator.cpp
@@ -735,7 +735,6 @@ FORCEINLINE_DEBUGGABLE void ExportRigidBodyConvexElements(UBodySetup& BodySetup,
 {
 	const int32 ConvexCount = BodySetup.AggGeom.ConvexElems.Num();
 	FKConvexElem const* ConvexElem = BodySetup.AggGeom.ConvexElems.GetData();
-	const FTransform NegXScale(FQuat::Identity, FVector::ZeroVector, FVector(-1, 1, 1));
 
 	for (int32 i = 0; i < ConvexCount; ++i, ++ConvexElem)
 	{
@@ -744,8 +743,7 @@ FORCEINLINE_DEBUGGABLE void ExportRigidBodyConvexElements(UBodySetup& BodySetup,
 
 		if (ConvexElem->GetChaosConvexMesh())
 		{
-			// TODO use ConvexElem->GetTransform?() transform?
-			ExportChaosConvexMesh(ConvexElem, LocalToWorld, VertexBuffer, IndexBuffer, UnrealBounds);
+			ExportChaosConvexMesh(ConvexElem, ConvexElem->GetTransform() * LocalToWorld, VertexBuffer, IndexBuffer, UnrealBounds);
 		}
 	}
 }
