Added GetPolyTileRef method to retrieve a FNavTileRef from a NavNodeRef

diff --git a/Engine/Source/Runtime/NavigationSystem/Private/NavMesh/PImplRecastNavMesh.cpp b/Engine/Source/Runtime/NavigationSystem/Private/NavMesh/PImplRecastNavMesh.cpp
--- a/Engine/Source/Runtime/NavigationSystem/Private/NavMesh/PImplRecastNavMesh.cpp
+++ b/Engine/Source/Runtime/NavigationSystem/Private/NavMesh/PImplRecastNavMesh.cpp
@@ -2381,6 +2381,22 @@ bool FPImplRecastNavMesh::GetPolyTileIndex(NavNodeRef PolyID, uint32& PolyIndex,
 	return false;
 }
 
+bool FPImplRecastNavMesh::GetPolyTileRef(NavNodeRef PolyId, uint32& OutPolyIndex, FNavTileRef& OutTileRef) const
+{
+	if (DetourNavMesh && PolyId)
+	{
+		// Similar to UE::NavMesh::Private::GetTileRefFromPolyRef
+		unsigned int Salt = 0;
+		unsigned int TileIndex = 0;
+		DetourNavMesh->decodePolyId(PolyId, Salt, TileIndex, OutPolyIndex);
+		const dtTileRef TileRef = DetourNavMesh->encodePolyId(Salt, TileIndex, 0);
+		OutTileRef = FNavTileRef(TileRef);
+		return true;
+	}
+
+	return false;
+}
+
 bool FPImplRecastNavMesh::GetClosestPointOnPoly(NavNodeRef PolyID, const FVector& TestPt, FVector& PointOnPoly) const
 {
 	if (DetourNavMesh && PolyID)
diff --git a/Engine/Source/Runtime/NavigationSystem/Public/NavMesh/PImplRecastNavMesh.h b/Engine/Source/Runtime/NavigationSystem/Public/NavMesh/PImplRecastNavMesh.h
--- a/Engine/Source/Runtime/NavigationSystem/Public/NavMesh/PImplRecastNavMesh.h
+++ b/Engine/Source/Runtime/NavigationSystem/Public/NavMesh/PImplRecastNavMesh.h
@@ -181,6 +181,8 @@ public:
 	bool GetClosestPointOnPoly(NavNodeRef PolyID, const FVector& TestPt, FVector& PointOnPoly) const;
 	/** Decode poly ID into tile index and poly index */
 	bool GetPolyTileIndex(NavNodeRef PolyID, uint32& PolyIndex, uint32& TileIndex) const;
+	/** Decode poly ID into FNavTileRef and poly index */
+	bool GetPolyTileRef(NavNodeRef PolyId, uint32& OutPolyIndex, FNavTileRef& OutTileRef) const;
 	/** Retrieves user ID for given offmesh link poly */
 	uint32 GetLinkUserId(NavNodeRef LinkPolyID) const;
 	/** Retrieves start and end point of offmesh link */
