[Navmesh] Only include navigation relevant component when computing UNavModifierComponent navigation bounds

diff --git a/Engine/Source/Runtime/NavigationSystem/Private/NavModifierComponent.cpp b/Engine/Source/Runtime/NavigationSystem/Private/NavModifierComponent.cpp
--- a/Engine/Source/Runtime/NavigationSystem/Private/NavModifierComponent.cpp
+++ b/Engine/Source/Runtime/NavigationSystem/Private/NavModifierComponent.cpp
@@ -36,7 +36,7 @@ void UNavModifierComponent::CalcAndCacheBounds() const
 		for (UActorComponent* Component : MyOwner->GetComponents())
 		{
 			UPrimitiveComponent* PrimComp = Cast<UPrimitiveComponent>(Component);
-			if (PrimComp && PrimComp->IsRegistered() && PrimComp->IsCollisionEnabled())
+			if (PrimComp && PrimComp->IsRegistered() && PrimComp->IsCollisionEnabled() && PrimComp->CanEverAffectNavigation())
 			{
 				UBodySetup* BodySetup = PrimComp->GetBodySetup();
 				if (BodySetup)
