#ue5 - Move ConvertArrayType and ConvertArrayTypeClampMax to UE::LWC namespace.

diff --git a/Engine/Plugins/FX/Niagara/Source/Niagara/Private/NiagaraDataInterfaceSkeletalMesh.cpp b/Engine/Plugins/FX/Niagara/Source/Niagara/Private/NiagaraDataInterfaceSkeletalMesh.cpp
--- a/Engine/Plugins/FX/Niagara/Source/Niagara/Private/NiagaraDataInterfaceSkeletalMesh.cpp
+++ b/Engine/Plugins/FX/Niagara/Source/Niagara/Private/NiagaraDataInterfaceSkeletalMesh.cpp
@@ -476,7 +476,7 @@ void FSkeletalMeshSkinningData::UpdateBoneTransforms()
 	else
 	{
 		SkelComp->CacheRefToLocalMatrices(CurrBones);
-		CurrTransforms = LWC::ConvertArrayType<FTransform3f>(SkelComp->GetComponentSpaceTransforms());	// LWC_TODO: Perf pessimization
+		CurrTransforms = UE::LWC::ConvertArrayType<FTransform3f>(SkelComp->GetComponentSpaceTransforms());	// LWC_TODO: Perf pessimization
 	}
 }
 
diff --git a/Engine/Source/Editor/UnrealEd/Private/ImportUtils/StaticMeshImportUtils.cpp b/Engine/Source/Editor/UnrealEd/Private/ImportUtils/StaticMeshImportUtils.cpp
--- a/Engine/Source/Editor/UnrealEd/Private/ImportUtils/StaticMeshImportUtils.cpp
+++ b/Engine/Source/Editor/UnrealEd/Private/ImportUtils/StaticMeshImportUtils.cpp
@@ -579,7 +579,7 @@ bool StaticMeshImportUtils::AddConvexGeomFromVertices( const TArray<FVector3f>&
 	}
 
 	FKConvexElem* ConvexElem = new(AggGeom->ConvexElems) FKConvexElem();
-	ConvexElem->VertexData = LWC::ConvertArrayType<FVector>(Verts);	// LWC_TODO: Perf pessimization
+	ConvexElem->VertexData = UE::LWC::ConvertArrayType<FVector>(Verts);	// LWC_TODO: Perf pessimization
 	ConvexElem->UpdateElemBox();
 
 	return true;
diff --git a/Engine/Source/Runtime/AIModule/Private/Navigation/CrowdManager.cpp b/Engine/Source/Runtime/AIModule/Private/Navigation/CrowdManager.cpp
--- a/Engine/Source/Runtime/AIModule/Private/Navigation/CrowdManager.cpp
+++ b/Engine/Source/Runtime/AIModule/Private/Navigation/CrowdManager.cpp
@@ -921,8 +921,8 @@ void UCrowdManager::CreateCrowdManager()
 			if (Info.Angles.Num() > 0 && Info.Angles.Num() == Info.Radii.Num())
 			{
 				// LWC_TODO_AI: Info.Angles and Info.Radii should be FReal. Probably not until after 5.0!
-				const TArray<FVector::FReal> Angles = LWC::ConvertArrayType<FVector::FReal>(Info.Angles);
-				const TArray<FVector::FReal> Radii = LWC::ConvertArrayType<FVector::FReal>(Info.Radii);
+				const TArray<FVector::FReal> Angles = UE::LWC::ConvertArrayType<FVector::FReal>(Info.Angles);
+				const TArray<FVector::FReal> Radii = UE::LWC::ConvertArrayType<FVector::FReal>(Info.Radii);
 
 				DetourCrowd->setObstacleAvoidancePattern(Idx, Angles.GetData(), Radii.GetData(), Angles.Num());
 			}
diff --git a/Engine/Source/Runtime/Core/Public/Math/UnrealMathUtility.h b/Engine/Source/Runtime/Core/Public/Math/UnrealMathUtility.h
--- a/Engine/Source/Runtime/Core/Public/Math/UnrealMathUtility.h
+++ b/Engine/Source/Runtime/Core/Public/Math/UnrealMathUtility.h
@@ -2442,50 +2442,53 @@ public:
 	UE_NODISCARD static FORCEINLINE float DynamicWeightedMovingAverage(T1&& CurrentSample, T2&& PreviousSample, T3&& MaxDistance, T4&& MinWeight, T5&& MaxWeight) { return DynamicWeightedMovingAverage((float)CurrentSample, (float)PreviousSample, (float)MaxDistance, (float)MinWeight, (float)MaxWeight); }
 };
 
-// LWC Conversion helpers - LWC_TODO: These are temporary and should be removed for UE5.0 final.
+// LWC Conversion helpers
+namespace UE
+{
 namespace LWC
 {
 
-// Convert array to a new type
-template<typename TDest, typename TSrc, typename InAllocatorType>
-TArray<TDest, InAllocatorType> ConvertArrayType(const TArray<TSrc, InAllocatorType>& From)
-{
-	//static_assert(!std::is_same<TDest, TSrc>::value, "Redundant call to ConvertArrayType");	// Unavoidable if supporting LWC toggle, but a useful check once LWC is locked to enabled.
-	if constexpr (std::is_same<TDest, TSrc>::value)
+	// Convert array to a new type
+	template<typename TDest, typename TSrc, typename InAllocatorType>
+	TArray<TDest, InAllocatorType> ConvertArrayType(const TArray<TSrc, InAllocatorType>& From)
 	{
-		return From;
-	}
-	else
-	{
-		TArray<TDest, InAllocatorType> Converted;
-		Converted.Reserve(From.Num());
-		for (const TSrc& Item : From)
+		//static_assert(!std::is_same<TDest, TSrc>::value, "Redundant call to ConvertArrayType");	// Unavoidable if supporting LWC toggle, but a useful check once LWC is locked to enabled.
+		if constexpr (std::is_same<TDest, TSrc>::value)
 		{
-			Converted.Add(static_cast<TDest>(Item));
+			return From;
+		}
+		else
+		{
+			TArray<TDest, InAllocatorType> Converted;
+			Converted.Reserve(From.Num());
+			for (const TSrc& Item : From)
+			{
+				Converted.Add(static_cast<TDest>(Item));
+			}
+			return Converted;
 		}
-		return Converted;
 	}
-}
 
-// Convert array to a new type and clamps values to the Max of TDest type
-template<typename TDest, typename TSrc, typename InAllocatorType>
-TArray<TDest, InAllocatorType> ConvertArrayTypeClampMax(const TArray<TSrc, InAllocatorType>& From)
-{
-	//static_assert(!std::is_same<TDest, TSrc>::value, "Redundant call to ConvertArrayType");	// Unavoidable if supporting LWC toggle, but a useful check once LWC is locked to enabled.
-	if constexpr (std::is_same<TDest, TSrc>::value)
-	{
-		return From;
-	}
-	else
+	// Convert array to a new type and clamps values to the Max of TDest type
+	template<typename TDest, typename TSrc, typename InAllocatorType>
+	TArray<TDest, InAllocatorType> ConvertArrayTypeClampMax(const TArray<TSrc, InAllocatorType>& From)
 	{
-		TArray<TDest, InAllocatorType> Converted;
-		Converted.Reserve(From.Num());
-		for (const TSrc& Item : From)
+		//static_assert(!std::is_same<TDest, TSrc>::value, "Redundant call to ConvertArrayType");	// Unavoidable if supporting LWC toggle, but a useful check once LWC is locked to enabled.
+		if constexpr (std::is_same<TDest, TSrc>::value)
+		{
+			return From;
+		}
+		else
 		{
-			Converted.Add(FMath::Min(TNumericLimits<TDest>::Max(), static_cast<TDest>(Item)));
+			TArray<TDest, InAllocatorType> Converted;
+			Converted.Reserve(From.Num());
+			for (const TSrc& Item : From)
+			{
+				Converted.Add(FMath::Min(TNumericLimits<TDest>::Max(), static_cast<TDest>(Item)));
+			}
+			return Converted;
 		}
-		return Converted;
 	}
-}
 
-}
\ No newline at end of file
+} // namespace LWC
+} // namespace UE
diff --git a/Engine/Source/Runtime/Engine/Private/AI/Navigation/NavigationModifier.cpp b/Engine/Source/Runtime/Engine/Private/AI/Navigation/NavigationModifier.cpp
--- a/Engine/Source/Runtime/Engine/Private/AI/Navigation/NavigationModifier.cpp
+++ b/Engine/Source/Runtime/Engine/Private/AI/Navigation/NavigationModifier.cpp
@@ -345,7 +345,7 @@ void FAreaNavModifier::InitializeConvex(const TNavStatArray<FVector>& InPoints,
 	check(InPoints.IsValidIndex(FirstIndex) && InPoints.IsValidIndex(LastIndex-1));
 
 	Init(InAreaClass);
-	SetConvex(LWC::ConvertArrayType<FVector>(InPoints).GetData(), FirstIndex, LastIndex, ENavigationCoordSystem::Unreal, LocalToWorld);	// LWC_TODO: Perf pessimization
+	SetConvex(UE::LWC::ConvertArrayType<FVector>(InPoints).GetData(), FirstIndex, LastIndex, ENavigationCoordSystem::Unreal, LocalToWorld);	// LWC_TODO: Perf pessimization
 }
 
 void FAreaNavModifier::InitializePerInstanceConvex(const TNavStatArray<FVector>& InPoints, const int32 FirstIndex, const int32 LastIndex, const TSubclassOf<UNavAreaBase> InAreaClass)
@@ -353,7 +353,7 @@ void FAreaNavModifier::InitializePerInstanceConvex(const TNavStatArray<FVector>&
 	check(InPoints.IsValidIndex(FirstIndex) && InPoints.IsValidIndex(LastIndex - 1));
 
 	Init(InAreaClass);
-	SetPerInstanceConvex(LWC::ConvertArrayType<FVector>(InPoints).GetData(), FirstIndex, LastIndex);	// LWC_TODO: Perf pessimization
+	SetPerInstanceConvex(UE::LWC::ConvertArrayType<FVector>(InPoints).GetData(), FirstIndex, LastIndex);	// LWC_TODO: Perf pessimization
 }
 
 FAreaNavModifier::FAreaNavModifier(const UBrushComponent* BrushComponent, const TSubclassOf<UNavAreaBase> InAreaClass)
@@ -919,7 +919,7 @@ void FCompositeNavModifier::CreateAreaModifiers(const UPrimitiveComponent* PrimC
 		const FKConvexElem& ConvexElem = BodySetup->AggGeom.ConvexElems[Idx];
 		if (ConvexElem.VertexData.Num() > 0)
 		{
-			FAreaNavModifier AreaMod(LWC::ConvertArrayType<FVector>(ConvexElem.VertexData), 0, ConvexElem.VertexData.Num(), ENavigationCoordSystem::Unreal, PrimComp->GetComponentTransform(), AreaClass);
+			FAreaNavModifier AreaMod(UE::LWC::ConvertArrayType<FVector>(ConvexElem.VertexData), 0, ConvexElem.VertexData.Num(), ENavigationCoordSystem::Unreal, PrimComp->GetComponentTransform(), AreaClass);
 			Add(AreaMod);
 		}
 		else
diff --git a/Engine/Source/Runtime/Engine/Private/PhysicsEngine/PhysAnim.cpp b/Engine/Source/Runtime/Engine/Private/PhysicsEngine/PhysAnim.cpp
--- a/Engine/Source/Runtime/Engine/Private/PhysicsEngine/PhysAnim.cpp
+++ b/Engine/Source/Runtime/Engine/Private/PhysicsEngine/PhysAnim.cpp
@@ -708,7 +708,7 @@ void USkeletalMeshComponent::UpdateKinematicBonesToAnim(const TArray<FTransform>
 						}
 					}
 				}
-				BodyInstance.UpdateTriMeshVertices(LWC::ConvertArrayType<FVector>(NewPositions));
+				BodyInstance.UpdateTriMeshVertices(UE::LWC::ConvertArrayType<FVector>(NewPositions));
 			}
 			
 			BodyInstance.SetBodyTransform(CurrentLocalToWorld, Teleport);
diff --git a/Engine/Source/Runtime/Engine/Private/PhysicsField/PhysicsFieldComponent.cpp b/Engine/Source/Runtime/Engine/Private/PhysicsField/PhysicsFieldComponent.cpp
--- a/Engine/Source/Runtime/Engine/Private/PhysicsField/PhysicsFieldComponent.cpp
+++ b/Engine/Source/Runtime/Engine/Private/PhysicsField/PhysicsFieldComponent.cpp
@@ -558,8 +558,8 @@ void FPhysicsFieldResource::UpdateResource(FRHICommandListImmediate& RHICmdList,
 		}
 		// LWC_TODO: Perf pessimization
 		// LWC_TODO: Precision loss
-		UpdateInternalBuffer<FVector4f, 1, EPixelFormat::PF_A32B32G32R32F>(BoundsMinDatas.Num(), LWC::ConvertArrayType<FVector4f>(BoundsMinDatas).GetData(), BoundsMin, true);	
-		UpdateInternalBuffer<FVector4f, 1, EPixelFormat::PF_A32B32G32R32F>(BoundsMaxDatas.Num(),  LWC::ConvertArrayType<FVector4f>(BoundsMaxDatas).GetData(), BoundsMax, true);
+		UpdateInternalBuffer<FVector4f, 1, EPixelFormat::PF_A32B32G32R32F>(BoundsMinDatas.Num(), UE::LWC::ConvertArrayType<FVector4f>(BoundsMinDatas).GetData(), BoundsMin, true);	
+		UpdateInternalBuffer<FVector4f, 1, EPixelFormat::PF_A32B32G32R32F>(BoundsMaxDatas.Num(),  UE::LWC::ConvertArrayType<FVector4f>(BoundsMaxDatas).GetData(), BoundsMax, true);
 	}
 
 	if (FieldInfos.bBuildClipmap)
diff --git a/Engine/Source/Runtime/NavigationSystem/Private/NavMesh/PImplRecastNavMesh.cpp b/Engine/Source/Runtime/NavigationSystem/Private/NavMesh/PImplRecastNavMesh.cpp
--- a/Engine/Source/Runtime/NavigationSystem/Private/NavMesh/PImplRecastNavMesh.cpp
+++ b/Engine/Source/Runtime/NavigationSystem/Private/NavMesh/PImplRecastNavMesh.cpp
@@ -2054,7 +2054,7 @@ bool FPImplRecastNavMesh::FindPolysAroundCircle(const FVector& CenterPos, const
 
 		if (OutPolysCost)
 		{
-			*OutPolysCost = LWC::ConvertArrayTypeClampMax<float>(PolysCost);
+			*OutPolysCost = UE::LWC::ConvertArrayTypeClampMax<float>(PolysCost);
 		}
 
 		if (dtStatusSucceed(Status))
diff --git a/Engine/Source/Runtime/NavigationSystem/Private/NavMesh/RecastNavMeshGenerator.cpp b/Engine/Source/Runtime/NavigationSystem/Private/NavMesh/RecastNavMeshGenerator.cpp
--- a/Engine/Source/Runtime/NavigationSystem/Private/NavMesh/RecastNavMeshGenerator.cpp
+++ b/Engine/Source/Runtime/NavigationSystem/Private/NavMesh/RecastNavMeshGenerator.cpp
@@ -2567,7 +2567,7 @@ ETimeSliceWorkResult FRecastTileGenerator::RasterizeGeometryRecastTimeSliced(FNa
 {
 	TArray<FVector::FReal> RealCoords;
 
-	RealCoords = LWC::ConvertArrayType<FVector::FReal>(Coords);
+	RealCoords = UE::LWC::ConvertArrayType<FVector::FReal>(Coords);
 
 	return RasterizeGeometryRecastTimeSliced(BuildContext, RealCoords, Indices, RasterizationFlags, RasterContext);
 }
@@ -2606,7 +2606,7 @@ void FRecastTileGenerator::RasterizeGeometryRecast(FNavMeshBuildContext& BuildCo
 {
 	TArray<FVector::FReal> RealCoords;
 
-	RealCoords = LWC::ConvertArrayType<FVector::FReal>(Coords);
+	RealCoords = UE::LWC::ConvertArrayType<FVector::FReal>(Coords);
 
 	RasterizeGeometryRecast(BuildContext, RealCoords, Indices, RasterizationFlags, RasterContext);
 }
@@ -2635,7 +2635,7 @@ void FRecastTileGenerator::RasterizeGeometryTransformCoords(const TArray<float>&
 {
 	TArray<FVector::FReal> RealCoords;
 
-	RealCoords = LWC::ConvertArrayType<FVector::FReal>(Coords);
+	RealCoords = UE::LWC::ConvertArrayType<FVector::FReal>(Coords);
 
 	RasterizeGeometryTransformCoords(RealCoords, LocalToWorld);
 }
@@ -2687,7 +2687,7 @@ ETimeSliceWorkResult FRecastTileGenerator::RasterizeGeometryTimeSliced(FNavMeshB
 {
 	TArray<FVector::FReal> RealCoords;
 
-	RealCoords = LWC::ConvertArrayType<FVector::FReal>(Coords);
+	RealCoords = UE::LWC::ConvertArrayType<FVector::FReal>(Coords);
 
 	return RasterizeGeometryTimeSliced(BuildContext, RealCoords, Indices, LocalToWorld, RasterizationFlags, RasterContext);
 }
@@ -2704,7 +2704,7 @@ void FRecastTileGenerator::RasterizeGeometry(FNavMeshBuildContext& BuildContext,
 {
 	TArray<FVector::FReal> RealCoords;
 
-	RealCoords = LWC::ConvertArrayType<FVector::FReal>(Coords);
+	RealCoords = UE::LWC::ConvertArrayType<FVector::FReal>(Coords);
 
 	return RasterizeGeometry(BuildContext, RealCoords, Indices, LocalToWorld, RasterizationFlags, RasterContext);
 }
@@ -4102,7 +4102,7 @@ void FRecastTileGenerator::MarkRasterizationMask(rcContext* /*BuildContext*/, rc
 		TArray<FVector> ConvexVerts;
 		const FVector::FReal Expand = 0.f;
 
-		const TArray<FVector> Points = LWC::ConvertArrayType<FVector>(ConvexData.Points);
+		const TArray<FVector> Points = UE::LWC::ConvertArrayType<FVector>(ConvexData.Points);
 		GrowConvexHull(Expand, Points, ConvexVerts);
 
 		if (ConvexVerts.Num())
@@ -4225,7 +4225,7 @@ void FRecastTileGenerator::MarkDynamicArea(const FAreaNavModifier& Modifier, con
 
 			TArray<FVector> ConvexVerts;
 
-			const TArray<FVector> Points = LWC::ConvertArrayType<FVector>(ConvexData.Points);
+			const TArray<FVector> Points = UE::LWC::ConvertArrayType<FVector>(ConvexData.Points);
 			GrowConvexHull(ExpandBy, Points, ConvexVerts);
 			ConvexData.MinZ -= OffsetZMin;
 			ConvexData.MaxZ += OffsetZMax;
@@ -6693,7 +6693,7 @@ void FRecastNavMeshGenerator::GrabDebugSnapshot(struct FVisualLogEntry* Snapshot
 							{
 								Verts.Reset();
 
-								const TArray<FVector> Points = LWC::ConvertArrayType<FVector>(InConvexNavAreaData.Points);
+								const TArray<FVector> Points = UE::LWC::ConvertArrayType<FVector>(InConvexNavAreaData.Points);
 								GrowConvexHull(NavData->AgentRadius, Points, Verts);
 
 								if (Verts.Num())
@@ -6827,14 +6827,14 @@ void FRecastNavMeshGenerator::ExportNavigationData(const FString& FileName) cons
 							{
 								TArray<FVector> ConvexVerts;
 
-								const TArray<FVector> Points = LWC::ConvertArrayType<FVector>(ExportInfo.Convex.Points);
+								const TArray<FVector> Points = UE::LWC::ConvertArrayType<FVector>(ExportInfo.Convex.Points);
 								GrowConvexHull(NavData->AgentRadius, Points, ConvexVerts);
 								if (ConvexVerts.Num())
 								{
 									ExportInfo.Convex.MinZ -= NavData->CellHeight;
 									ExportInfo.Convex.MaxZ += NavData->CellHeight;
 
-									ExportInfo.Convex.Points = LWC::ConvertArrayType<FVector>(ConvexVerts);
+									ExportInfo.Convex.Points = UE::LWC::ConvertArrayType<FVector>(ConvexVerts);
 
 									AreaExport.Add(ExportInfo);
 								}								
diff --git a/Engine/Source/Runtime/StaticMeshDescription/Private/StaticMeshOperations.cpp b/Engine/Source/Runtime/StaticMeshDescription/Private/StaticMeshOperations.cpp
--- a/Engine/Source/Runtime/StaticMeshDescription/Private/StaticMeshOperations.cpp
+++ b/Engine/Source/Runtime/StaticMeshDescription/Private/StaticMeshOperations.cpp
@@ -2087,7 +2087,7 @@ bool FStaticMeshOperations::GenerateUniqueUVsForStaticMesh(const FMeshDescriptio
 	if (bPackSuccess)
 	{
 		Packer.CommitPackedUVs();
-		OutTexCoords = LWC::ConvertArrayType<FVector2D>(MeshDescriptionView.RetrievePackedUVs());	// LWC_TODO: Perf pessimization.
+		OutTexCoords = UE::LWC::ConvertArrayType<FVector2D>(MeshDescriptionView.RetrievePackedUVs());	// LWC_TODO: Perf pessimization.
 	}
 
 	return bPackSuccess;
