PR #8743: Prevent RecastNavMeshGenerator dirtying tiles if dirty areas merely touch them (Contributed by [REDACTED_PERSON])

diff --git a/Engine/Source/Runtime/NavigationSystem/Public/NavMesh/RecastNavMeshGenerator.h b/Engine/Source/Runtime/NavigationSystem/Public/NavMesh/RecastNavMeshGenerator.h
--- a/Engine/Source/Runtime/NavigationSystem/Public/NavMesh/RecastNavMeshGenerator.h
+++ b/Engine/Source/Runtime/NavigationSystem/Public/NavMesh/RecastNavMeshGenerator.h
@@ -183,11 +183,24 @@ struct FRcTileBox
 	{
 		check(TileSizeInWorldUnits > 0);
 
+		auto CalcMaxCoordExclusive = [](const FVector::FReal MaxAsFloat, const int32 MinCoord) -> int32
+		{
+			FVector::FReal UnusedIntPart;
+			// If MaxCoord falls exactly on the boundary of a tile
+			if (FMath::Modf(MaxAsFloat, &UnusedIntPart) == 0)
+			{
+				// Return the lower tile
+				return FMath::Max(FMath::FloorToInt(MaxAsFloat) - 1, MinCoord);
+			}
+			// Otherwise use default behaviour
+			return FMath::FloorToInt(MaxAsFloat);
+		};
+
 		const FBox RcAreaBounds = Unreal2RecastBox(UnrealBounds);
 		XMin = FMath::FloorToInt((RcAreaBounds.Min.X - RcNavMeshOrigin.X) / TileSizeInWorldUnits);
-		XMax = FMath::FloorToInt((RcAreaBounds.Max.X - RcNavMeshOrigin.X) / TileSizeInWorldUnits);
+		XMax = CalcMaxCoordExclusive((RcAreaBounds.Max.X - RcNavMeshOrigin.X) / TileSizeInWorldUnits, XMin);
 		YMin = FMath::FloorToInt((RcAreaBounds.Min.Z - RcNavMeshOrigin.Z) / TileSizeInWorldUnits);
-		YMax = FMath::FloorToInt((RcAreaBounds.Max.Z - RcNavMeshOrigin.Z) / TileSizeInWorldUnits);
+		YMax = CalcMaxCoordExclusive((RcAreaBounds.Max.Z - RcNavMeshOrigin.Z) / TileSizeInWorldUnits, YMin);
 	}
 
 	FORCEINLINE bool Contains(const FIntPoint& Point) const
