[Navigation] Addition of a warning in UNavigationSystemV1::GetNavDataForProps to warn about using invalid FNavAgentProperties

diff --git a/Engine/Source/Runtime/NavigationSystem/Private/NavigationSystem.cpp b/Engine/Source/Runtime/NavigationSystem/Private/NavigationSystem.cpp
--- a/Engine/Source/Runtime/NavigationSystem/Private/NavigationSystem.cpp
+++ b/Engine/Source/Runtime/NavigationSystem/Private/NavigationSystem.cpp
@@ -1828,6 +1828,8 @@ const ANavigationData* UNavigationSystemV1::GetNavDataForProps(const FNavAgentPr
 		return MainNavData;
 	}
 	
+	UE_CLOG(!AgentProperties.IsValid(), LogNavigation, Warning, TEXT("Looking for NavData using invalid FNavAgentProperties."));
+	
 	const TWeakObjectPtr<ANavigationData>* NavDataForAgent = AgentToNavDataMap.Find(AgentProperties);
 	const ANavigationData* NavDataInstance = NavDataForAgent ? NavDataForAgent->Get() : nullptr;
 
