PR #8622: Make UNavigationSystemV1::PerformNavigationBoundsUpdate() virtual (Contributed by [REDACTED_PERSON])

diff --git a/Engine/Source/Runtime/NavigationSystem/Public/NavigationSystem.h b/Engine/Source/Runtime/NavigationSystem/Public/NavigationSystem.h
--- a/Engine/Source/Runtime/NavigationSystem/Public/NavigationSystem.h
+++ b/Engine/Source/Runtime/NavigationSystem/Public/NavigationSystem.h
@@ -756,7 +756,7 @@ protected:
 	void ProcessCustomLinkPendingRegistration();
 
 	/** used to apply updates of nav volumes in navigation system's tick */
-	void PerformNavigationBoundsUpdate(const TArray<FNavigationBoundsUpdateRequest>& UpdateRequests);
+	virtual void PerformNavigationBoundsUpdate(const TArray<FNavigationBoundsUpdateRequest>& UpdateRequests);
 	
 	/** adds data to RegisteredNavBounds */
 	void AddNavigationBounds(const FNavigationBounds& NewBounds);
