[Navigation] Only update actor and components in navoctree if the actor is initalized.

diff --git a/Engine/Source/Runtime/NavigationSystem/Private/NavigationSystem.cpp b/Engine/Source/Runtime/NavigationSystem/Private/NavigationSystem.cpp
--- a/Engine/Source/Runtime/NavigationSystem/Private/NavigationSystem.cpp
+++ b/Engine/Source/Runtime/NavigationSystem/Private/NavigationSystem.cpp
@@ -452,8 +452,11 @@ UNavigationSystemV1::UNavigationSystemV1(const FObjectInitializer& ObjectInitial
 					UNavigationSystemV1::UpdateNavOctreeBounds(ParentedActors[Idx]);
 				}
 
-				// not doing manual update of all attached actors since UpdateActorAndComponentsInNavOctree should take care of it
-				UNavigationSystemV1::UpdateActorAndComponentsInNavOctree(Actor);
+				if (Actor.IsActorInitialized())
+				{
+					// not doing manual update of all attached actors since UpdateActorAndComponentsInNavOctree should take care of it
+					UNavigationSystemV1::UpdateActorAndComponentsInNavOctree(Actor);
+				}
 			});
 			UNavigationSystemBase::OnComponentTransformChangedDelegate().BindLambda([](USceneComponent& Comp) {
 				if (UNavigationSystemV1::ShouldUpdateNavOctreeOnComponentChange())
@@ -4080,7 +4083,7 @@ void UNavigationSystemV1::OnActorMoved(AActor* Actor)
 	{
 		OnNavigationBoundsUpdated((ANavMeshBoundsVolume*)Actor);
 	}
-	else if (Actor)
+	else if (Actor && Actor->IsActorInitialized())
 	{
 		UpdateActorAndComponentsInNavOctree(*Actor, /*bUpdateAttachedActors=*/true);
 	}
