Remove FSerializeFloatAsDoubleHack. This was straight forward as we were not seriliazling any floats currently just FReals as doubles.

diff --git a/Engine/Source/Runtime/NavigationSystem/Private/NavMesh/PImplRecastNavMesh.cpp b/Engine/Source/Runtime/NavigationSystem/Private/NavMesh/PImplRecastNavMesh.cpp
--- a/Engine/Source/Runtime/NavigationSystem/Private/NavMesh/PImplRecastNavMesh.cpp
+++ b/Engine/Source/Runtime/NavigationSystem/Private/NavMesh/PImplRecastNavMesh.cpp
@@ -476,115 +476,13 @@ void FPImplRecastNavMesh::ReleaseDetourNavMesh()
 #endif
 }
 
-// LWC_TODO_AI: Remove prior to UE5 5.0 Release.
-// Currenlty floats are serialized as doubles for the navigation data so it can be loaded in LWC and non LWC builds (mainly used for regression testing).
-class FSerializeFloatAsDoubleHack
-{
-public:
-	FSerializeFloatAsDoubleHack(FArchive& InArchive)
-		: Archive(InArchive)
-	{}
-
-	/** Returns true if this archive is for loading data. */
-	FORCEINLINE bool IsLoading() const
-	{
-		return Archive.IsLoading();
-	}
-
-	/** Returns true if this archive is for saving data, this can also be a pre-save preparation archive. */
-	FORCEINLINE bool IsSaving() const
-	{
-		return Archive.IsSaving();
-	}
-
-	operator FArchive&() const { return Archive; }
-
-	friend FSerializeFloatAsDoubleHack& operator<<(FSerializeFloatAsDoubleHack& Ar, uint8& Value)
-	{
-		Ar.Archive << Value;
-		return Ar;
-	}
-
-	friend FSerializeFloatAsDoubleHack& operator<<(FSerializeFloatAsDoubleHack& Ar, int8& Value)
-	{
-		Ar.Archive << Value;
-		return Ar;
-	}
-
-	friend FSerializeFloatAsDoubleHack& operator<<(FSerializeFloatAsDoubleHack& Ar, uint16& Value)
-	{
-		 Ar.Archive << Value;
-		 return Ar;
-	}
-
-	friend FSerializeFloatAsDoubleHack& operator<<(FSerializeFloatAsDoubleHack& Ar, int16& Value)
-	{
-		Ar.Archive << Value;
-		return Ar;
-	}
-
-	friend FSerializeFloatAsDoubleHack& operator<<(FSerializeFloatAsDoubleHack& Ar, uint32& Value)
-	{
-		Ar.Archive << Value;
-		return Ar;
-	}
-
-	friend FSerializeFloatAsDoubleHack& operator<<(FSerializeFloatAsDoubleHack& Ar, bool& Value)
-	{
-		Ar.Archive << Value;
-		return Ar;
-	}
-
-	friend FSerializeFloatAsDoubleHack& operator<<(FSerializeFloatAsDoubleHack& Ar, int32& Value)
-	{
-		Ar.Archive << Value;
-		return Ar;
-	}
-
-	friend FSerializeFloatAsDoubleHack& operator<<(FSerializeFloatAsDoubleHack& Ar, float& Value)
-	{
-		double DoubleValue = Value;
-		Ar.Archive << DoubleValue;
-		Value = DoubleValue;
-		return Ar;
-	}
-
-	friend FSerializeFloatAsDoubleHack& operator<<(FSerializeFloatAsDoubleHack& Ar, double& Value)
-	{
-		Ar.Archive << Value;
-		return Ar;
-	}
-
-	FORCEINLINE friend FSerializeFloatAsDoubleHack& operator<<(FSerializeFloatAsDoubleHack& Ar, uint64& Value)
-	{
-		Ar.Archive << Value;
-		return Ar;
-	}
-
-	friend FSerializeFloatAsDoubleHack& operator<<(FSerializeFloatAsDoubleHack& Ar, int64& Value)
-	{
-		Ar.Archive << Value;
-		return Ar;
-	}
-
-	virtual void Serialize(void* Value, int64 Length)
-	{
-		Archive.Serialize(Value, Length);
-	}
-
-protected:
-	FArchive& Archive;
-};
-
 /**
  * Serialization.
  * @param Ar - The archive with which to serialize.
  * @returns true if serialization was successful.
  */
-void FPImplRecastNavMesh::Serialize( FArchive& ArWrapped, int32 NavMeshVersion )
+void FPImplRecastNavMesh::Serialize( FArchive& Ar, int32 NavMeshVersion )
 {
-	FSerializeFloatAsDoubleHack Ar(ArWrapped); // LWC_TODO_AI: Remove prior to UE5 5.0 Release.
-
 	//@todo: How to handle loading nav meshes saved w/ recast when recast isn't present????
 
 	if (!Ar.IsLoading() && DetourNavMesh == NULL)
@@ -795,13 +693,11 @@ void FPImplRecastNavMesh::Serialize( FArchive& ArWrapped, int32 NavMeshVersion )
 	}
 }
 
-void FPImplRecastNavMesh::SerializeRecastMeshTile(FArchive& ArWrapped, int32 NavMeshVersion, unsigned char*& TileData, int32& TileDataSize)
+void FPImplRecastNavMesh::SerializeRecastMeshTile(FArchive& Ar, int32 NavMeshVersion, unsigned char*& TileData, int32& TileDataSize)
 {
 	// The strategy here is to serialize the data blob that is passed into addTile()
 	// @see dtCreateNavMeshData() for details on how this data is laid out
 
-	FSerializeFloatAsDoubleHack Ar(ArWrapped); // LWC_TODO_AI: Remove prior to UE5 5.0 Release.
-
 	FDetourTileSizeInfo SizeInfo;
 
 	if (Ar.IsSaving())
@@ -1042,10 +938,8 @@ void FPImplRecastNavMesh::SerializeRecastMeshTile(FArchive& ArWrapped, int32 Nav
 	}
 }
 
-void FPImplRecastNavMesh::SerializeCompressedTileCacheData(FArchive& ArWrapped, int32 NavMeshVersion, unsigned char*& CompressedData, int32& CompressedDataSize)
+void FPImplRecastNavMesh::SerializeCompressedTileCacheData(FArchive& Ar, int32 NavMeshVersion, unsigned char*& CompressedData, int32& CompressedDataSize)
 {
-	// LWC_TODO_AI: Remove prior to UE5 5.0 Release.
-	FSerializeFloatAsDoubleHack Ar(ArWrapped);
 	constexpr int32 EmptyDataValue = -1;
 
 	// Note when saving the CompressedDataSize is either 0 or it must be big enough to include the size of the uncompressed dtTileCacheLayerHeader.
