[Navigation] Addition of a warning when exporting component with too many verticies.
With a custimizable warning threshold through the GeometryExportVertexCountWarningThreshold property.

diff --git a/Engine/Source/Runtime/NavigationSystem/Private/NavMesh/RecastNavMeshGenerator.cpp b/Engine/Source/Runtime/NavigationSystem/Private/NavMesh/RecastNavMeshGenerator.cpp
--- a/Engine/Source/Runtime/NavigationSystem/Private/NavMesh/RecastNavMeshGenerator.cpp
+++ b/Engine/Source/Runtime/NavigationSystem/Private/NavMesh/RecastNavMeshGenerator.cpp
@@ -1012,6 +1012,23 @@ FORCEINLINE_DEBUGGABLE void ExportComponent(UActorComponent* Component, FRecastG
 	}
 }
 
+#if !UE_BUILD_SHIPPING
+FORCEINLINE_DEBUGGABLE void ValidateGeometryExport(const FRecastGeometryExport& GeomExport)
+{
+	if (const UObject* Owner = GeomExport.Data->GetOwner())
+	{
+		if (const UNavigationSystemV1* NavSys = FNavigationSystem::GetCurrent<UNavigationSystemV1>(Owner->GetWorld()))
+		{
+			UE_CLOG(NavSys->GeometryExportVertexCountWarningThreshold > 0 && GeomExport.VertexBuffer.Num() > NavSys->GeometryExportVertexCountWarningThreshold,
+				LogNavigation,
+				Warning, 
+				TEXT("Exporting geometry with too many verticies (%i). This might cause performance and memory issues. Simplify collision or change GeometryExportVertexCountWarningThreshold. See '%s'."), 
+				GeomExport.VertexBuffer.Num(), *GetFullNameSafe(Owner));
+		}
+	}
+}
+#endif //!UE_BUILD_SHIPPING
+	
 FORCEINLINE void TransformVertexSoupToRecast(const TArray<FVector>& VertexSoup, TNavStatArray<FVector>& Verts, TNavStatArray<int32>& Faces)
 {
 	if (VertexSoup.Num() == 0)
@@ -6432,6 +6449,11 @@ void FRecastNavMeshGenerator::ExportComponentGeometry(UActorComponent* Component
 {
 	FRecastGeometryExport GeomExport(Data);
 	RecastGeometryExport::ExportComponent(Component, GeomExport);
+
+#if !UE_BUILD_SHIPPING	
+	RecastGeometryExport::ValidateGeometryExport(GeomExport);
+#endif
+	
 	RecastGeometryExport::CovertCoordDataToRecast(GeomExport.VertexBuffer);
 	RecastGeometryExport::StoreCollisionCache(GeomExport);
 }
@@ -6440,6 +6462,11 @@ void FRecastNavMeshGenerator::ExportVertexSoupGeometry(const TArray<FVector>& Ve
 {
 	FRecastGeometryExport GeomExport(Data);
 	RecastGeometryExport::ExportVertexSoup(Verts, GeomExport.VertexBuffer, GeomExport.IndexBuffer, GeomExport.Data->Bounds);
+
+#if !UE_BUILD_SHIPPING	
+	RecastGeometryExport::ValidateGeometryExport(GeomExport);
+#endif
+	
 	RecastGeometryExport::StoreCollisionCache(GeomExport);
 }
 
diff --git a/Engine/Source/Runtime/NavigationSystem/Public/NavigationSystem.h b/Engine/Source/Runtime/NavigationSystem/Public/NavigationSystem.h
--- a/Engine/Source/Runtime/NavigationSystem/Public/NavigationSystem.h
+++ b/Engine/Source/Runtime/NavigationSystem/Public/NavigationSystem.h
@@ -314,8 +314,12 @@ public:
 	UPROPERTY(config, EditAnywhere, Category=NavigationSystem)
 	uint32 bSkipAgentHeightCheckWhenPickingNavData:1;
 
+	/** Warnings are logged if exporting the navigation collision for an object exceed this vertex count.
+	 * Use -1 to disable. */
+	UPROPERTY(config, EditAnywhere, AdvancedDisplay, Category = NavigationSystem)
+	int32 GeometryExportVertexCountWarningThreshold = 1000000;
+	
 protected:
-
 	/** If set to true navigation will be generated only around registered "navigation enforcers"
 	*	This has a range of consequences (including how navigation octree operates) so it needs to
 	*	be a conscious decision.
