Destroy NavigationData in CleanUpAndMarkPendingKill to fix KillPendingKill errors
#tests Forced a call to CleanUpAndMarkPendingKill in each place that it can be called from to make sure that a crash doesn't occur

diff --git a/Engine/Source/Runtime/NavigationSystem/Private/NavigationData.cpp b/Engine/Source/Runtime/NavigationSystem/Private/NavigationData.cpp
--- a/Engine/Source/Runtime/NavigationSystem/Private/NavigationData.cpp
+++ b/Engine/Source/Runtime/NavigationSystem/Private/NavigationData.cpp
@@ -16,6 +16,12 @@
 // for NavigationData
 #define NAVDATAVER_LATEST 13	
 
+static TAutoConsoleVariable<int32> CVarDestroyNavDataInCleanUpAndMarkPendingKill(
+	TEXT("ai.DestroyNavDataInCleanUpAndMarkPendingKill"),
+	1,
+	TEXT("If set to 1 NavData will be destroyed in CleanUpAndMarkPendingKill rather than being marked as garbage.\n"),
+	ECVF_Default);
+
 //----------------------------------------------------------------------//
 // FPathFindingQuery
 //----------------------------------------------------------------------//
@@ -502,12 +508,24 @@ void ANavigationData::ApplyWorldOffset(const FVector& InOffset, bool bWorldShift
 void ANavigationData::CleanUpAndMarkPendingKill()
 {
 	CleanUp();
-	SetActorHiddenInGame(true);
 
-	// do NOT destroy here! it can be called from PostLoad and will crash in DestroyActor()
-	GetWorld()->RemoveNetworkActor(this);
-	MarkAsGarbage();
-	MarkComponentsAsPendingKill();
+	/* Need to check if the world is valid since, when this is called from Serialize, the World won't be set and Destroy will do nothing, 
+	 * in which case it will crash when it tries to register with the NavSystem in UNavigationSystemV1::ProcessRegistrationCandidates. */
+	if (CVarDestroyNavDataInCleanUpAndMarkPendingKill.GetValueOnGameThread() && IsValid(GetWorld()))
+	{
+		Destroy();
+	}
+	else
+	{
+		SetActorHiddenInGame(true);
+
+		if (UWorld* World = GetWorld())
+		{
+			World->RemoveNetworkActor(this);
+		}
+		MarkAsGarbage();
+		MarkComponentsAsPendingKill();
+	}
 }
 
 bool ANavigationData::SupportsRuntimeGeneration() const
