Upgraded a lot of the navigation code. to use FVector::FReal rather than floats. Areas of code where there was little to gain from doing this but would break existing code have been left as float.

diff --git a/Engine/Source/Runtime/AIModule/Classes/Navigation/CrowdManager.h b/Engine/Source/Runtime/AIModule/Classes/Navigation/CrowdManager.h
--- a/Engine/Source/Runtime/AIModule/Classes/Navigation/CrowdManager.h
+++ b/Engine/Source/Runtime/AIModule/Classes/Navigation/CrowdManager.h
@@ -109,7 +109,6 @@ struct AIMODULE_API FCrowdAvoidanceConfig
 	{}
 };
 
-// LWC_TODO_AI: Info.Angles and Info.Radii should be FReals. Most likely leave this until after 5.0
 USTRUCT()
 struct AIMODULE_API FCrowdAvoidanceSamplingPattern
 {
diff --git a/Engine/Source/Runtime/AIModule/Classes/Navigation/MetaNavMeshPath.h b/Engine/Source/Runtime/AIModule/Classes/Navigation/MetaNavMeshPath.h
--- a/Engine/Source/Runtime/AIModule/Classes/Navigation/MetaNavMeshPath.h
+++ b/Engine/Source/Runtime/AIModule/Classes/Navigation/MetaNavMeshPath.h
@@ -97,10 +97,10 @@ struct AIMODULE_API FMetaNavMeshPath : public FNavMeshPath
 	void SetWaypointSwitchRadius(float InSwitchRadius) { WaypointSwitchRadius = InSwitchRadius; }
 
 	/** returns approximate length of path, ignores parameters */
-	virtual float GetLengthFromPosition(FVector SegmentStart, uint32 NextPathPointIndex) const override;
+	virtual FVector::FReal GetLengthFromPosition(FVector SegmentStart, uint32 NextPathPointIndex) const override;
 	
-	/** returns approximate length of path, ignores parameter */
-	virtual float GetCostFromIndex(int32 PathPointIndex) const override;
+	/** returns approximate cost of path, ignores parameter */
+	virtual FVector::FReal GetCostFromIndex(int32 PathPointIndex) const override;
 
 #if ENABLE_VISUAL_LOG
 	virtual void DescribeSelfToVisLog(FVisualLogEntry* Snapshot) const override;
@@ -113,7 +113,7 @@ protected:
 	TArray<FMetaPathWayPoint> Waypoints;
 
 	/** sum of 3D distance along waypoints, used for approximating length of path */
-	float ApproximateLength;
+	FVector::FReal ApproximateLength;
 
 	/** update navmesh path when this close to target waypoint */
 	float WaypointSwitchRadius;
diff --git a/Engine/Source/Runtime/AIModule/Private/EnvironmentQuery/Generators/EnvQueryGenerator_PathingGrid.cpp b/Engine/Source/Runtime/AIModule/Private/EnvironmentQuery/Generators/EnvQueryGenerator_PathingGrid.cpp
--- a/Engine/Source/Runtime/AIModule/Private/EnvironmentQuery/Generators/EnvQueryGenerator_PathingGrid.cpp
+++ b/Engine/Source/Runtime/AIModule/Private/EnvironmentQuery/Generators/EnvQueryGenerator_PathingGrid.cpp
@@ -65,14 +65,14 @@ void UEnvQueryGenerator_PathingGrid::ProjectAndFilterNavPoints(TArray<FNavLocati
 
 		for (int32 ContextIdx = 0; ContextIdx < ContextLocations.Num() && Points.Num(); ContextIdx++)
 		{
-			float CollectDistanceSq = 0.0f;
+			FVector::FReal CollectDistanceSq = 0.0f;
 			for (int32 Idx = 0; Idx < Points.Num(); Idx++)
 			{
-				const float TestDistanceSq = FVector::DistSquared(Points[Idx].Location, ContextLocations[ContextIdx]);
+				const FVector::FReal TestDistanceSq = FVector::DistSquared(Points[Idx].Location, ContextLocations[ContextIdx]);
 				CollectDistanceSq = FMath::Max(CollectDistanceSq, TestDistanceSq);
 			}
 
-			const float MaxPathDistance = FMath::Sqrt(CollectDistanceSq) * RangeMultiplierValue;
+			const FVector::FReal MaxPathDistance = FMath::Sqrt(CollectDistanceSq) * RangeMultiplierValue;
 
 			Polys.Reset();
 
diff --git a/Engine/Source/Runtime/AIModule/Private/Navigation/CrowdManager.cpp b/Engine/Source/Runtime/AIModule/Private/Navigation/CrowdManager.cpp
--- a/Engine/Source/Runtime/AIModule/Private/Navigation/CrowdManager.cpp
+++ b/Engine/Source/Runtime/AIModule/Private/Navigation/CrowdManager.cpp
@@ -920,7 +920,6 @@ void UCrowdManager::CreateCrowdManager()
 			const FCrowdAvoidanceSamplingPattern& Info = SamplingPatterns[Idx];
 			if (Info.Angles.Num() > 0 && Info.Angles.Num() == Info.Radii.Num())
 			{
-				// LWC_TODO_AI: Info.Angles and Info.Radii should be FReal. Probably not until after 5.0!
 				const TArray<FVector::FReal> Angles = UE::LWC::ConvertArrayType<FVector::FReal>(Info.Angles);
 				const TArray<FVector::FReal> Radii = UE::LWC::ConvertArrayType<FVector::FReal>(Info.Radii);
 
diff --git a/Engine/Source/Runtime/AIModule/Private/Navigation/MetaNavMeshPath.cpp b/Engine/Source/Runtime/AIModule/Private/Navigation/MetaNavMeshPath.cpp
--- a/Engine/Source/Runtime/AIModule/Private/Navigation/MetaNavMeshPath.cpp
+++ b/Engine/Source/Runtime/AIModule/Private/Navigation/MetaNavMeshPath.cpp
@@ -19,7 +19,7 @@ FMetaNavMeshPath::FMetaNavMeshPath()
 	// suggested value: 5-10x agent radius
 	WaypointSwitchRadius = 200.0f;
 
-	ApproximateLength = 0.0f;
+	ApproximateLength = 0.;
 	PathGoalTetherDistance = 0.0f;
 
 	// initialize to 0, path following will try to update it immediately after receiving request
@@ -140,7 +140,7 @@ bool FMetaNavMeshPath::ConditionalMoveToNextSection(const FVector& AgentLocation
 {
 	if (Waypoints.IsValidIndex(TargetWaypointIdx))
 	{
-		const float DistSq = FVector::DistSquared(AgentLocation, Waypoints[TargetWaypointIdx]);
+		const FVector::FReal DistSq = FVector::DistSquared(AgentLocation, Waypoints[TargetWaypointIdx]);
 		if (((Reason == EMetaPathUpdateReason::PathFinished) || (DistSq < FMath::Square(WaypointSwitchRadius)))
 			&& (GetSourceActorAsNavAgent() == nullptr || GetSourceActorAsNavAgent()->ShouldPostponePathUpdates() == false))
 		{
@@ -213,13 +213,13 @@ void FMetaNavMeshPath::CopyFrom(const FMetaNavMeshPath& Other)
 	ApproximateLength = Other.ApproximateLength;
 }
 
-float FMetaNavMeshPath::GetLengthFromPosition(FVector SegmentStart, uint32 NextPathPointIndex) const
+FVector::FReal FMetaNavMeshPath::GetLengthFromPosition(FVector SegmentStart, uint32 NextPathPointIndex) const
 {
 	// return approximation of full path, there's not enough data to give accurate value
 	return ApproximateLength;
 }
 
-float FMetaNavMeshPath::GetCostFromIndex(int32 PathPointIndex) const
+FVector::FReal FMetaNavMeshPath::GetCostFromIndex(int32 PathPointIndex) const
 {
 	// return approximation of full path * default cost, there's not enough data to give accurate value
 	const UNavArea* DefaultAreaOb = static_cast<const UNavArea*>(FNavigationSystem::GetDefaultWalkableArea().GetDefaultObject());
diff --git a/Engine/Source/Runtime/Core/Public/Misc/LargeWorldCoordinates.h b/Engine/Source/Runtime/Core/Public/Misc/LargeWorldCoordinates.h
--- a/Engine/Source/Runtime/Core/Public/Misc/LargeWorldCoordinates.h
+++ b/Engine/Source/Runtime/Core/Public/Misc/LargeWorldCoordinates.h
@@ -30,8 +30,8 @@ namespace ispc { struct DEFAULT_TYPENAME; }							/* ISPC forward declaration */
 #define UE_REAL_TO_FLOAT(argument) static_cast<float>(argument)
 
 // Use to make any narrowing casts searchable in code when it is updated to work with a 64 bit count/range clamped to max float.
-#define UE_REAL_TO_FLOAT_CLAMPED_MAX(argument) static_cast<float>(FMath::Min(argument, TNumericLimits<float>::Max()))
+#define UE_REAL_TO_FLOAT_CLAMPED_MAX(argument) static_cast<float>(FMath::Min(argument, (FVector::FReal)TNumericLimits<float>::Max()))
 
 // Use to make any narrowing casts searchable in code when it is updated to work with a 64 bit count/range clamped to min / max float.
-#define UE_REAL_TO_FLOAT_CLAMPED(argument) static_cast<float>(FMath::Clamp(argument, TNumericLimits<float>::Lowest(), TNumericLimits<float>::Max()))
+#define UE_REAL_TO_FLOAT_CLAMPED(argument) static_cast<float>(FMath::Clamp(argument, (FVector::FReal)TNumericLimits<float>::Lowest(), (FVector::FReal)TNumericLimits<float>::Max()))
 
diff --git a/Engine/Source/Runtime/Engine/Classes/AI/Navigation/NavigationTypes.h b/Engine/Source/Runtime/Engine/Classes/AI/Navigation/NavigationTypes.h
--- a/Engine/Source/Runtime/Engine/Classes/AI/Navigation/NavigationTypes.h
+++ b/Engine/Source/Runtime/Engine/Classes/AI/Navigation/NavigationTypes.h
@@ -36,8 +36,6 @@ struct FNavigationPath;
 /** uniform identifier type for navigation data elements may it be a polygon or graph node */
 typedef uint64 NavNodeRef;
 
-// LWC_TODO_AI: Most the floats in this file should really be FReal. Probably not until after 5.0
-
 namespace FNavigationSystem
 {
 	/** used as a fallback value for navigation agent radius, when none specified via UNavigationSystemV1::SupportedAgents */
@@ -48,7 +46,7 @@ namespace FNavigationSystem
 
 	static const FBox InvalidBoundingBox(ForceInit);
 
-	static const FVector InvalidLocation = FVector(FLT_MAX);
+	static const FVector InvalidLocation = FVector(TNumericLimits<FVector::FReal>::Max());
 
 	FORCEINLINE bool IsValidLocation(const FVector& TestLocation)
 	{
diff --git a/Engine/Source/Runtime/Engine/Public/AI/Navigation/NavQueryFilter.h b/Engine/Source/Runtime/Engine/Public/AI/Navigation/NavQueryFilter.h
--- a/Engine/Source/Runtime/Engine/Public/AI/Navigation/NavQueryFilter.h
+++ b/Engine/Source/Runtime/Engine/Public/AI/Navigation/NavQueryFilter.h
@@ -4,6 +4,8 @@
 
 #include "CoreMinimal.h"
 
+// For now leave all the real numbers here as floats instead of conversion to FVector::FReal. Its not worth changing the virtual functions and breaking
+// any existing code. Its not required for heuristic scale or area costs to be larger or have greater precison than a float.
 class INavigationQueryFilterInterface
 {
 public:
diff --git a/Engine/Source/Runtime/NavigationSystem/Private/NavMesh/NavMeshPath.cpp b/Engine/Source/Runtime/NavigationSystem/Private/NavMesh/NavMeshPath.cpp
--- a/Engine/Source/Runtime/NavigationSystem/Private/NavMesh/NavMeshPath.cpp
+++ b/Engine/Source/Runtime/NavigationSystem/Private/NavMesh/NavMeshPath.cpp
@@ -53,7 +53,7 @@ void FNavMeshPath::InternalResetNavMeshPath()
 	// - bWantsPathCorridor
 }
 
-float FNavMeshPath::GetStringPulledLength(const int32 StartingPoint) const
+FVector::FReal FNavMeshPath::GetStringPulledLength(const int32 StartingPoint) const
 {
 	if (IsValid() == false || StartingPoint >= PathPoints.Num())
 	{
@@ -69,10 +69,10 @@ float FNavMeshPath::GetStringPulledLength(const int32 StartingPoint) const
 		TotalLength += FVector::Dist(PrevPoint->Location, PathPoint->Location);
 	}
 
-	return UE_REAL_TO_FLOAT_CLAMPED_MAX(TotalLength);
+	return TotalLength;
 }
 
-float FNavMeshPath::GetPathCorridorLength(const int32 StartingEdge) const
+FVector::FReal FNavMeshPath::GetPathCorridorLength(const int32 StartingEdge) const
 {
 	if (bCorridorEdgesGenerated == false)
 	{
@@ -80,7 +80,7 @@ float FNavMeshPath::GetPathCorridorLength(const int32 StartingEdge) const
 	}
 	else if (StartingEdge >= PathCorridorEdges.Num())
 	{
-		return StartingEdge == 0 && PathPoints.Num() > 1 ? UE_REAL_TO_FLOAT_CLAMPED_MAX(FVector::Dist(PathPoints[0].Location, PathPoints[PathPoints.Num()-1].Location)) : 0.f;
+		return StartingEdge == 0 && PathPoints.Num() > 1 ? FVector::Dist(PathPoints[0].Location, PathPoints[PathPoints.Num()-1].Location) : 0.;
 	}
 	
 	const FNavigationPortalEdge* PrevEdge = PathCorridorEdges.GetData() + StartingEdge;
@@ -97,7 +97,7 @@ float FNavMeshPath::GetPathCorridorLength(const int32 StartingEdge) const
 		PrevEdgeMiddle = CurrentEdgeMiddle;
 	}
 	// @todo add distance to last point here!
-	return UE_REAL_TO_FLOAT_CLAMPED_MAX(TotalLength);
+	return TotalLength;
 }
 
 const TArray<FNavigationPortalEdge>& FNavMeshPath::GeneratePathCorridorEdges() const
@@ -279,7 +279,7 @@ void AppendPathPointsHelper(TArray<FNavPathPoint>& PathPoints, const TArray<FPat
 	}
 }
 
-void FNavMeshPath::OffsetFromCorners(float Distance)
+void FNavMeshPath::OffsetFromCorners(FVector::FReal Distance)
 {
 	SCOPE_CYCLE_COUNTER(STAT_Navigation_OffsetFromCorners);
 
diff --git a/Engine/Source/Runtime/NavigationSystem/Private/NavMesh/NavTestRenderingComponent.cpp b/Engine/Source/Runtime/NavigationSystem/Private/NavMesh/NavTestRenderingComponent.cpp
--- a/Engine/Source/Runtime/NavigationSystem/Private/NavMesh/NavTestRenderingComponent.cpp
+++ b/Engine/Source/Runtime/NavigationSystem/Private/NavMesh/NavTestRenderingComponent.cpp
@@ -228,7 +228,7 @@ void FNavTestSceneProxy::GatherPathStep()
 
 			FNodeDebugData NewNodeData;
 
-			float DisplayedCost = FLT_MAX; 
+			FVector::FReal DisplayedCost = TNumericLimits<FVector::FReal>::Max();
 			switch (NavTestActor->CostDisplayMode)
 			{
 			case ENavCostDisplay::TotalCost:
diff --git a/Engine/Source/Runtime/NavigationSystem/Private/NavMesh/PImplRecastNavMesh.cpp b/Engine/Source/Runtime/NavigationSystem/Private/NavMesh/PImplRecastNavMesh.cpp
--- a/Engine/Source/Runtime/NavigationSystem/Private/NavMesh/PImplRecastNavMesh.cpp
+++ b/Engine/Source/Runtime/NavigationSystem/Private/NavMesh/PImplRecastNavMesh.cpp
@@ -367,7 +367,6 @@ void FRecastQueryFilter::SetAllAreaCosts(const float* CostArray, const int32 Cou
 	}
 }
 
-// LWC_TODO_AI: These costs should probably be FVector::FReal. Not until after 5.0!
 void FRecastQueryFilter::GetAllAreaCosts(float* CostArray, float* FixedCostArray, const int32 Count) const
 {
 	const FVector::FReal* DetourCosts = getAllAreaCosts();
@@ -376,7 +375,6 @@ void FRecastQueryFilter::GetAllAreaCosts(float* CostArray, float* FixedCostArray
 
 	for (int i = 0; i < NumItems; ++i)
 	{
-		// LWC_TODO_AI: Replace with memcpy when we move to FReal. Not until after 5.0!
 		CostArray[i] = UE_REAL_TO_FLOAT_CLAMPED(DetourCosts[i]);
 		FixedCostArray[i] = UE_REAL_TO_FLOAT_CLAMPED(DetourFixedCosts[i]);
 	}
@@ -1091,11 +1089,11 @@ void FPImplRecastNavMesh::Raycast(const FVector& StartLoc, const FVector& EndLoc
 
 ENavigationQueryResult::Type FPImplRecastNavMesh::FindPath(const FVector& StartLoc, const FVector& EndLoc, FNavMeshPath& Path, const FNavigationQueryFilter& Filter, const UObject* Owner) const
 {
-	return FindPath(StartLoc, EndLoc, FLT_MAX, Path, Filter, Owner);
+	return FindPath(StartLoc, EndLoc, TNumericLimits<FVector::FReal>::Max(), Path, Filter, Owner);
 }
 
 // @TODONAV
-ENavigationQueryResult::Type FPImplRecastNavMesh::FindPath(const FVector& StartLoc, const FVector& EndLoc, const float CostLimit, FNavMeshPath& Path, const FNavigationQueryFilter& InQueryFilter, const UObject* Owner) const
+ENavigationQueryResult::Type FPImplRecastNavMesh::FindPath(const FVector& StartLoc, const FVector& EndLoc, const FVector::FReal CostLimit, FNavMeshPath& Path, const FNavigationQueryFilter& InQueryFilter, const UObject* Owner) const
 {
 	// temporarily disabling this check due to it causing too much "crashes"
 	// @todo but it needs to be back at some point since it realy checks for a buggy setup
@@ -1289,13 +1287,13 @@ bool FPImplRecastNavMesh::InitPathfinding(const FVector& UnrealStart, const FVec
 	return true;
 }
 
-float FPImplRecastNavMesh::CalcSegmentCostOnPoly(NavNodeRef PolyID, const dtQueryFilter* Filter, const FVector& StartLoc, const FVector& EndLoc) const
+FVector::FReal FPImplRecastNavMesh::CalcSegmentCostOnPoly(NavNodeRef PolyID, const dtQueryFilter* Filter, const FVector& StartLoc, const FVector& EndLoc) const
 {
 	uint8 AreaID = RECAST_DEFAULT_AREA;
 	DetourNavMesh->getPolyArea(PolyID, &AreaID);
 
 	const FVector::FReal AreaTravelCost = Filter->getAreaCost(AreaID);
-	return UE_REAL_TO_FLOAT_CLAMPED_MAX(AreaTravelCost * (EndLoc - StartLoc).Size());
+	return AreaTravelCost * (EndLoc - StartLoc).Size();
 }
 
 void FPImplRecastNavMesh::PostProcessPath(dtStatus FindPathStatus, FNavMeshPath& Path,
@@ -1526,9 +1524,9 @@ static void StorePathfindingDebugData(const dtNavMeshQuery& NavQuery, const dtNa
 		FRecastDebugPathfindingNode NodeInfo;
 		NodeInfo.PolyRef = Node->id;
 		NodeInfo.ParentRef = ParentNode ? ParentNode->id : 0;
-		NodeInfo.Cost = UE_REAL_TO_FLOAT_CLAMPED_MAX(Node->cost); // LWC_TODO_AI: Make FRecastDebugPathfindingNode::Cost and FRecastDebugPathfindingNode::Total FVector::FReal. Not until after 5.0!
-		NodeInfo.TotalCost = UE_REAL_TO_FLOAT_CLAMPED_MAX(Node->total);
-		NodeInfo.Length = 0.0f;
+		NodeInfo.Cost = Node->cost; 
+		NodeInfo.TotalCost = Node->total;
+		NodeInfo.Length = 0.;
 		NodeInfo.bOpenSet = (Node->flags & DT_NODE_OPEN) != 0;
 		NodeInfo.bModified = true;
 		NodeInfo.NodePos = Recast2UnrealPoint(&Node->pos[0]);
@@ -1561,8 +1559,7 @@ static void StorePathfindingDebugData(const dtNavMeshQuery& NavQuery, const dtNa
 			
 			NodePathLength[NodePoolIdx] = LinkAndParentLength;
 
-			// LWC_TODO_AI: Make Length FVector::FReal. Not until after 5.0!
-			NodeInfo.Length = UE_REAL_TO_FLOAT_CLAMPED_MAX(LinkAndParentLength);
+			NodeInfo.Length = LinkAndParentLength;
 		}
 
 		Data.Nodes.Add(NodeInfo);
@@ -1601,7 +1598,7 @@ static void StorePathfindingDebugStep(const dtNavMeshQuery& NavQuery, const dtNa
 	}
 }
 
-int32 FPImplRecastNavMesh::DebugPathfinding(const FVector& StartLoc, const FVector& EndLoc, const float CostLimit, const FNavigationQueryFilter& Filter, const UObject* Owner, TArray<FRecastDebugPathfindingData>& Steps)
+int32 FPImplRecastNavMesh::DebugPathfinding(const FVector& StartLoc, const FVector& EndLoc, const FVector::FReal CostLimit, const FNavigationQueryFilter& Filter, const UObject* Owner, TArray<FRecastDebugPathfindingData>& Steps)
 {
 	int32 NumSteps = 0;
 
@@ -1909,7 +1906,6 @@ bool FPImplRecastNavMesh::FindPolysAroundCircle(const FVector& CenterPos, const
 		return false;
 	}
 
-	// LWC_TODO_AI: Remove this when costs are treated as FReals throughout. Not until after 5.0!
 	TArray<FVector::FReal> PolysCost;
 
 	// limit max number of polys found by that function
@@ -1960,11 +1956,11 @@ bool FPImplRecastNavMesh::FindPolysAroundCircle(const FVector& CenterPos, const
 	return false;
 }
 
-bool FPImplRecastNavMesh::GetPolysWithinPathingDistance(FVector const& StartLoc, const float PathingDistance,
+bool FPImplRecastNavMesh::GetPolysWithinPathingDistance(FVector const& StartLoc, const FVector::FReal PathingDistance,
 	const FNavigationQueryFilter& Filter, const UObject* Owner,
 	TArray<NavNodeRef>& FoundPolys, FRecastDebugPathfindingData* OutDebugData) const
 {
-	ensure(PathingDistance > 0.0f && "PathingDistance <= 0 doesn't make sense");
+	ensure(PathingDistance > 0. && "PathingDistance <= 0 doesn't make sense");
 	
 	// limit max number of polys found by that function
 	// if you need more, please scan manually using ARecastNavMesh::GetPolyNeighbors for A*/Dijkstra loop
diff --git a/Engine/Source/Runtime/NavigationSystem/Private/NavMesh/RecastNavMesh.cpp b/Engine/Source/Runtime/NavigationSystem/Private/NavMesh/RecastNavMesh.cpp
--- a/Engine/Source/Runtime/NavigationSystem/Private/NavMesh/RecastNavMesh.cpp
+++ b/Engine/Source/Runtime/NavigationSystem/Private/NavMesh/RecastNavMesh.cpp
@@ -1724,21 +1724,21 @@ bool ARecastNavMesh::ProjectPointMulti(const FVector& Point, TArray<FNavLocation
 	return RecastNavMeshImpl && RecastNavMeshImpl->ProjectPointMulti(Point, OutLocations, Extent, MinZ, MaxZ, GetRightFilterRef(Filter), QueryOwner);
 }
 
-ENavigationQueryResult::Type ARecastNavMesh::CalcPathCost(const FVector& PathStart, const FVector& PathEnd, float& OutPathCost, FSharedConstNavQueryFilter QueryFilter, const UObject* QueryOwner) const
+ENavigationQueryResult::Type ARecastNavMesh::CalcPathCost(const FVector& PathStart, const FVector& PathEnd, FVector::FReal& OutPathCost, FSharedConstNavQueryFilter QueryFilter, const UObject* QueryOwner) const
 {
-	float TmpPathLength = 0.f;
+	FVector::FReal TmpPathLength = 0.f;
 	ENavigationQueryResult::Type Result = CalcPathLengthAndCost(PathStart, PathEnd, TmpPathLength, OutPathCost, QueryFilter, QueryOwner);
 	return Result;
 }
 
-ENavigationQueryResult::Type ARecastNavMesh::CalcPathLength(const FVector& PathStart, const FVector& PathEnd, float& OutPathLength, FSharedConstNavQueryFilter QueryFilter, const UObject* QueryOwner) const
+ENavigationQueryResult::Type ARecastNavMesh::CalcPathLength(const FVector& PathStart, const FVector& PathEnd, FVector::FReal& OutPathLength, FSharedConstNavQueryFilter QueryFilter, const UObject* QueryOwner) const
 {
-	float TmpPathCost = 0.f;
+	FVector::FReal TmpPathCost = 0.f;
 	ENavigationQueryResult::Type Result = CalcPathLengthAndCost(PathStart, PathEnd, OutPathLength, TmpPathCost, QueryFilter, QueryOwner);
 	return Result;
 }
 
-ENavigationQueryResult::Type ARecastNavMesh::CalcPathLengthAndCost(const FVector& PathStart, const FVector& PathEnd, float& OutPathLength, float& OutPathCost, FSharedConstNavQueryFilter QueryFilter, const UObject* QueryOwner) const
+ENavigationQueryResult::Type ARecastNavMesh::CalcPathLengthAndCost(const FVector& PathStart, const FVector& PathEnd, FVector::FReal& OutPathLength, FVector::FReal& OutPathCost, FSharedConstNavQueryFilter QueryFilter, const UObject* QueryOwner) const
 {
 	ENavigationQueryResult::Type Result = ENavigationQueryResult::Invalid;
 
@@ -1746,7 +1746,7 @@ ENavigationQueryResult::Type ARecastNavMesh::CalcPathLengthAndCost(const FVector
 	{
 		if ((PathStart - PathEnd).IsNearlyZero() == true)
 		{
-			OutPathLength = 0.f;
+			OutPathLength = 0.;
 			Result = ENavigationQueryResult::Success;
 		}
 		else
@@ -1754,16 +1754,14 @@ ENavigationQueryResult::Type ARecastNavMesh::CalcPathLengthAndCost(const FVector
 			TSharedRef<FNavMeshPath> Path = MakeShareable(new FNavMeshPath());
 			Path->SetWantsStringPulling(false);
 			Path->SetWantsPathCorridor(true);
-			
-			// LWC_TODO_AI: CostLimit should be FVector::FReal. Not until after 5.0!
-			const float CostLimit = FLT_MAX;
+
+			const FVector::FReal CostLimit = TNumericLimits<FVector::FReal>::Max();
 			Result = RecastNavMeshImpl->FindPath(PathStart, PathEnd, CostLimit, Path.Get(), GetRightFilterRef(QueryFilter), QueryOwner);
 
 			if (Result == ENavigationQueryResult::Success || (Result == ENavigationQueryResult::Fail && Path->IsPartial()))
 			{
-				// LWC_TODO_AI: These should be FVector::FReal. Not until after 5.0!
-				OutPathLength = FMath::Min(TNumericLimits<float>::Max(), Path->GetTotalPathLength());
-				OutPathCost = FMath::Min(TNumericLimits<float>::Max(), Path->GetCost());
+				OutPathLength = Path->GetTotalPathLength();
+				OutPathCost = Path->GetCost();
 			}
 		}
 	}
@@ -1805,11 +1803,11 @@ NavNodeRef ARecastNavMesh::FindNearestPoly(FVector const& Loc, FVector const& Ex
 	return PolyRef;
 }
 
-float ARecastNavMesh::FindDistanceToWall(const FVector& StartLoc, FSharedConstNavQueryFilter Filter, float MaxDistance, FVector* OutClosestPointOnWall) const
+FVector::FReal ARecastNavMesh::FindDistanceToWall(const FVector& StartLoc, FSharedConstNavQueryFilter Filter, FVector::FReal MaxDistance, FVector* OutClosestPointOnWall) const
 {
 	if (HasValidNavmesh() == false)
 	{
-		return 0.f;
+		return 0.;
 	}
 
 	const FNavigationQueryFilter& FilterToUse = GetRightFilterRef(Filter);
@@ -1820,7 +1818,7 @@ float ARecastNavMesh::FindDistanceToWall(const FVector& StartLoc, FSharedConstNa
 	if (QueryFilter == nullptr)
 	{
 		UE_VLOG(this, LogNavigation, Warning, TEXT("ARecastNavMesh::FindDistanceToWall failing due to QueryFilter == NULL"));
-		return 0.f;
+		return 0.;
 	}
 
 	const FVector NavExtent = GetModifiedQueryExtent(GetDefaultQueryExtent());
@@ -1844,11 +1842,11 @@ float ARecastNavMesh::FindDistanceToWall(const FVector& StartLoc, FSharedConstNa
 			{
 				*OutClosestPointOnWall = Recast2UnrealPoint(TmpHitPos);
 			}
-			return UE_REAL_TO_FLOAT_CLAMPED_MAX(DistanceToWall);
+			return DistanceToWall;
 		}
 	}
 
-	return 0.f;
+	return 0.;
 }
 
 void ARecastNavMesh::UpdateCustomLink(const INavLinkCustomInterface* CustomLink)
@@ -2116,7 +2114,7 @@ bool ARecastNavMesh::GetClusterBounds(NavNodeRef ClusterRef, FBox& OutBounds) co
 }
 #endif // WITH_NAVMESH_CLUSTER_LINKS
 
-bool ARecastNavMesh::GetPolysWithinPathingDistance(FVector const& StartLoc, const float PathingDistance, TArray<NavNodeRef>& FoundPolys,
+bool ARecastNavMesh::GetPolysWithinPathingDistance(FVector const& StartLoc, const FVector::FReal PathingDistance, TArray<NavNodeRef>& FoundPolys,
 	FSharedConstNavQueryFilter Filter, const UObject* QueryOwner, FRecastDebugPathfindingData* DebugData) const
 {
 	return RecastNavMeshImpl && RecastNavMeshImpl->GetPolysWithinPathingDistance(StartLoc, PathingDistance, GetRightFilterRef(Filter), QueryOwner, FoundPolys, DebugData);
diff --git a/Engine/Source/Runtime/NavigationSystem/Private/NavMesh/RecastNavMeshGenerator.cpp b/Engine/Source/Runtime/NavigationSystem/Private/NavMesh/RecastNavMeshGenerator.cpp
--- a/Engine/Source/Runtime/NavigationSystem/Private/NavMesh/RecastNavMeshGenerator.cpp
+++ b/Engine/Source/Runtime/NavigationSystem/Private/NavMesh/RecastNavMeshGenerator.cpp
@@ -2195,7 +2195,6 @@ void FRecastTileGenerator::PrepareVoxelCache(const TNavStatArray<uint8>& RawColl
 
 	VoxelCacheContext.SetupForTile(TileConfig.bmin, TileConfig.bmax, RasterizationPadding);
 
-	// LWC_TODO_AI: Precision loss and unnecessary float cast (if ModifyWalkableFloorZ() takes FVector::FReal). Not until after 5.0!
 	float SlopeCosPerActor = UE_REAL_TO_FLOAT(WalkableSlopeCos);
 	CachedCollisions.Header.SlopeOverride.ModifyWalkableFloorZ(SlopeCosPerActor);
 
@@ -4175,7 +4174,7 @@ void FRecastTileGenerator::MarkDynamicArea(const FAreaNavModifier& Modifier, con
 
 			// Only scaling and translation
 			FVector Scale3D = LocalToWorld.GetScale3D().GetAbs();
-			// LWC_TODO_AI: Precision loss FCylinderNavAreaData should probably work with FVector::FReal. Probably not until after 5.0!
+
 			CylinderData.Height = UE_REAL_TO_FLOAT(CylinderData.Height * Scale3D.Z);
 			CylinderData.Radius = UE_REAL_TO_FLOAT(CylinderData.Radius * FMath::Max(Scale3D.X, Scale3D.Y));
 			CylinderData.Origin = LocalToWorld.TransformPosition(CylinderData.Origin);
diff --git a/Engine/Source/Runtime/NavigationSystem/Private/NavigationData.cpp b/Engine/Source/Runtime/NavigationSystem/Private/NavigationData.cpp
--- a/Engine/Source/Runtime/NavigationSystem/Private/NavigationData.cpp
+++ b/Engine/Source/Runtime/NavigationSystem/Private/NavigationData.cpp
@@ -25,7 +25,7 @@ static TAutoConsoleVariable<int32> CVarDestroyNavDataInCleanUpAndMarkPendingKill
 //----------------------------------------------------------------------//
 // FPathFindingQuery
 //----------------------------------------------------------------------//
-FPathFindingQuery::FPathFindingQuery(const UObject* InOwner, const ANavigationData& InNavData, const FVector& Start, const FVector& End, FSharedConstNavQueryFilter SourceQueryFilter, FNavPathSharedPtr InPathInstanceToFill, const float CostLimit) :
+FPathFindingQuery::FPathFindingQuery(const UObject* InOwner, const ANavigationData& InNavData, const FVector& Start, const FVector& End, FSharedConstNavQueryFilter SourceQueryFilter, FNavPathSharedPtr InPathInstanceToFill, const FVector::FReal CostLimit) :
 	FPathFindingQueryData(InOwner, Start, End, SourceQueryFilter, 0 /*InNavDataFlags*/, true /*bInAllowPartialPaths*/, CostLimit),
 	NavData(&InNavData), PathInstanceToFill(InPathInstanceToFill), NavAgentProperties(InNavData.GetConfig())
 {
@@ -35,7 +35,7 @@ FPathFindingQuery::FPathFindingQuery(const UObject* InOwner, const ANavigationDa
 	}
 }
 
-FPathFindingQuery::FPathFindingQuery(const INavAgentInterface& InNavAgent, const ANavigationData& InNavData, const FVector& Start, const FVector& End, FSharedConstNavQueryFilter SourceQueryFilter, FNavPathSharedPtr InPathInstanceToFill, const float CostLimit) :
+FPathFindingQuery::FPathFindingQuery(const INavAgentInterface& InNavAgent, const ANavigationData& InNavData, const FVector& Start, const FVector& End, FSharedConstNavQueryFilter SourceQueryFilter, FNavPathSharedPtr InPathInstanceToFill, const FVector::FReal CostLimit) :
 	FPathFindingQueryData(Cast<UObject>(&InNavAgent), Start, End, SourceQueryFilter, 0 /*InNavDataFlags*/, true /*bInAllowPartialPaths*/, CostLimit),
 	NavData(&InNavData), PathInstanceToFill(InPathInstanceToFill), NavAgentProperties(InNavAgent.GetNavAgentPropertiesRef())
 {
@@ -88,7 +88,7 @@ FPathFindingQuery::FPathFindingQuery(FNavPathSharedRef PathToRecalculate, const
 	}
 }
 
-float FPathFindingQuery::ComputeCostLimitFromHeuristic(const FVector& StartPos, const FVector& EndPos, const float HeuristicScale, const float CostLimitFactor, const float MinimumCostLimit) const
+FVector::FReal FPathFindingQuery::ComputeCostLimitFromHeuristic(const FVector& StartPos, const FVector& EndPos, const FVector::FReal HeuristicScale, const FVector::FReal CostLimitFactor, const FVector::FReal MinimumCostLimit) const
 {
 	if (CostLimitFactor == FLT_MAX)
 	{
@@ -96,8 +96,8 @@ float FPathFindingQuery::ComputeCostLimitFromHeuristic(const FVector& StartPos,
 	}
 	else
 	{
-		const float OriginalHeuristicEstimate = HeuristicScale * FVector::Dist(StartPos, EndPos);
-		return FMath::Clamp(CostLimitFactor * OriginalHeuristicEstimate, MinimumCostLimit, FLT_MAX);
+		const FVector::FReal OriginalHeuristicEstimate = HeuristicScale * FVector::Dist(StartPos, EndPos);
+		return FMath::Clamp(CostLimitFactor * OriginalHeuristicEstimate, MinimumCostLimit, TNumericLimits<FVector::FReal>::Max());
 	}
 }
 
@@ -106,7 +106,7 @@ float FPathFindingQuery::ComputeCostLimitFromHeuristic(const FVector& StartPos,
 //----------------------------------------------------------------------//
 uint32 FAsyncPathFindingQuery::LastPathFindingUniqueID = INVALID_NAVQUERYID;
 
-FAsyncPathFindingQuery::FAsyncPathFindingQuery(const UObject* InOwner, const ANavigationData& InNavData, const FVector& Start, const FVector& End, const FNavPathQueryDelegate& Delegate, FSharedConstNavQueryFilter SourceQueryFilter, const float CostLimit)
+FAsyncPathFindingQuery::FAsyncPathFindingQuery(const UObject* InOwner, const ANavigationData& InNavData, const FVector& Start, const FVector& End, const FNavPathQueryDelegate& Delegate, FSharedConstNavQueryFilter SourceQueryFilter, const FVector::FReal CostLimit)
 : FPathFindingQuery(InOwner, InNavData, Start, End, SourceQueryFilter)
 , QueryID(GetUniqueID())
 , OnDoneDelegate(Delegate)
@@ -876,3 +876,35 @@ void ANavigationData::DrawDebugPath(FNavigationPath* Path, FColor PathColor, UCa
 {
 	DrawDebugPath(Path, PathColor, Canvas, bPersistent, -1.f, NextPathPointIndex);
 }
+
+ENavigationQueryResult::Type ANavigationData::CalcPathCost(const FVector& PathStart, const FVector& PathEnd, float& OutPathCost, FSharedConstNavQueryFilter QueryFilter, const UObject* Querier) const
+{
+	FVector::FReal PathCost = OutPathCost;
+
+	const ENavigationQueryResult::Type Result = CalcPathCost(PathStart, PathEnd, PathCost, QueryFilter, Querier);
+	OutPathCost = UE_REAL_TO_FLOAT_CLAMPED(PathCost);
+
+	return Result;
+}
+
+ENavigationQueryResult::Type ANavigationData::CalcPathLength(const FVector& PathStart, const FVector& PathEnd, float& OutPathLength, FSharedConstNavQueryFilter QueryFilter, const UObject* Querier) const
+{
+	FVector::FReal PathLength = OutPathLength;
+
+	const ENavigationQueryResult::Type Result = CalcPathLength(PathStart, PathEnd, PathLength, QueryFilter, Querier);
+	OutPathLength = UE_REAL_TO_FLOAT_CLAMPED(PathLength);
+
+	return Result;
+}
+
+ENavigationQueryResult::Type ANavigationData::CalcPathLengthAndCost(const FVector& PathStart, const FVector& PathEnd, float& OutPathLength, float& OutPathCost, FSharedConstNavQueryFilter QueryFilter, const UObject* Querier) const
+{
+	FVector::FReal PathLength = OutPathLength;
+	FVector::FReal PathCost = OutPathCost;
+
+	const ENavigationQueryResult::Type Result = CalcPathLengthAndCost(PathStart, PathEnd, PathLength, PathCost, QueryFilter, Querier);
+	OutPathLength = UE_REAL_TO_FLOAT_CLAMPED(PathLength);
+	OutPathCost = UE_REAL_TO_FLOAT_CLAMPED(PathCost);
+
+	return Result;
+}
diff --git a/Engine/Source/Runtime/NavigationSystem/Private/NavigationPath.cpp b/Engine/Source/Runtime/NavigationSystem/Private/NavigationPath.cpp
--- a/Engine/Source/Runtime/NavigationSystem/Private/NavigationPath.cpp
+++ b/Engine/Source/Runtime/NavigationSystem/Private/NavigationPath.cpp
@@ -215,7 +215,7 @@ void FNavigationPath::DebugDraw(const ANavigationData* NavData, FColor PathColor
 	// draw last vert
 	if (NumPathVerts > 0)
 	{
-		DrawDebugBox(World, PathPoints[NumPathVerts-1].Location + NavigationDebugDrawing::PathOffset, FVector(15.f), PathColor, bPersistent, LifeTime);
+		DrawDebugBox(World, PathPoints[NumPathVerts-1].Location + NavigationDebugDrawing::PathOffset, FVector(15.), PathColor, bPersistent, LifeTime);
 	}
 
 	// if observing goal actor draw a radius and a line to the goal
@@ -223,7 +223,7 @@ void FNavigationPath::DebugDraw(const ANavigationData* NavData, FColor PathColor
 	{
 		const FVector GoalLocation = GetGoalLocation() + NavigationDebugDrawing::PathOffset;
 		const FVector EndLocation = GetEndLocation() + NavigationDebugDrawing::PathOffset;
-		static const FVector CylinderHalfHeight = FVector::UpVector * 10.f;
+		static const FVector CylinderHalfHeight = FVector::UpVector * 10.;
 		DrawDebugCylinder(World, EndLocation - CylinderHalfHeight, EndLocation + CylinderHalfHeight, FMath::Sqrt(GoalActorLocationTetherDistanceSq), 16, PathColor, bPersistent, LifeTime);
 		DrawDebugLine(World, EndLocation, GoalLocation, Grey, bPersistent, LifeTime);
 	}
@@ -244,7 +244,7 @@ bool FNavigationPath::ContainsNode(NavNodeRef NodeRef) const
 	return ShortcutNodeRefs.Find(NodeRef) != INDEX_NONE;
 }
 
-float FNavigationPath::GetLengthFromPosition(FVector SegmentStart, uint32 NextPathPointIndex) const
+FVector::FReal FNavigationPath::GetLengthFromPosition(FVector SegmentStart, uint32 NextPathPointIndex) const
 {
 	if (NextPathPointIndex >= (uint32)PathPoints.Num())
 	{
@@ -252,7 +252,7 @@ float FNavigationPath::GetLengthFromPosition(FVector SegmentStart, uint32 NextPa
 	}
 	
 	const uint32 PathPointsCount = PathPoints.Num();
-	FVector::FReal PathDistance = 0.f;
+	FVector::FReal PathDistance = 0.;
 
 	for (uint32 PathIndex = NextPathPointIndex; PathIndex < PathPointsCount; ++PathIndex)
 	{
@@ -261,7 +261,7 @@ float FNavigationPath::GetLengthFromPosition(FVector SegmentStart, uint32 NextPa
 		SegmentStart = SegmentEnd;
 	}
 
-	return UE_REAL_TO_FLOAT_CLAMPED_MAX(PathDistance);
+	return PathDistance;
 }
 
 bool FNavigationPath::ContainsCustomLink(uint32 LinkUniqueId) const
@@ -524,16 +524,16 @@ void UNavigationPath::EnableRecalculationOnInvalidation(TEnumAsByte<ENavigationO
 	}
 }
 
-float UNavigationPath::GetPathLength() const
+double UNavigationPath::GetPathLength() const
 {
 	check((SharedPath.IsValid() && SharedPath->IsValid()) == !!bIsValid);
-	return !!bIsValid ? SharedPath->GetLength() : -1.f;
+	return !!bIsValid ? SharedPath->GetLength() : -1.;
 }
 
-float UNavigationPath::GetPathCost() const
+double UNavigationPath::GetPathCost() const
 {
 	check((SharedPath.IsValid() && SharedPath->IsValid()) == !!bIsValid);
-	return !!bIsValid ? SharedPath->GetCost() : -1.f;
+	return !!bIsValid ? SharedPath->GetCost() : -1.;
 }
 
 bool UNavigationPath::IsPartial() const
diff --git a/Engine/Source/Runtime/NavigationSystem/Private/NavigationSystem.cpp b/Engine/Source/Runtime/NavigationSystem/Private/NavigationSystem.cpp
--- a/Engine/Source/Runtime/NavigationSystem/Private/NavigationSystem.cpp
+++ b/Engine/Source/Runtime/NavigationSystem/Private/NavigationSystem.cpp
@@ -1642,6 +1642,17 @@ bool UNavigationSystemV1::GetRandomPointInNavigableRadius(const FVector& Origin,
 }
 
 ENavigationQueryResult::Type UNavigationSystemV1::GetPathCost(const FVector& PathStart, const FVector& PathEnd, float& OutPathCost, const ANavigationData* NavData, FSharedConstNavQueryFilter QueryFilter) const
+{
+	FVector::FReal PathCost = OutPathCost;
+
+	const ENavigationQueryResult::Type Result = GetPathCost(PathStart, PathEnd, PathCost, NavData, QueryFilter);
+
+	OutPathCost = UE_REAL_TO_FLOAT_CLAMPED(PathCost);
+
+	return Result;
+}
+
+ENavigationQueryResult::Type UNavigationSystemV1::GetPathCost(const FVector& PathStart, const FVector& PathEnd, FVector::FReal& OutPathCost, const ANavigationData* NavData, FSharedConstNavQueryFilter QueryFilter) const
 {
 	SCOPE_CYCLE_COUNTER(STAT_Navigation_QueriesTimeSync);
 
@@ -1654,6 +1665,17 @@ ENavigationQueryResult::Type UNavigationSystemV1::GetPathCost(const FVector& Pat
 }
 
 ENavigationQueryResult::Type UNavigationSystemV1::GetPathLength(const FVector& PathStart, const FVector& PathEnd, float& OutPathLength, const ANavigationData* NavData, FSharedConstNavQueryFilter QueryFilter) const
+{
+	FVector::FReal PathLength = OutPathLength;
+
+	const ENavigationQueryResult::Type Result = GetPathLength(PathStart, PathEnd, PathLength, NavData, QueryFilter);
+
+	OutPathLength = UE_REAL_TO_FLOAT_CLAMPED(PathLength);
+
+	return Result;
+}
+
+ENavigationQueryResult::Type UNavigationSystemV1::GetPathLength(const FVector& PathStart, const FVector& PathEnd, FVector::FReal& OutPathLength, const ANavigationData* NavData, FSharedConstNavQueryFilter QueryFilter) const
 {
 	SCOPE_CYCLE_COUNTER(STAT_Navigation_QueriesTimeSync);
 
@@ -1666,6 +1688,19 @@ ENavigationQueryResult::Type UNavigationSystemV1::GetPathLength(const FVector& P
 }
 
 ENavigationQueryResult::Type UNavigationSystemV1::GetPathLengthAndCost(const FVector& PathStart, const FVector& PathEnd, float& OutPathLength, float& OutPathCost, const ANavigationData* NavData, FSharedConstNavQueryFilter QueryFilter) const
+{
+	FVector::FReal PathLength = OutPathLength;
+	FVector::FReal PathCost = OutPathCost;
+
+	const ENavigationQueryResult::Type Result = GetPathLengthAndCost(PathStart, PathEnd, PathLength, PathCost, NavData, QueryFilter);
+
+	OutPathLength = UE_REAL_TO_FLOAT_CLAMPED(PathLength);
+	OutPathCost = UE_REAL_TO_FLOAT_CLAMPED(PathCost);
+
+	return Result;
+}
+
+ENavigationQueryResult::Type UNavigationSystemV1::GetPathLengthAndCost(const FVector& PathStart, const FVector& PathEnd, FVector::FReal& OutPathLength, FVector::FReal& OutPathCost, const ANavigationData* NavData, FSharedConstNavQueryFilter QueryFilter) const
 {
 	SCOPE_CYCLE_COUNTER(STAT_Navigation_QueriesTimeSync);
 
@@ -4337,6 +4372,17 @@ bool UNavigationSystemV1::K2_GetRandomLocationInNavigableRadius(UObject* WorldCo
 }
 
 ENavigationQueryResult::Type UNavigationSystemV1::GetPathCost(UObject* WorldContextObject, const FVector& PathStart, const FVector& PathEnd, float& OutPathCost, ANavigationData* NavData, TSubclassOf<UNavigationQueryFilter> FilterClass)
+{
+	FVector::FReal PathCost = OutPathCost;
+
+	const ENavigationQueryResult::Type Result = GetPathCost(WorldContextObject, PathStart, PathEnd, PathCost, NavData, FilterClass);
+
+	OutPathCost = UE_REAL_TO_FLOAT_CLAMPED(PathCost);
+
+	return Result;
+}
+
+ENavigationQueryResult::Type UNavigationSystemV1::GetPathCost(UObject* WorldContextObject, const FVector& PathStart, const FVector& PathEnd, double& OutPathCost, ANavigationData* NavData, TSubclassOf<UNavigationQueryFilter> FilterClass)
 {
 	UWorld* World = GEngine->GetWorldFromContextObject(WorldContextObject, EGetWorldErrorMode::LogAndReturnNull);
 	UNavigationSystemV1* NavSys = FNavigationSystem::GetCurrent<UNavigationSystemV1>(World);
@@ -4354,8 +4400,17 @@ ENavigationQueryResult::Type UNavigationSystemV1::GetPathCost(UObject* WorldCont
 
 ENavigationQueryResult::Type UNavigationSystemV1::GetPathLength(UObject* WorldContextObject, const FVector& PathStart, const FVector& PathEnd, float& OutPathLength, ANavigationData* NavData, TSubclassOf<UNavigationQueryFilter> FilterClass)
 {
-	float PathLength = 0.f;
+	FVector::FReal PathLength = OutPathLength;
+
+	const ENavigationQueryResult::Type Result = GetPathLength(WorldContextObject, PathStart, PathEnd, PathLength, NavData, FilterClass);
+
+	OutPathLength = UE_REAL_TO_FLOAT_CLAMPED(PathLength);
 
+	return Result;
+}
+
+ENavigationQueryResult::Type UNavigationSystemV1::GetPathLength(UObject* WorldContextObject, const FVector& PathStart, const FVector& PathEnd, double& OutPathLength, ANavigationData* NavData, TSubclassOf<UNavigationQueryFilter> FilterClass)
+{
 	UWorld* World = GEngine->GetWorldFromContextObject(WorldContextObject, EGetWorldErrorMode::LogAndReturnNull);
 	UNavigationSystemV1* NavSys = FNavigationSystem::GetCurrent<UNavigationSystemV1>(World);
 	if (NavSys)
diff --git a/Engine/Source/Runtime/NavigationSystem/Public/AbstractNavData.h b/Engine/Source/Runtime/NavigationSystem/Public/AbstractNavData.h
--- a/Engine/Source/Runtime/NavigationSystem/Public/AbstractNavData.h
+++ b/Engine/Source/Runtime/NavigationSystem/Public/AbstractNavData.h
@@ -78,9 +78,9 @@ public:
 	virtual bool ProjectPoint(const FVector& Point, FNavLocation& OutLocation, const FVector& Extent, FSharedConstNavQueryFilter Filter = NULL, const UObject* Querier = NULL) const override { return false; }
 	virtual void BatchProjectPoints(TArray<FNavigationProjectionWork>& Workload, const FVector& Extent, FSharedConstNavQueryFilter Filter = NULL, const UObject* Querier = NULL) const override {};
 	virtual void BatchProjectPoints(TArray<FNavigationProjectionWork>& Workload, FSharedConstNavQueryFilter Filter = NULL, const UObject* Querier = NULL) const override {}
-	virtual ENavigationQueryResult::Type CalcPathCost(const FVector& PathStart, const FVector& PathEnd, float& OutPathCost, FSharedConstNavQueryFilter QueryFilter = NULL, const UObject* Querier = NULL) const override { return ENavigationQueryResult::Invalid; }
-	virtual ENavigationQueryResult::Type CalcPathLength(const FVector& PathStart, const FVector& PathEnd, float& OutPathLength, FSharedConstNavQueryFilter QueryFilter = NULL, const UObject* Querier = NULL) const override { return ENavigationQueryResult::Invalid; }
-	virtual ENavigationQueryResult::Type CalcPathLengthAndCost(const FVector& PathStart, const FVector& PathEnd, float& OutPathLength, float& OutPathCost, FSharedConstNavQueryFilter QueryFilter = NULL, const UObject* Querier = NULL) const override { return ENavigationQueryResult::Invalid;	}
+	virtual ENavigationQueryResult::Type CalcPathCost(const FVector& PathStart, const FVector& PathEnd, FVector::FReal& OutPathCost, FSharedConstNavQueryFilter QueryFilter = NULL, const UObject* Querier = NULL) const override { return ENavigationQueryResult::Invalid; }
+	virtual ENavigationQueryResult::Type CalcPathLength(const FVector& PathStart, const FVector& PathEnd, FVector::FReal& OutPathLength, FSharedConstNavQueryFilter QueryFilter = NULL, const UObject* Querier = NULL) const override { return ENavigationQueryResult::Invalid; }
+	virtual ENavigationQueryResult::Type CalcPathLengthAndCost(const FVector& PathStart, const FVector& PathEnd, FVector::FReal& OutPathLength, FVector::FReal& OutPathCost, FSharedConstNavQueryFilter QueryFilter = NULL, const UObject* Querier = NULL) const override { return ENavigationQueryResult::Invalid; }
 	virtual bool DoesNodeContainLocation(NavNodeRef NodeRef, const FVector& WorldSpaceLocation) const override { return true; }
 	virtual void OnNavAreaAdded(const UClass* NavAreaClass, int32 AgentIndex) override {}
 	virtual void OnNavAreaRemoved(const UClass* NavAreaClass) override {};
diff --git a/Engine/Source/Runtime/NavigationSystem/Public/NavMesh/NavMeshPath.h b/Engine/Source/Runtime/NavigationSystem/Public/NavMesh/NavMeshPath.h
--- a/Engine/Source/Runtime/NavigationSystem/Public/NavMesh/NavMeshPath.h
+++ b/Engine/Source/Runtime/NavigationSystem/Public/NavMesh/NavMeshPath.h
@@ -10,7 +10,6 @@
 
 #define RECAST_STRAIGHTPATH_OFFMESH_CONNECTION 0x04
 
-// LWC_TODO_AI: A lot of the floats in this file should be FVector::FReal. Not until after 5.0!
 /** Helper to translate FNavPathPoint.Flags. */
 struct NAVIGATIONSYSTEM_API FNavMeshNodeFlags
 {
@@ -59,7 +58,7 @@ struct NAVIGATIONSYSTEM_API FNavMeshPath : public FNavigationPath
 
 	bool ContainsWithSameEnd(const FNavMeshPath* Other) const;
 
-	void OffsetFromCorners(float Distance);
+	void OffsetFromCorners(FVector::FReal Distance);
 
 	void ApplyFlags(int32 NavDataFlags);
 
@@ -69,10 +68,10 @@ struct NAVIGATIONSYSTEM_API FNavMeshPath : public FNavigationPath
 	bool GetNodeFlags(int32 NodeIdx, FNavMeshNodeFlags& Flags) const;
 
 	/** get cost of path, starting from next poly in corridor */
-	virtual float GetCostFromNode(NavNodeRef PathNode) const override { return GetCostFromIndex(PathCorridor.Find(PathNode) + 1); }
+	virtual FVector::FReal GetCostFromNode(NavNodeRef PathNode) const override { return GetCostFromIndex(PathCorridor.Find(PathNode) + 1); }
 
 	/** get cost of path, starting from given point */
-	virtual float GetCostFromIndex(int32 PathPointIndex) const override
+	virtual FVector::FReal GetCostFromIndex(int32 PathPointIndex) const override
 	{
 		FVector::FReal TotalCost = 0.f;
 		const FVector::FReal* Cost = PathCorridorCost.GetData();
@@ -81,11 +80,10 @@ struct NAVIGATIONSYSTEM_API FNavMeshPath : public FNavigationPath
 			TotalCost += *Cost;
 		}
 
-		// LWC_TODO_AI: Precision Loss. This should return a FReal. Not until after 5.0!
-		return UE_REAL_TO_FLOAT_CLAMPED_MAX(TotalCost);
+		return TotalCost;
 	}
 
-	FORCEINLINE_DEBUGGABLE float GetTotalPathLength() const
+	FORCEINLINE_DEBUGGABLE FVector::FReal GetTotalPathLength() const
 	{
 		return bStringPulled ? GetStringPulledLength(0) : GetPathCorridorLength(0);
 	}
@@ -122,12 +120,12 @@ protected:
 	/** calculates total length of string pulled path. Does not generate string pulled
 	*	path if it's not already generated (see bWantsStringPulling and bStrigPulled)
 	*	Internal use only */
-	float GetStringPulledLength(const int32 StartingPoint) const;
+	FVector::FReal GetStringPulledLength(const int32 StartingPoint) const;
 
 	/** calculates estimated length of path expressed as sequence of navmesh edges.
 	*	It basically sums up distances between every subsequent nav edge pair edge middles.
 	*	Internal use only */
-	float GetPathCorridorLength(const int32 StartingEdge) const;
+	FVector::FReal GetPathCorridorLength(const int32 StartingEdge) const;
 
 	/** it's only const to be callable in const environment. It's not supposed to be called directly externally anyway,
 	*	just as part of retrieving corridor on demand or generating it in internal processes. It fills a mutable
diff --git a/Engine/Source/Runtime/NavigationSystem/Public/NavMesh/PImplRecastNavMesh.h b/Engine/Source/Runtime/NavigationSystem/Public/NavMesh/PImplRecastNavMesh.h
--- a/Engine/Source/Runtime/NavigationSystem/Public/NavMesh/PImplRecastNavMesh.h
+++ b/Engine/Source/Runtime/NavigationSystem/Public/NavMesh/PImplRecastNavMesh.h
@@ -32,7 +32,6 @@ class FRecastNavMeshGenerator;
 
 #define RECAST_VERY_SMALL_AGENT_RADIUS 0.0f
 
-// LWC_TODO_AI: Costs and pathing distance should be FReal. Not until after 5.0!
 /** Engine Private! - Private Implementation details of ARecastNavMesh */
 class NAVIGATIONSYSTEM_API FPImplRecastNavMesh
 {
@@ -83,24 +82,24 @@ public:
 
 	// @TODONAV
 	/** Generates path from the given query. Synchronous. */
-	UE_DEPRECATED(4.25, "Use the version with the added CostLimit parameter (FLT_MAX can be used as default).")
+	UE_DEPRECATED(4.25, "Use the version with the added CostLimit parameter (TNumericLimits<FVector::FReal>::Max() can be used as default).")
 	ENavigationQueryResult::Type FindPath(const FVector& StartLoc, const FVector& EndLoc, FNavMeshPath& Path, const FNavigationQueryFilter& Filter, const UObject* Owner) const;
 	
 	/** Generates path from the given query. Synchronous. */
-	ENavigationQueryResult::Type FindPath(const FVector& StartLoc, const FVector& EndLoc, const float CostLimit, FNavMeshPath& Path, const FNavigationQueryFilter& Filter, const UObject* Owner) const;
+	ENavigationQueryResult::Type FindPath(const FVector& StartLoc, const FVector& EndLoc, const FVector::FReal CostLimit, FNavMeshPath& Path, const FNavigationQueryFilter& Filter, const UObject* Owner) const;
 
 	/** Check if path exists */
 	ENavigationQueryResult::Type TestPath(const FVector& StartLoc, const FVector& EndLoc, const FNavigationQueryFilter& Filter, const UObject* Owner, int32* NumVisitedNodes = 0) const;
 
 	template< typename TRecastAStar, typename TRecastAStartGraph, typename TRecastGraphAStarFilter, typename TRecastAStarResult >
-	ENavigationQueryResult::Type FindPathCustomAStar(TRecastAStartGraph& RecastGraphWrapper, TRecastAStar& AStarAlgo, const FVector& StartLoc, const FVector& EndLoc, const float CostLimit, FNavMeshPath& Path, const FNavigationQueryFilter& Filter, const UObject* Owner) const;
+	ENavigationQueryResult::Type FindPathCustomAStar(TRecastAStartGraph& RecastGraphWrapper, TRecastAStar& AStarAlgo, const FVector& StartLoc, const FVector& EndLoc, const FVector::FReal CostLimit, FNavMeshPath& Path, const FNavigationQueryFilter& Filter, const UObject* Owner) const;
 
 	/** Checks if the whole segment is in navmesh */
 	void Raycast(const FVector& StartLoc, const FVector& EndLoc, const FNavigationQueryFilter& InQueryFilter, const UObject* Owner,
 		ARecastNavMesh::FRaycastResult& RaycastResult, NavNodeRef StartNode = INVALID_NAVNODEREF) const;
 
 	/** Generates path from given query and collect data for every step of A* algorithm */
-	int32 DebugPathfinding(const FVector& StartLoc, const FVector& EndLoc, const float CostLimit, const FNavigationQueryFilter& Filter, const UObject* Owner, TArray<FRecastDebugPathfindingData>& Steps);
+	int32 DebugPathfinding(const FVector& StartLoc, const FVector& EndLoc, const FVector::FReal CostLimit, const FNavigationQueryFilter& Filter, const UObject* Owner, TArray<FRecastDebugPathfindingData>& Steps);
 
 	/** Returns a random location on the navmesh. */
 	FNavLocation GetRandomPoint(const FNavigationQueryFilter& Filter, const UObject* Owner) const;
@@ -135,7 +134,7 @@ public:
 	 *	@NOTE query is not using string-pulled path distance (for performance reasons),
 	 *		it measured distance between middles of portal edges, do you might want to 
 	 *		add an extra margin to PathingDistance */
-	bool GetPolysWithinPathingDistance(FVector const& StartLoc, const float PathingDistance,
+	bool GetPolysWithinPathingDistance(FVector const& StartLoc, const FVector::FReal PathingDistance,
 		const FNavigationQueryFilter& Filter, const UObject* Owner,
 		TArray<NavNodeRef>& FoundPolys, FRecastDebugPathfindingData* DebugData) const;
 
@@ -236,7 +235,7 @@ public:
 	/** calculated cost of given segment if traversed on specified poly. Function measures distance between specified points
 	 *	and returns cost of traversing this distance on given poly.
 	 *	@note no check if segment is on poly is performed. */
-	float CalcSegmentCostOnPoly(NavNodeRef PolyID, const dtQueryFilter* Filter, const FVector& StartLoc, const FVector& EndLoc) const;
+	FVector::FReal CalcSegmentCostOnPoly(NavNodeRef PolyID, const dtQueryFilter* Filter, const FVector& StartLoc, const FVector& EndLoc) const;
 
 	ARecastNavMesh* NavMeshOwner;
 	
@@ -292,7 +291,7 @@ protected:
 };
 
 template< typename TRecastAStar, typename TRecastAStartGraph, typename TRecastGraphAStarFilter, typename TRecastAStarResult >
-ENavigationQueryResult::Type FPImplRecastNavMesh::FindPathCustomAStar(TRecastAStartGraph& RecastGraphWrapper, TRecastAStar& AStarAlgo, const FVector& StartLoc, const FVector& EndLoc, const float CostLimit, FNavMeshPath& Path, const FNavigationQueryFilter& InQueryFilter, const UObject* Owner) const
+ENavigationQueryResult::Type FPImplRecastNavMesh::FindPathCustomAStar(TRecastAStartGraph& RecastGraphWrapper, TRecastAStar& AStarAlgo, const FVector& StartLoc, const FVector& EndLoc, const FVector::FReal CostLimit, FNavMeshPath& Path, const FNavigationQueryFilter& InQueryFilter, const UObject* Owner) const
 {
 	const FRecastQueryFilter* FilterImplementation = (const FRecastQueryFilter*)(InQueryFilter.GetImplementation());
 	if (FilterImplementation == nullptr)
diff --git a/Engine/Source/Runtime/NavigationSystem/Public/NavMesh/RecastNavMesh.h b/Engine/Source/Runtime/NavigationSystem/Public/NavMesh/RecastNavMesh.h
--- a/Engine/Source/Runtime/NavigationSystem/Public/NavMesh/RecastNavMesh.h
+++ b/Engine/Source/Runtime/NavigationSystem/Public/NavMesh/RecastNavMesh.h
@@ -20,7 +20,9 @@
 #define RECAST_DEFAULT_AREA			(RECAST_MAX_AREAS - 1)
 #define RECAST_LOW_AREA				(RECAST_MAX_AREAS - 2)
 #define RECAST_NULL_AREA			0
-#define RECAST_UNWALKABLE_POLY_COST	FLT_MAX // LWC_TODO_AI: This should be TNumericLimits<FVector::FReal>::Max() when costs are upgraded to FReals. Not until after 5.0!
+// Note poly costs are still floats in UE so we are using FLT_MAX as unwalkable still. 
+// Path CostLimit and Distance are based on FVector::FReal.
+#define RECAST_UNWALKABLE_POLY_COST	FLT_MAX 
 
 // If set, recast will use async workers for rebuilding tiles in runtime
 // All access to tile data must be guarded with critical sections
@@ -120,9 +122,9 @@ struct FRecastDebugPathfindingNode
 {
 	NavNodeRef PolyRef;
 	NavNodeRef ParentRef;
-	float Cost; // LWC_TODO_AI: These should be FVector::FReal in the long run! Not until after 5.0!
-	float TotalCost;
-	float Length;
+	FVector::FReal Cost = 0.;
+	FVector::FReal TotalCost = 0.;
+	FVector::FReal Length = 0.;
 
 	FVector NodePos;
 	TArray<FVector3f, TInlineAllocator<6> > Verts; // LWC_TODO: Precision loss. Issue here is regarding debug rendering needing to work with FVector3f.
@@ -132,14 +134,14 @@ struct FRecastDebugPathfindingNode
 	uint8 bOffMeshLink : 1;
 	uint8 bModified : 1;
 
-	FRecastDebugPathfindingNode() : PolyRef(0), ParentRef(0), NumVerts(0) {}
-	FRecastDebugPathfindingNode(NavNodeRef InPolyRef) : PolyRef(InPolyRef), ParentRef(0), NumVerts(0) {}
+	FRecastDebugPathfindingNode() : PolyRef(0), ParentRef(0), NumVerts(0), bOpenSet(0), bOffMeshLink(0), bModified(0) {}
+	FRecastDebugPathfindingNode(NavNodeRef InPolyRef) : PolyRef(InPolyRef), ParentRef(0), NumVerts(0), bOpenSet(0), bOffMeshLink(0), bModified(0) {}
 
 	FORCEINLINE bool operator==(const NavNodeRef& OtherPolyRef) const { return PolyRef == OtherPolyRef; }
 	FORCEINLINE bool operator==(const FRecastDebugPathfindingNode& Other) const { return PolyRef == Other.PolyRef; }
 	FORCEINLINE friend uint32 GetTypeHash(const FRecastDebugPathfindingNode& Other) { return GetTypeHash(Other.PolyRef); }
 
-	FORCEINLINE float GetHeuristicCost() const { return TotalCost - Cost; }
+	FORCEINLINE FVector::FReal GetHeuristicCost() const { return TotalCost - Cost; }
 };
 
 namespace ERecastDebugPathfindingFlags
@@ -563,7 +565,6 @@ namespace FNavMeshConfig
 	};
 }
 
-// LWC_TODO_AI: Many of the virtual methods and members should be changed from float to FVector::FReal. Not for 5.0!
 UCLASS(config=Engine, defaultconfig, hidecategories=(Input,Rendering,Tags,Transformation,Actor,Layers,Replication), notplaceable)
 class NAVIGATIONSYSTEM_API ARecastNavMesh : public ANavigationData
 {
@@ -947,9 +948,9 @@ public:
 	 *	@note function will assert if item's FNavigationProjectionWork.ProjectionLimit is invalid */
 	virtual void BatchProjectPoints(TArray<FNavigationProjectionWork>& Workload, FSharedConstNavQueryFilter Filter = NULL, const UObject* Querier = NULL) const override;
 	
-	virtual ENavigationQueryResult::Type CalcPathCost(const FVector& PathStart, const FVector& PathEnd, float& OutPathCost, FSharedConstNavQueryFilter Filter = NULL, const UObject* Querier = NULL) const override;
-	virtual ENavigationQueryResult::Type CalcPathLength(const FVector& PathStart, const FVector& PathEnd, float& OutPathLength, FSharedConstNavQueryFilter QueryFilter = NULL, const UObject* Querier = NULL) const override;
-	virtual ENavigationQueryResult::Type CalcPathLengthAndCost(const FVector& PathStart, const FVector& PathEnd, float& OutPathLength, float& OutPathCost, FSharedConstNavQueryFilter QueryFilter = NULL, const UObject* Querier = NULL) const override;
+	virtual ENavigationQueryResult::Type CalcPathCost(const FVector& PathStart, const FVector& PathEnd, FVector::FReal& OutPathCost, FSharedConstNavQueryFilter Filter = NULL, const UObject* Querier = NULL) const override;
+	virtual ENavigationQueryResult::Type CalcPathLength(const FVector& PathStart, const FVector& PathEnd, FVector::FReal& OutPathLength, FSharedConstNavQueryFilter QueryFilter = NULL, const UObject* Querier = NULL) const override;
+	virtual ENavigationQueryResult::Type CalcPathLengthAndCost(const FVector& PathStart, const FVector& PathEnd, FVector::FReal& OutPathLength, FVector::FReal& OutPathCost, FSharedConstNavQueryFilter QueryFilter = NULL, const UObject* Querier = NULL) const override;
 	virtual bool DoesNodeContainLocation(NavNodeRef NodeRef, const FVector& WorldSpaceLocation) const override;
 
 	virtual UPrimitiveComponent* ConstructRenderingComponent() override;
@@ -1114,7 +1115,7 @@ public:
 	FVector GetModifiedQueryExtent(const FVector& QueryExtent) const
 	{
 		// Using HALF_WORLD_MAX instead of BIG_NUMBER, else using the extent for a box will result in NaN.
-		return FVector(QueryExtent.X, QueryExtent.Y, QueryExtent.Z >= (float)HALF_WORLD_MAX ? (float)HALF_WORLD_MAX : (QueryExtent.Z + FMath::Max(0.0f, VerticalDeviationFromGroundCompensation)));
+		return FVector(QueryExtent.X, QueryExtent.Y, QueryExtent.Z >= HALF_WORLD_MAX ? HALF_WORLD_MAX : (QueryExtent.Z + FMath::Max(0., VerticalDeviationFromGroundCompensation)));
 	}
 
 	//----------------------------------------------------------------------//
@@ -1158,13 +1159,13 @@ public:
 
 	/** Finds the polygons along the navigation graph that touch the specified circle. */
 	bool FindPolysAroundCircle(const FVector& CenterPos, const NavNodeRef CenterNodeRef, const FVector::FReal Radius, const FSharedConstNavQueryFilter& Filter, const UObject* QueryOwner, TArray<NavNodeRef>* OutPolys = nullptr, TArray<NavNodeRef>* OutPolysParent = nullptr, TArray<float>* OutPolysCost = nullptr, int32* OutPolysCount = nullptr) const;
-
+	
 	/** Returns nearest navmesh polygon to Loc, or INVALID_NAVMESHREF if Loc is not on the navmesh. */
 	NavNodeRef FindNearestPoly(FVector const& Loc, FVector const& Extent, FSharedConstNavQueryFilter Filter = NULL, const UObject* Querier = NULL) const;
 
 	/** Finds the distance to the closest wall, limited to MaxDistance
 	 *	[out] OutClosestPointOnWall, if supplied, will be set to closest point on closest wall. Will not be set if no wall in the area (return value 0.f) */
-	float FindDistanceToWall(const FVector& StartLoc, FSharedConstNavQueryFilter Filter = nullptr, float MaxDistance = FLT_MAX, FVector* OutClosestPointOnWall = nullptr) const;
+	FVector::FReal FindDistanceToWall(const FVector& StartLoc, FSharedConstNavQueryFilter Filter = nullptr, FVector::FReal MaxDistance = TNumericLimits<FVector::FReal>::Max(), FVector* OutClosestPointOnWall = nullptr) const;
 
 	/** Retrieves center of the specified polygon. Returns false on error. */
 	bool GetPolyCenter(NavNodeRef PolyID, FVector& OutCenter) const;
@@ -1223,7 +1224,7 @@ public:
 	 *	@NOTE query is not using string-pulled path distance (for performance reasons),
 	 *		it measured distance between middles of portal edges, do you might want to 
 	 *		add an extra margin to PathingDistance */
-	bool GetPolysWithinPathingDistance(FVector const& StartLoc, const float PathingDistance, TArray<NavNodeRef>& FoundPolys,
+	bool GetPolysWithinPathingDistance(FVector const& StartLoc, const FVector::FReal PathingDistance, TArray<NavNodeRef>& FoundPolys,
 		FSharedConstNavQueryFilter Filter = nullptr, const UObject* Querier = nullptr, FRecastDebugPathfindingData* DebugData = nullptr) const;
 
 	/** Filters nav polys in PolyRefs with Filter */
diff --git a/Engine/Source/Runtime/NavigationSystem/Public/NavMesh/RecastQueryFilter.h b/Engine/Source/Runtime/NavigationSystem/Public/NavMesh/RecastQueryFilter.h
--- a/Engine/Source/Runtime/NavigationSystem/Public/NavMesh/RecastQueryFilter.h
+++ b/Engine/Source/Runtime/NavigationSystem/Public/NavMesh/RecastQueryFilter.h
@@ -14,7 +14,6 @@
 
 class UNavigationSystemV1;
 
-// LWC_TODO_AI: make a lot of the float params used in this class FReals. Not until after 5.0!
 class NAVIGATIONSYSTEM_API FRecastQueryFilter : public INavigationQueryFilterInterface, public dtQueryFilter
 {
 public:
diff --git a/Engine/Source/Runtime/NavigationSystem/Public/NavigationData.h b/Engine/Source/Runtime/NavigationSystem/Public/NavigationData.h
--- a/Engine/Source/Runtime/NavigationSystem/Public/NavigationData.h
+++ b/Engine/Source/Runtime/NavigationSystem/Public/NavigationData.h
@@ -24,8 +24,6 @@ class INavLinkCustomInterface;
 class UNavArea;
 class UPrimitiveComponent;
 
-// LWC_TODO_AI: A lot of the floats in this file should be FVector::FReal. Not until after 5.0!
-
 USTRUCT()
 struct NAVIGATIONSYSTEM_API FSupportedAreaData
 {
@@ -245,26 +243,26 @@ struct NAVIGATIONSYSTEM_API FNavigationPath : public TSharedFromThis<FNavigation
 	virtual bool ContainsNode(NavNodeRef NodeRef) const;
 
 	/** get cost of path, starting from given point */
-	virtual float GetCostFromIndex(int32 PathPointIndex) const
+	virtual FVector::FReal GetCostFromIndex(int32 PathPointIndex) const
 	{
-		return 0.0f;
+		return 0.;
 	}
 
 	/** get cost of path, starting from given node */
-	virtual float GetCostFromNode(NavNodeRef PathNode) const
+	virtual FVector::FReal GetCostFromNode(NavNodeRef PathNode) const
 	{
-		return 0.0f;
+		return 0.;
 	}
 
-	FORCEINLINE float GetCost() const
+	FORCEINLINE FVector::FReal GetCost() const
 	{
 		return GetCostFromIndex(0);
 	}
 
 	/** calculates total length of segments from NextPathPoint to the end of path, plus distance from CurrentPosition to NextPathPoint */
-	virtual float GetLengthFromPosition(FVector SegmentStart, uint32 NextPathPointIndex) const;
+	virtual FVector::FReal GetLengthFromPosition(FVector SegmentStart, uint32 NextPathPointIndex) const;
 
-	FORCEINLINE float GetLength() const
+	FORCEINLINE FVector::FReal GetLength() const
 	{
 		return PathPoints.Num() ? GetLengthFromPosition(PathPoints[0].Location, 1) : 0.0f;
 	}
@@ -863,20 +861,28 @@ public:
 	 *	@note function should assert if item's FNavigationProjectionWork.ProjectionLimit is invalid */
 	virtual void BatchProjectPoints(TArray<FNavigationProjectionWork>& Workload, FSharedConstNavQueryFilter Filter = NULL, const UObject* Querier = NULL) const PURE_VIRTUAL(ANavigationData::BatchProjectPoints, );
 
-	/** Calculates path from PathStart to PathEnd and retrieves its cost. 
-	 *	@NOTE this function does not generate string pulled path so the result is an (over-estimated) approximation
-	 *	@NOTE potentially expensive, so use it with caution */
-	virtual ENavigationQueryResult::Type CalcPathCost(const FVector& PathStart, const FVector& PathEnd, float& OutPathCost, FSharedConstNavQueryFilter QueryFilter = NULL, const UObject* Querier = NULL) const PURE_VIRTUAL(ANavigationData::CalcPathCost, return ENavigationQueryResult::Invalid;);
+	UE_DEPRECATED(5.2, "Use new version with FVector::FReal")
+	virtual ENavigationQueryResult::Type CalcPathCost(const FVector& PathStart, const FVector& PathEnd, float& OutPathCost, FSharedConstNavQueryFilter QueryFilter = NULL, const UObject* Querier = NULL) const final;
+
+	/** Calculates path from PathStart to PathEnd and retrieves its cost.
+ *	@NOTE this function does not generate string pulled path so the result is an (over-estimated) approximation
+ *	@NOTE potentially expensive, so use it with caution */
+	virtual ENavigationQueryResult::Type CalcPathCost(const FVector& PathStart, const FVector& PathEnd, FVector::FReal& OutPathCost, FSharedConstNavQueryFilter QueryFilter = NULL, const UObject* Querier = NULL) const PURE_VIRTUAL(ANavigationData::CalcPathCost, return ENavigationQueryResult::Invalid;);
+
+	UE_DEPRECATED(5.2, "Use new version with FVector::FReal")
+	virtual ENavigationQueryResult::Type CalcPathLength(const FVector& PathStart, const FVector& PathEnd, float& OutPathLength, FSharedConstNavQueryFilter QueryFilter = NULL, const UObject* Querier = NULL) const final;
 
 	/** Calculates path from PathStart to PathEnd and retrieves its length.
 	 *	@NOTE this function does not generate string pulled path so the result is an (over-estimated) approximation
 	 *	@NOTE potentially expensive, so use it with caution */
-	virtual ENavigationQueryResult::Type CalcPathLength(const FVector& PathStart, const FVector& PathEnd, float& OutPathLength, FSharedConstNavQueryFilter QueryFilter = NULL, const UObject* Querier = NULL) const PURE_VIRTUAL(ANavigationData::CalcPathLength, return ENavigationQueryResult::Invalid;);
+	virtual ENavigationQueryResult::Type CalcPathLength(const FVector& PathStart, const FVector& PathEnd, FVector::FReal& OutPathLength, FSharedConstNavQueryFilter QueryFilter = NULL, const UObject* Querier = NULL) const PURE_VIRTUAL(ANavigationData::CalcPathLength, return ENavigationQueryResult::Invalid;);
 
+	UE_DEPRECATED(5.2, "Use new version with FVector::FReal")
+	virtual ENavigationQueryResult::Type CalcPathLengthAndCost(const FVector& PathStart, const FVector& PathEnd, float& OutPathLength, float& OutPathCost, FSharedConstNavQueryFilter QueryFilter = NULL, const UObject* Querier = NULL) const final;
 	/** Calculates path from PathStart to PathEnd and retrieves its length.
 	 *	@NOTE this function does not generate string pulled path so the result is an (over-estimated) approximation
 	 *	@NOTE potentially expensive, so use it with caution */
-	virtual ENavigationQueryResult::Type CalcPathLengthAndCost(const FVector& PathStart, const FVector& PathEnd, float& OutPathLength, float& OutPathCost, FSharedConstNavQueryFilter QueryFilter = NULL, const UObject* Querier = NULL) const PURE_VIRTUAL(ANavigationData::CalcPathLength, return ENavigationQueryResult::Invalid;);
+	virtual ENavigationQueryResult::Type CalcPathLengthAndCost(const FVector& PathStart, const FVector& PathEnd, FVector::FReal& OutPathLength, FVector::FReal& OutPathCost, FSharedConstNavQueryFilter QueryFilter = NULL, const UObject* Querier = NULL) const PURE_VIRTUAL(ANavigationData::CalcPathLengthAndCost, return ENavigationQueryResult::Invalid;);
 
 	/** Checks if specified navigation node contains given location 
 	 *	@param Location is expressed in WorldSpace, navigation data is responsible for tansforming if need be */
@@ -1046,7 +1052,7 @@ struct FAsyncPathFindingQuery : public FPathFindingQuery
 		, Mode(EPathFindingMode::Regular)
 	{ }
 
-	FAsyncPathFindingQuery(const UObject* InOwner, const ANavigationData& InNavData, const FVector& Start, const FVector& End, const FNavPathQueryDelegate& Delegate, FSharedConstNavQueryFilter SourceQueryFilter, const float CostLimit = FLT_MAX);
+	FAsyncPathFindingQuery(const UObject* InOwner, const ANavigationData& InNavData, const FVector& Start, const FVector& End, const FNavPathQueryDelegate& Delegate, FSharedConstNavQueryFilter SourceQueryFilter, const FVector::FReal CostLimit = TNumericLimits<FVector::FReal>::Max());
 	FAsyncPathFindingQuery(const FPathFindingQuery& Query, const FNavPathQueryDelegate& Delegate, const EPathFindingMode::Type QueryMode);
 
 protected:
diff --git a/Engine/Source/Runtime/NavigationSystem/Public/NavigationPath.h b/Engine/Source/Runtime/NavigationSystem/Public/NavigationPath.h
--- a/Engine/Source/Runtime/NavigationSystem/Public/NavigationPath.h
+++ b/Engine/Source/Runtime/NavigationSystem/Public/NavigationPath.h
@@ -62,10 +62,10 @@ public:
 	void EnableRecalculationOnInvalidation(TEnumAsByte<ENavigationOptionFlag::Type> DoRecalculation);
 
 	UFUNCTION(BlueprintCallable, Category = "AI|Navigation")
-	float GetPathLength() const;
+	double GetPathLength() const;
 
 	UFUNCTION(BlueprintCallable, Category = "AI|Navigation")
-	float GetPathCost() const;
+	double GetPathCost() const;
 
 	UFUNCTION(BlueprintCallable, Category = "AI|Navigation")
 	bool IsPartial() const;
diff --git a/Engine/Source/Runtime/NavigationSystem/Public/NavigationSystem.h b/Engine/Source/Runtime/NavigationSystem/Public/NavigationSystem.h
--- a/Engine/Source/Runtime/NavigationSystem/Public/NavigationSystem.h
+++ b/Engine/Source/Runtime/NavigationSystem/Public/NavigationSystem.h
@@ -423,13 +423,21 @@ public:
 	static bool K2_GetRandomLocationInNavigableRadius(UObject* WorldContextObject, const FVector& Origin, FVector& RandomLocation, float Radius, ANavigationData* NavData = NULL, TSubclassOf<UNavigationQueryFilter> FilterClass = NULL);
 	
 	/** Potentially expensive. Use with caution. Consider using UPathFollowingComponent::GetRemainingPathCost instead */
-	UFUNCTION(BlueprintPure, Category="AI|Navigation", meta=(WorldContext="WorldContextObject" ) )
+	UE_DEPRECATED(5.2, "Use new version with double")
 	static ENavigationQueryResult::Type GetPathCost(UObject* WorldContextObject, const FVector& PathStart, const FVector& PathEnd, float& PathCost, ANavigationData* NavData = NULL, TSubclassOf<UNavigationQueryFilter> FilterClass = NULL);
 
+	/** Potentially expensive. Use with caution. Consider using UPathFollowingComponent::GetRemainingPathCost instead */
+	UFUNCTION(BlueprintPure, Category = "AI|Navigation", meta = (WorldContext = "WorldContextObject"))
+	static ENavigationQueryResult::Type GetPathCost(UObject* WorldContextObject, const FVector& PathStart, const FVector& PathEnd, double& PathCost, ANavigationData* NavData = NULL, TSubclassOf<UNavigationQueryFilter> FilterClass = NULL);
+
 	/** Potentially expensive. Use with caution */
-	UFUNCTION(BlueprintPure, Category="AI|Navigation", meta=(WorldContext="WorldContextObject" ) )
+	UE_DEPRECATED(5.2, "Use new version with double")
 	static ENavigationQueryResult::Type GetPathLength(UObject* WorldContextObject, const FVector& PathStart, const FVector& PathEnd, float& PathLength, ANavigationData* NavData = NULL, TSubclassOf<UNavigationQueryFilter> FilterClass = NULL);
 
+	/** Potentially expensive. Use with caution */
+	UFUNCTION(BlueprintPure, Category = "AI|Navigation", meta = (WorldContext = "WorldContextObject"))
+	static ENavigationQueryResult::Type GetPathLength(UObject* WorldContextObject, const FVector& PathStart, const FVector& PathEnd, double& PathLength, ANavigationData* NavData = NULL, TSubclassOf<UNavigationQueryFilter> FilterClass = NULL);
+
 	UFUNCTION(BlueprintPure, Category="AI|Navigation", meta=(WorldContext="WorldContextObject" ) )
 	static bool IsNavigationBeingBuilt(UObject* WorldContextObject);
 
@@ -599,16 +607,31 @@ public:
 	
 	/** Calculates a path from PathStart to PathEnd and retrieves its cost. 
 	 *	@NOTE potentially expensive, so use it with caution */
+	UE_DEPRECATED(5.2, "Use new version with FVector::FReal")
 	ENavigationQueryResult::Type GetPathCost(const FVector& PathStart, const FVector& PathEnd, float& PathCost, const ANavigationData* NavData = NULL, FSharedConstNavQueryFilter QueryFilter = NULL) const;
 
-	/** Calculates a path from PathStart to PathEnd and retrieves its overestimated length.
+	/** Calculates a path from PathStart to PathEnd and retrieves its cost.
 	 *	@NOTE potentially expensive, so use it with caution */
+	ENavigationQueryResult::Type GetPathCost(const FVector& PathStart, const FVector& PathEnd, FVector::FReal& PathCost, const ANavigationData* NavData = NULL, FSharedConstNavQueryFilter QueryFilter = NULL) const;
+
+	/** Calculates a path from PathStart to PathEnd and retrieves its overestimated length.
+	*	@NOTE potentially expensive, so use it with caution */
+	UE_DEPRECATED(5.2, "Use new version with FVector::FReal")
 	ENavigationQueryResult::Type GetPathLength(const FVector& PathStart, const FVector& PathEnd, float& PathLength, const ANavigationData* NavData = NULL, FSharedConstNavQueryFilter QueryFilter = NULL) const;
 
-	/** Calculates a path from PathStart to PathEnd and retrieves its overestimated length and cost.
+	/** Calculates a path from PathStart to PathEnd and retrieves its overestimated length.
 	 *	@NOTE potentially expensive, so use it with caution */
+	ENavigationQueryResult::Type GetPathLength(const FVector& PathStart, const FVector& PathEnd, FVector::FReal& PathLength, const ANavigationData* NavData = NULL, FSharedConstNavQueryFilter QueryFilter = NULL) const;
+
+	/** Calculates a path from PathStart to PathEnd and retrieves its overestimated length and cost.
+	*	@NOTE potentially expensive, so use it with caution */
+	UE_DEPRECATED(5.2, "Use new version with FVector::FReal")
 	ENavigationQueryResult::Type GetPathLengthAndCost(const FVector& PathStart, const FVector& PathEnd, float& PathLength, float& PathCost, const ANavigationData* NavData = NULL, FSharedConstNavQueryFilter QueryFilter = NULL) const;
 
+	/** Calculates a path from PathStart to PathEnd and retrieves its overestimated length and cost.
+	 *	@NOTE potentially expensive, so use it with caution */
+	ENavigationQueryResult::Type GetPathLengthAndCost(const FVector& PathStart, const FVector& PathEnd, FVector::FReal& PathLength, FVector::FReal& PathCost, const ANavigationData* NavData = NULL, FSharedConstNavQueryFilter QueryFilter = NULL) const;
+
 	// @todo document
 	bool ProjectPointToNavigation(const FVector& Point, FNavLocation& OutLocation, const FVector& Extent = INVALID_NAVEXTENT, const FNavAgentProperties* AgentProperties = NULL, FSharedConstNavQueryFilter QueryFilter = NULL)
 	{
diff --git a/Engine/Source/Runtime/NavigationSystem/Public/NavigationSystemTypes.h b/Engine/Source/Runtime/NavigationSystem/Public/NavigationSystemTypes.h
--- a/Engine/Source/Runtime/NavigationSystem/Public/NavigationSystemTypes.h
+++ b/Engine/Source/Runtime/NavigationSystem/Public/NavigationSystemTypes.h
@@ -20,8 +20,6 @@ class FNavigationOctree;
 class UNavigationPath;
 class ANavigationData;
 
-// LWC_TODO_AI: A lot of the floats in this file should be FVector::FReal. Not until after 5.0!
-
 struct NAVIGATIONSYSTEM_API FPathFindingQueryData
 {
 	TWeakObjectPtr<const UObject> Owner;
@@ -30,7 +28,7 @@ struct NAVIGATIONSYSTEM_API FPathFindingQueryData
 	FSharedConstNavQueryFilter QueryFilter;
 
 	/** cost limit of nodes allowed to be added to the open list */
-	float CostLimit;
+	FVector::FReal CostLimit;
 	
 	/** additional flags passed to navigation data handling request */
 	int32 NavDataFlags;
@@ -40,7 +38,7 @@ struct NAVIGATIONSYSTEM_API FPathFindingQueryData
 
 	FPathFindingQueryData() : StartLocation(FNavigationSystem::InvalidLocation), EndLocation(FNavigationSystem::InvalidLocation), CostLimit(FLT_MAX), NavDataFlags(0), bAllowPartialPaths(true) {}
 
-	FPathFindingQueryData(const UObject* InOwner, const FVector& InStartLocation, const FVector& InEndLocation, FSharedConstNavQueryFilter InQueryFilter = nullptr, int32 InNavDataFlags = 0, bool bInAllowPartialPaths = true, const float InCostLimit = FLT_MAX) :
+	FPathFindingQueryData(const UObject* InOwner, const FVector& InStartLocation, const FVector& InEndLocation, FSharedConstNavQueryFilter InQueryFilter = nullptr, int32 InNavDataFlags = 0, bool bInAllowPartialPaths = true, const FVector::FReal InCostLimit = TNumericLimits<FVector::FReal>::Max()) :
 		Owner(InOwner), StartLocation(InStartLocation), EndLocation(InEndLocation), QueryFilter(InQueryFilter), CostLimit(InCostLimit), NavDataFlags(InNavDataFlags), bAllowPartialPaths(bInAllowPartialPaths) {}
 };
 
@@ -52,8 +50,8 @@ struct NAVIGATIONSYSTEM_API FPathFindingQuery : public FPathFindingQueryData
 
 	FPathFindingQuery() : FPathFindingQueryData() {}
 	FPathFindingQuery(const FPathFindingQuery& Source);
-	FPathFindingQuery(const UObject* InOwner, const ANavigationData& InNavData, const FVector& Start, const FVector& End, FSharedConstNavQueryFilter SourceQueryFilter = NULL, FNavPathSharedPtr InPathInstanceToFill = NULL, const float CostLimit = FLT_MAX);
-	FPathFindingQuery(const INavAgentInterface& InNavAgent, const ANavigationData& InNavData, const FVector& Start, const FVector& End, FSharedConstNavQueryFilter SourceQueryFilter = NULL, FNavPathSharedPtr InPathInstanceToFill = NULL, const float CostLimit = FLT_MAX);
+	FPathFindingQuery(const UObject* InOwner, const ANavigationData& InNavData, const FVector& Start, const FVector& End, FSharedConstNavQueryFilter SourceQueryFilter = NULL, FNavPathSharedPtr InPathInstanceToFill = NULL, const FVector::FReal CostLimit = TNumericLimits<FVector::FReal>::Max());
+	FPathFindingQuery(const INavAgentInterface& InNavAgent, const ANavigationData& InNavData, const FVector& Start, const FVector& End, FSharedConstNavQueryFilter SourceQueryFilter = NULL, FNavPathSharedPtr InPathInstanceToFill = NULL, const FVector::FReal CostLimit = TNumericLimits<FVector::FReal>::Max());
 
 	explicit FPathFindingQuery(FNavPathSharedRef PathToRecalculate, const ANavigationData* NavDataOverride = NULL);
 
@@ -64,7 +62,7 @@ struct NAVIGATIONSYSTEM_API FPathFindingQuery : public FPathFindingQueryData
 	/** utility function to compute a cost limit using an Euclidean heuristic, an heuristic scale and a cost limit factor
 	*	CostLimitFactor: multiplier used to compute the cost limit value from the initial heuristic
 	*	MinimumCostLimit: minimum clamping value used to prevent low cost limit for short path query */
-	float ComputeCostLimitFromHeuristic(const FVector& StartPos, const FVector& EndPos, const float HeuristicScale, const float CostLimitFactor, const float MinimumCostLimit) const;
+	FVector::FReal ComputeCostLimitFromHeuristic(const FVector& StartPos, const FVector& EndPos, const FVector::FReal HeuristicScale, const FVector::FReal CostLimitFactor, const FVector::FReal MinimumCostLimit) const;
 };
 
 namespace EPathFindingMode
diff --git a/Engine/Source/Runtime/Navmesh/Private/Detour/DetourNavMeshQuery.cpp b/Engine/Source/Runtime/Navmesh/Private/Detour/DetourNavMeshQuery.cpp
--- a/Engine/Source/Runtime/Navmesh/Private/Detour/DetourNavMeshQuery.cpp
+++ b/Engine/Source/Runtime/Navmesh/Private/Detour/DetourNavMeshQuery.cpp
@@ -307,7 +307,7 @@ void dtNavMeshQuery::updateLinkFilter(dtQuerySpecialLinkFilter* linkFilter)
 	}
 }
 
-// LWC_TODO_AI: Should be double(*frand)() to be consistent with the rest of the API
+// LWC_TODO_AI: Should be double(*frand)() to be consistent with the rest of the API, UE currently does not have a double rand function but it is planned!
 dtStatus dtNavMeshQuery::findRandomPoint(const dtQueryFilter* filter, float(*frand)(),
 										 dtPolyRef* randomRef, dtReal* randomPt) const
 {
diff --git a/Engine/Source/Runtime/Navmesh/Private/DetourCrowd/DetourCrowd.cpp b/Engine/Source/Runtime/Navmesh/Private/DetourCrowd/DetourCrowd.cpp
--- a/Engine/Source/Runtime/Navmesh/Private/DetourCrowd/DetourCrowd.cpp
+++ b/Engine/Source/Runtime/Navmesh/Private/DetourCrowd/DetourCrowd.cpp
@@ -854,8 +854,7 @@ void dtCrowd::updateMoveRequest(const dtReal /*dt*/)
 			// Quick seach towards the goal.
 			static const int MAX_ITER = 20;
 			m_navquery->updateLinkFilter(ag->params.linkFilter.Get());
-			// LWC_TODO_AI: This should be DT_REAL_MAX but leaving as FLT_MAX until after 5.0 as UE side has not been converted to taking FReals for Costs.
-			const dtReal costLimit = FLT_MAX; //@UE
+			const dtReal costLimit = DT_REAL_MAX; //@UE
 			m_navquery->initSlicedFindPath(path[0], ag->targetRef, ag->npos, ag->targetPos, costLimit, &m_filters[ag->params.filter]); //@UE
 			m_navquery->updateSlicedFindPath(MAX_ITER, 0);
 			dtStatus status = 0;
@@ -922,8 +921,7 @@ void dtCrowd::updateMoveRequest(const dtReal /*dt*/)
 	for (int i = 0; i < nqueue; ++i)
 	{
 		dtCrowdAgent* ag = queue[i];
-		// LWC_TODO_AI: This should be DT_REAL_MAX but leaving as FLT_MAX until after 5.0 as UE side has not been converted to taking FReals for Costs.
-		const dtReal costLimit = FLT_MAX; //@UE
+		const dtReal costLimit = DT_REAL_MAX; //@UE
 		ag->targetPathqRef = m_pathq.request(ag->corridor.getLastPoly(), ag->targetRef,
 			ag->corridor.getTarget(), ag->targetPos, costLimit, &m_filters[ag->params.filter], ag->params.linkFilter); //@UE
 		if (ag->targetPathqRef != DT_PATHQ_INVALID)
diff --git a/Engine/Source/Runtime/Navmesh/Private/DetourCrowd/DetourLocalBoundary.cpp b/Engine/Source/Runtime/Navmesh/Private/DetourCrowd/DetourLocalBoundary.cpp
--- a/Engine/Source/Runtime/Navmesh/Private/DetourCrowd/DetourLocalBoundary.cpp
+++ b/Engine/Source/Runtime/Navmesh/Private/DetourCrowd/DetourLocalBoundary.cpp
@@ -28,8 +28,7 @@ dtLocalBoundary::dtLocalBoundary() :
 	m_nsegs(0),
 	m_npolys(0)
 {
-	// LWC_TODO_AI: This should be DT_REAL_MAX but leaving as FLT_MAX until after 5.0 to ease licensee LWCoord integration pain for 5.0!
-	dtVset(m_center, FLT_MAX, FLT_MAX, FLT_MAX);
+	dtVset(m_center, DT_REAL_MAX, DT_REAL_MAX, DT_REAL_MAX);
 }
 
 dtLocalBoundary::~dtLocalBoundary()
@@ -38,8 +37,7 @@ dtLocalBoundary::~dtLocalBoundary()
 
 void dtLocalBoundary::reset()
 {
-	// LWC_TODO_AI: This should be DT_REAL_MAX but leaving as FLT_MAX until after 5.0 to ease licensee LWCoord integration pain for 5.0!
-	dtVset(m_center, FLT_MAX, FLT_MAX, FLT_MAX);
+	dtVset(m_center, DT_REAL_MAX, DT_REAL_MAX, DT_REAL_MAX);
 	m_npolys = 0;
 	m_nsegs = 0;
 }
@@ -118,8 +116,7 @@ void dtLocalBoundary::update(dtPolyRef ref, const dtReal* pos, const dtReal coll
 	
 	if (!ref)
 	{
-		// LWC_TODO_AI: This should be DT_REAL_MAX but leaving as FLT_MAX until after 5.0 to ease licensee LWCoord integration pain for 5.0!!
-		dtVset(m_center, FLT_MAX, FLT_MAX, FLT_MAX);
+		dtVset(m_center, DT_REAL_MAX, DT_REAL_MAX, DT_REAL_MAX);
 		m_nsegs = 0;
 		m_npolys = 0;
 		return;
diff --git a/Engine/Source/Runtime/Navmesh/Private/DetourCrowd/DetourPathCorridor.cpp b/Engine/Source/Runtime/Navmesh/Private/DetourCrowd/DetourPathCorridor.cpp
--- a/Engine/Source/Runtime/Navmesh/Private/DetourCrowd/DetourPathCorridor.cpp
+++ b/Engine/Source/Runtime/Navmesh/Private/DetourCrowd/DetourPathCorridor.cpp
@@ -524,8 +524,7 @@ bool dtPathCorridor::optimizePathTopology(dtNavMeshQuery* navquery, const dtQuer
 	
 	dtPolyRef res[MAX_RES];
 	int nres = 0;
-	// LWC_TODO_AI: This should be DT_REAL_MAX but leaving as FLT_MAX until after 5.0, as UE side has not been converted to taking FReals for Costs.
-	const dtReal costLimit = FLT_MAX; //@UE
+	const dtReal costLimit = DT_REAL_MAX; //@UE
 	navquery->initSlicedFindPath(m_path[0], m_path[m_npath-1], m_pos, m_target, costLimit, filter); //@UE
 	navquery->updateSlicedFindPath(MAX_ITER, 0);
 	dtStatus status = navquery->finalizeSlicedFindPathPartial(m_path, m_npath, res, &nres, MAX_RES);
diff --git a/Engine/Source/Runtime/Navmesh/Public/DebugUtils/DebugDraw.h b/Engine/Source/Runtime/Navmesh/Public/DebugUtils/DebugDraw.h
--- a/Engine/Source/Runtime/Navmesh/Public/DebugUtils/DebugDraw.h
+++ b/Engine/Source/Runtime/Navmesh/Public/DebugUtils/DebugDraw.h
@@ -25,8 +25,6 @@
 #include "CoreMinimal.h"
 #include "DebugDrawLargeWorldCoordinates.h"
 
-//@UE BEGIN Adding support for LWCoords.
-// LWC_TODO_AI: Remove the float version of these functions when we go to doubles.Before 5.0!
 // Some math headers don't have PI defined.
 static const duReal DU_PI = 3.14159265358979323846;
 
@@ -69,7 +67,6 @@ inline double duSqrt(double x)
 {
 	return sqrt(x);
 }
-//@UE END Adding support for LWCoords.
 
 enum duDebugDrawPrimitives
 {
diff --git a/Engine/Source/Runtime/Navmesh/Public/Detour/DetourCommon.h b/Engine/Source/Runtime/Navmesh/Public/Detour/DetourCommon.h
--- a/Engine/Source/Runtime/Navmesh/Public/Detour/DetourCommon.h
+++ b/Engine/Source/Runtime/Navmesh/Public/Detour/DetourCommon.h
@@ -85,8 +85,6 @@ template<class T> inline T dtClamp(T v, T mn, T mx) { return v < mn ? mn : (v >
 inline dtReal dtClamp(dtReal v, dtReal mn, dtReal mx) { return dtClamp<dtReal>(v, mn, mx); }
 //@UE END Adding support for LWCoords.
 
-//@UE BEGIN Adding support for LWCoords.
-// LWC_TODO_AI: Remove the float version of these functions when we go to doubles. Remove prior to 5.0!
 inline float dtFloor(float x)
 {
 	return floorf(x);
@@ -156,7 +154,6 @@ inline double dtfMod(double x, double y)
 {
 	return fmod(x, y);
 }
-//@UE END Adding support for LWCoords.
 
 /// @}
 /// @name Vector helper functions.
diff --git a/Engine/Source/Runtime/Navmesh/Public/Detour/DetourNavMeshQuery.h b/Engine/Source/Runtime/Navmesh/Public/Detour/DetourNavMeshQuery.h
--- a/Engine/Source/Runtime/Navmesh/Public/Detour/DetourNavMeshQuery.h
+++ b/Engine/Source/Runtime/Navmesh/Public/Detour/DetourNavMeshQuery.h
@@ -33,7 +33,7 @@
 //@UE BEGIN
 #define WITH_FIXED_AREA_ENTERING_COST 1
 
-// LWC_TODO_AI: This should be DT_REAL_MAX but leaving as FLT_MAX until after 5.0, as UE side has not been converted to taking FReals for Costs.
+// Note poly costs are still floats in UE so we are using FLT_MAX as unwalkable still.
 #define DT_UNWALKABLE_POLY_COST FLT_MAX
 //@UE END
 
