Fix for crash while loading navmesh with async static mesh compilation enabled.

AddElementToNavOctree (through some of its resulting function calls) modifies PendingOctreeUpdates so invalidates the iterators, (via WaitUntilAsyncPropertyReleased(), UpdateComponentInNavOctree() ,  RegisterNavOctreeElement()).

This means we can't iterate through the TSet PendingOctreeUpdates in the normal way. Previously the code iterated through this which also left us open to other potential bugs in that we may have tried to modify elements we had already processed.

diff --git a/Engine/Source/Runtime/NavigationSystem/Private/NavigationDataHandler.cpp b/Engine/Source/Runtime/NavigationSystem/Private/NavigationDataHandler.cpp
--- a/Engine/Source/Runtime/NavigationSystem/Private/NavigationDataHandler.cpp
+++ b/Engine/Source/Runtime/NavigationSystem/Private/NavigationDataHandler.cpp
@@ -504,11 +504,19 @@ void FNavigationDataHandler::UpdateActorAndComponentsInNavOctree(AActor& Actor)
 
 void FNavigationDataHandler::ProcessPendingOctreeUpdates()
 {
-	if (OctreeController.PendingOctreeUpdates.Num() && OctreeController.NavOctree)
+	QUICK_SCOPE_CYCLE_COUNTER(STAT_Navigation_ProcessPendingOctreeUpdates);
+
+	if (OctreeController.NavOctree)
 	{
-		for (TSet<FNavigationDirtyElement>::TIterator It(OctreeController.PendingOctreeUpdates); It; ++It)
+		// AddElementToNavOctree (through some of its resulting function calls) modifies PendingOctreeUpdates so invalidates the iterators,
+		// (via WaitUntilAsyncPropertyReleased() / UpdateComponentInNavOctree() / RegisterNavOctreeElement()). This means we can't iterate
+		// through this set in the normal way. Previously the code iterated through this which also left us open to other potential bugs
+		// in that we may have tried to modify elements we had already processed.
+		while (TSet<FNavigationDirtyElement>::TIterator It = OctreeController.PendingOctreeUpdates.CreateIterator())
 		{
-			AddElementToNavOctree(*It);
+			FNavigationDirtyElement Element = *It;
+			It.RemoveCurrent();
+			AddElementToNavOctree(Element);
 		}
 	}
 	OctreeController.PendingOctreeUpdates.Empty(32);
