Fix assert while getting rid of old navigation data during serialize on async loading thread

diff --git a/Engine/Source/Runtime/NavigationSystem/Private/NavigationData.cpp b/Engine/Source/Runtime/NavigationSystem/Private/NavigationData.cpp
--- a/Engine/Source/Runtime/NavigationSystem/Private/NavigationData.cpp
+++ b/Engine/Source/Runtime/NavigationSystem/Private/NavigationData.cpp
@@ -513,7 +513,7 @@ void ANavigationData::CleanUpAndMarkPendingKill()
 
 	/* Need to check if the world is valid since, when this is called from Serialize, the World won't be set and Destroy will do nothing, 
 	 * in which case it will crash when it tries to register with the NavSystem in UNavigationSystemV1::ProcessRegistrationCandidates. */
-	if (CVarDestroyNavDataInCleanUpAndMarkPendingKill.GetValueOnGameThread() && IsValid(GetWorld()))
+	if (CVarDestroyNavDataInCleanUpAndMarkPendingKill.GetValueOnAnyThread() && IsValid(GetWorld()))
 	{
 		Destroy();
 	}
@@ -523,6 +523,8 @@ void ANavigationData::CleanUpAndMarkPendingKill()
 
 		if (UWorld* World = GetWorld())
 		{
+			// This part is not thread-safe and should only happen on the GT...
+			check(IsInGameThread());
 			World->RemoveNetworkActor(this);
 		}
 		MarkAsGarbage();
