[WP Navmesh] Prevent creation of partitioned navmesh in non world-partitioned worlds

diff --git a/Engine/Source/Runtime/NavigationSystem/Private/NavMesh/RecastNavMesh.cpp b/Engine/Source/Runtime/NavigationSystem/Private/NavMesh/RecastNavMesh.cpp
--- a/Engine/Source/Runtime/NavigationSystem/Private/NavMesh/RecastNavMesh.cpp
+++ b/Engine/Source/Runtime/NavigationSystem/Private/NavMesh/RecastNavMesh.cpp
@@ -485,6 +485,9 @@ ARecastNavMesh::ARecastNavMesh(const FObjectInitializer& ObjectInitializer)
 	, bUseVirtualFilters(true)
 	, bUseVirtualGeometryFilteringAndDirtying(false)
 	, bAllowNavLinkAsPathEnd(false)
+#if WITH_EDITORONLY_DATA
+	, bAllowWorldPartitionedNavMesh(false)
+#endif // WITH_EDITORONLY_DATA
 	, TileSetUpdateInterval(1.0f)
 	, NavMeshVersion(NAVMESHVER_LATEST)	
 	, RecastNavMeshImpl(NULL)
@@ -618,7 +621,12 @@ void ARecastNavMesh::PostLoad()
 
 	UE_CLOG(TileSizeUU < CellSize, LogNavigation, Error, TEXT("%s: TileSizeUU (%f) being less than CellSize (%f) is an invalid case and will cause navmesh generation issues.")
 		, *GetName(), TileSizeUU, CellSize);
-
+	
+	if (!UWorld::IsPartitionedWorld(GetWorld()))
+	{
+		bIsWorldPartitioned = false;
+	}
+	
 	RecreateDefaultFilter();
 	UpdatePolyRefBitsPreview();
 }
@@ -725,6 +733,10 @@ void ARecastNavMesh::PostInitProperties()
 			AgentMaxStepHeight = DefOb->AgentMaxStepHeight;
 		}
 	}
+	
+#if WITH_EDITORONLY_DATA
+	bAllowWorldPartitionedNavMesh = UWorld::IsPartitionedWorld(GetWorld());
+#endif // WITH_EDITORONLY_DATA
 
 	UpdatePolyRefBitsPreview();
 }
diff --git a/Engine/Source/Runtime/NavigationSystem/Public/NavMesh/RecastNavMesh.h b/Engine/Source/Runtime/NavigationSystem/Public/NavMesh/RecastNavMesh.h
--- a/Engine/Source/Runtime/NavigationSystem/Public/NavMesh/RecastNavMesh.h
+++ b/Engine/Source/Runtime/NavigationSystem/Public/NavMesh/RecastNavMesh.h
@@ -766,7 +766,7 @@ class NAVIGATIONSYSTEM_API ARecastNavMesh : public ANavigationData
 	uint32 bSortNavigationAreasByCost:1;
 
 	/* In a world partitioned map, is this navmesh using world partitioning */
-	UPROPERTY(EditAnywhere, Category=Generation, config)
+	UPROPERTY(EditAnywhere, Category=Generation, config, meta = (EditCondition = "bAllowWorldPartitionedNavMesh", HideEditConditionToggle, DisplayName = "IsWorldPartitionedNavMesh"))
 	uint32 bIsWorldPartitioned : 1;
 	
 	/** controls whether voxel filtering will be applied (via FRecastTileGenerator::ApplyVoxelFilter). 
@@ -825,6 +825,13 @@ class NAVIGATIONSYSTEM_API ARecastNavMesh : public ANavigationData
 	UPROPERTY(EditAnywhere, Category = TimeSlicing, config, AdvancedDisplay)
 	double TimeSliceLongDurationDebug = 0.002;
 
+protected:
+#if WITH_EDITORONLY_DATA
+	/** World partitioned navmesh are only allowed in partitioned worlds. */
+	UPROPERTY() 
+	uint32 bAllowWorldPartitionedNavMesh : 1;
+#endif // WITH_EDITORONLY_DATA
+	
 private:
 	/** Cache rasterized voxels instead of just collision vertices/indices in navigation octree */
 	UPROPERTY(config)
