[Navigation] Don't display "navmesh needs to be rebuilt" on-screen editor message in partitioned world.
It's not meaningful since loading and unloading parts of the world triggers it.

diff --git a/Engine/Source/Runtime/NavigationSystem/Private/NavigationSystem.cpp b/Engine/Source/Runtime/NavigationSystem/Private/NavigationSystem.cpp
--- a/Engine/Source/Runtime/NavigationSystem/Private/NavigationSystem.cpp
+++ b/Engine/Source/Runtime/NavigationSystem/Private/NavigationSystem.cpp
@@ -4835,9 +4835,12 @@ void UNavigationSystemV1::GetOnScreenMessages(TMultiMap<FCoreDelegates::EOnScree
 #else
 	const bool bIsNavigationAutoUpdateEnabled = true;
 #endif
-	if (IsNavigationDirty()
-		&& ((FNavigationSystem::IsEditorRunMode(OperationMode) && !bIsNavigationAutoUpdateEnabled)
-			|| !SupportsNavigationGeneration() || !CanRebuildDirtyNavigation()))
+
+	// Don't display "navmesh needs to be rebuilt" on-screen editor message in partitioned world. 
+	// It's not meaningful since loading and unloading parts of the world triggers it.
+	if (!UWorld::IsPartitionedWorld(GetWorld())
+		&& IsNavigationDirty()
+		&& ((FNavigationSystem::IsEditorRunMode(OperationMode) && !bIsNavigationAutoUpdateEnabled) || !SupportsNavigationGeneration() || !CanRebuildDirtyNavigation()))
 	{
 		OutMessages.Add(FCoreDelegates::EOnScreenMessageSeverity::Error
 			, LOCTEXT("NAVMESHERROR", "NAVMESH NEEDS TO BE REBUILT"));
