Remove commented out code only

diff --git a/Engine/Source/Runtime/NavigationSystem/Private/NavLinkCustomComponent.cpp b/Engine/Source/Runtime/NavigationSystem/Private/NavLinkCustomComponent.cpp
--- a/Engine/Source/Runtime/NavigationSystem/Private/NavLinkCustomComponent.cpp
+++ b/Engine/Source/Runtime/NavigationSystem/Private/NavLinkCustomComponent.cpp
@@ -372,11 +372,6 @@ void UNavLinkCustomComponent::BroadcastStateChange()
 	CollectNearbyAgents(NearbyAgents);
 	OnBroadcastFilter.ExecuteIfBound(this, NearbyAgents);
 
-	for (int32 i = 0; i < NearbyAgents.Num(); i++)
-	{
-//		NearbyAgents[i]->OnCustomLinkBroadcast(this);
-	}
-
 	if (BroadcastInterval > 0.0f)
 	{
 		GetWorld()->GetTimerManager().SetTimer(TimerHandle_BroadcastStateChange, this, &UNavLinkCustomComponent::BroadcastStateChange, BroadcastInterval);
