[WP Navmesh] For navmeshes supporting runtime generation, copy cache data when filling navigation data chunk actors.

diff --git a/Engine/Source/Runtime/NavigationSystem/Private/NavMesh/RecastNavMesh.cpp b/Engine/Source/Runtime/NavigationSystem/Private/NavMesh/RecastNavMesh.cpp
--- a/Engine/Source/Runtime/NavigationSystem/Private/NavMesh/RecastNavMesh.cpp
+++ b/Engine/Source/Runtime/NavigationSystem/Private/NavMesh/RecastNavMesh.cpp
@@ -2500,7 +2500,7 @@ void ARecastNavMesh::FillNavigationDataChunkActor(const FBox& QueryBounds, ANavi
 			DataChunk->NavigationDataName = GetFName();
 			DataChunkActor.GetMutableNavDataChunk().Add(DataChunk);
 
-			DataChunk->GetTiles(RecastNavMeshImpl, TileIndices, EGatherTilesCopyMode::CopyData);
+			DataChunk->GetTiles(RecastNavMeshImpl, TileIndices, SupportsRuntimeGeneration() ? EGatherTilesCopyMode::CopyDataAndCacheData : EGatherTilesCopyMode::CopyData);
 			DataChunk->GetTilesBounds(*RecastNavMeshImpl, TileIndices, OutTilesBounds);
 		}
 	}
