Fix for moving Nav Modifier Volume does not cause the navmesh to update.

Issue was IsActorInitialized() never returns true in the editor world.

diff --git a/Engine/Source/Runtime/Engine/Classes/GameFramework/Actor.h b/Engine/Source/Runtime/Engine/Classes/GameFramework/Actor.h
--- a/Engine/Source/Runtime/Engine/Classes/GameFramework/Actor.h
+++ b/Engine/Source/Runtime/Engine/Classes/GameFramework/Actor.h
@@ -459,6 +459,9 @@ private:
 	/** True if this actor is currently running user construction script (used to defer component registration) */
 	uint8 bRunningUserConstructionScript:1;
 
+	/** Set when RegisterAllComponents() is called and unset when UnregisterAllComponents() is called */
+	uint8 bHasRegisteredAllComponents:1;
+
 	/**
 	 * Enables any collision on this actor.
 	 * @see SetActorEnableCollision(), GetActorEnableCollision()
@@ -1832,6 +1835,9 @@ public:
 		return bActorIsBeingDestroyed;
 	}
 
+	/** Returns true if RegisterAllComponents() has been called without a subsequent call to UnregisterAllComponents() and false otherwise */
+	bool HasActorRegisteredAllComponents() const { return bHasRegisteredAllComponents; }
+
 	/** Event when this actor takes ANY damage */
 	UFUNCTION(BlueprintImplementableEvent, BlueprintAuthorityOnly, meta=(DisplayName = "AnyDamage"), Category="Game|Damage")
 	void ReceiveAnyDamage(float Damage, const class UDamageType* DamageType, class AController* InstigatedBy, AActor* DamageCauser);
diff --git a/Engine/Source/Runtime/Engine/Private/Actor.cpp b/Engine/Source/Runtime/Engine/Private/Actor.cpp
--- a/Engine/Source/Runtime/Engine/Private/Actor.cpp
+++ b/Engine/Source/Runtime/Engine/Private/Actor.cpp
@@ -178,6 +178,8 @@ void AActor::InitializeDefaults()
 	DefaultUpdateOverlapsMethodDuringLevelStreaming = EActorUpdateOverlapsMethod::OnlyUpdateMovable;
 	
 	bHasDeferredComponentRegistration = false;
+	bHasRegisteredAllComponents = false;
+
 #if WITH_EDITORONLY_DATA
 	bIsInEditingLevelInstance = false;
 	PivotOffset = FVector::ZeroVector;
@@ -5177,6 +5179,8 @@ void AActor::UnregisterAllComponents(const bool bForReregister)
 		}
 	}
 
+	bHasRegisteredAllComponents = false;
+
 	PostUnregisterAllComponents();
 }
 
@@ -5197,6 +5201,7 @@ void AActor::RegisterAllComponents()
 
 	// Clear this flag as it's no longer deferred
 	bHasDeferredComponentRegistration = false;
+	bHasRegisteredAllComponents = true;
 }
 
 // Walks through components hierarchy and returns closest to root parent component that is unregistered
diff --git a/Engine/Source/Runtime/NavigationSystem/Private/NavigationSystem.cpp b/Engine/Source/Runtime/NavigationSystem/Private/NavigationSystem.cpp
--- a/Engine/Source/Runtime/NavigationSystem/Private/NavigationSystem.cpp
+++ b/Engine/Source/Runtime/NavigationSystem/Private/NavigationSystem.cpp
@@ -454,7 +454,10 @@ UNavigationSystemV1::UNavigationSystemV1(const FObjectInitializer& ObjectInitial
 					UNavigationSystemV1::UpdateNavOctreeBounds(ParentedActors[Idx]);
 				}
 
-				if (Actor.IsActorInitialized())
+				// We need to check this actor has registered all their components post spawn / load
+				// before attempting to update the components in the nav octree.
+				// Without this check we were getting an issue with UNavRelevantComponent::GetNavigationParent().
+				if (Actor.HasActorRegisteredAllComponents())
 				{
 					// not doing manual update of all attached actors since UpdateActorAndComponentsInNavOctree should take care of it
 					UNavigationSystemV1::UpdateActorAndComponentsInNavOctree(Actor);
@@ -4127,7 +4130,10 @@ void UNavigationSystemV1::OnActorMoved(AActor* Actor)
 	{
 		OnNavigationBoundsUpdated((ANavMeshBoundsVolume*)Actor);
 	}
-	else if (Actor && Actor->IsActorInitialized())
+	// We need to check this actor has registered all their components post spawn / load
+	// before attempting to update the components in the nav octree.
+	// Without this check we were getting an issue with UNavRelevantComponent::GetNavigationParent().
+	else if (Actor && Actor->HasActorRegisteredAllComponents())
 	{
 		UpdateActorAndComponentsInNavOctree(*Actor, /*bUpdateAttachedActors=*/true);
 	}
