[Navmesh] In FPImplRecastNavMesh::GetNavMeshTilesIn(), use FRcTileBox to find min and max tile coordinates.
Note: At CL 21589037, FRcTileBox has been fixed to avoid including tiles that are just touching bounds borders so this applies the fix to GetNavMeshTilesIn() as well.

diff --git a/Engine/Source/Runtime/NavigationSystem/Private/NavMesh/PImplRecastNavMesh.cpp b/Engine/Source/Runtime/NavigationSystem/Private/NavMesh/PImplRecastNavMesh.cpp
--- a/Engine/Source/Runtime/NavigationSystem/Private/NavMesh/PImplRecastNavMesh.cpp
+++ b/Engine/Source/Runtime/NavigationSystem/Private/NavMesh/PImplRecastNavMesh.cpp
@@ -3078,15 +3078,12 @@ void FPImplRecastNavMesh::GetNavMeshTilesIn(const TArray<FBox>& InclusionBounds,
 		TSet<FIntPoint>	TileCoords;	
 		for (const FBox& Bounds : InclusionBounds)
 		{
-			const FBox RcBounds = Unreal2RecastBox(Bounds);
-			const int32 XMin = FMath::FloorToInt((RcBounds.Min.X - NavMeshOrigin[0]) / TileSize);
-			const int32 XMax = FMath::FloorToInt((RcBounds.Max.X - NavMeshOrigin[0]) / TileSize);
-			const int32 YMin = FMath::FloorToInt((RcBounds.Min.Z - NavMeshOrigin[2]) / TileSize);
-			const int32 YMax = FMath::FloorToInt((RcBounds.Max.Z - NavMeshOrigin[2]) / TileSize);
+			const FVector RcNavMeshOrigin(NavMeshOrigin[0], NavMeshOrigin[1], NavMeshOrigin[2]);
+			const FRcTileBox TileBox(Bounds, RcNavMeshOrigin, TileSize);
 
-			for (int32 y = YMin; y <= YMax; ++y)
+			for (int32 y = TileBox.YMin; y <= TileBox.YMax; ++y)
 			{
-				for (int32 x = XMin; x <= XMax; ++x)
+				for (int32 x = TileBox.XMin; x <= TileBox.XMax; ++x)
 				{
 					TileCoords.Add(FIntPoint(x, y));
 				}
