[Navigation] Adjust navoctree creation to fit navigable world bounds

diff --git a/Engine/Source/Runtime/NavigationSystem/Private/NavigationSystem.cpp b/Engine/Source/Runtime/NavigationSystem/Private/NavigationSystem.cpp
--- a/Engine/Source/Runtime/NavigationSystem/Private/NavigationSystem.cpp
+++ b/Engine/Source/Runtime/NavigationSystem/Private/NavigationSystem.cpp
@@ -758,7 +758,18 @@ void UNavigationSystemV1::PostInitProperties()
 void UNavigationSystemV1::ConstructNavOctree()
 {
 	DefaultOctreeController.Reset();
-	DefaultOctreeController.NavOctree = MakeShareable(new FNavigationOctree(FVector(0, 0, 0), 64000));
+
+	const FBox Bounds = GetNavigableWorldBounds();
+	if(Bounds.IsValid)
+	{
+		DefaultOctreeController.NavOctree = MakeShareable(new FNavigationOctree(Bounds.GetCenter(), 0.5*Bounds.GetSize().Length()));
+	}
+	else
+	{
+		// Fallback to previous default behavior.
+		DefaultOctreeController.NavOctree = MakeShareable(new FNavigationOctree(FVector(0, 0, 0), 64000));
+	}
+	
 	DefaultOctreeController.NavOctree->SetDataGatheringMode(DataGatheringMode);
 #if !UE_BUILD_SHIPPING	
 	DefaultOctreeController.NavOctree->SetGatheringNavModifiersTimeLimitWarning(GatheringNavModifiersWarningLimitTime);
