Improve dirty area warning report when receiving an oversized area

diff --git a/Engine/Source/Runtime/NavigationSystem/Private/NavigationDirtyAreasController.cpp b/Engine/Source/Runtime/NavigationSystem/Private/NavigationDirtyAreasController.cpp
--- a/Engine/Source/Runtime/NavigationSystem/Private/NavigationDirtyAreasController.cpp
+++ b/Engine/Source/Runtime/NavigationSystem/Private/NavigationDirtyAreasController.cpp
@@ -98,27 +98,30 @@ void FNavigationDirtyAreasController::AddArea(const FBox& NewArea, const int32 F
 			ObjectProvider = ObjectProviderFunc();
 		}
 
-		const UActorComponent* ObjectAsComponent = Cast<UActorComponent>(ObjectProvider);
-		const AActor* ComponentOwner = ObjectAsComponent ? ObjectAsComponent->GetOwner() : nullptr;
-		if (ComponentOwner)
+		FString ComponentInfo;
+		if (const UActorComponent* ObjectAsComponent = Cast<UActorComponent>(ObjectProvider))
 		{
-			UE_VLOG_BOX(ComponentOwner, LogNavigationDirtyArea, Log, NewArea, FColor::Red, TEXT(""));
+			if (const AActor* ComponentOwner = ObjectAsComponent->GetOwner())
+			{
+				UE_VLOG_BOX(ComponentOwner, LogNavigationDirtyArea, Log, NewArea, FColor::Red, TEXT(""));
+				ComponentInfo = FString::Printf(TEXT(" | Component's owner: %s"), *GetFullNameSafe(ComponentOwner));
+			}
 		}
-		return FString::Printf(TEXT("Adding dirty area object: %s (from: %s) | Potential component's owner: %s | Bounds size: %s"),
+
+		return FString::Printf(TEXT("Object: %s (from: %s)%s | Bounds: %s"),
 			*GetFullNameSafe(ObjectProvider),
 			*DebugReason.ToString(),
-			*GetFullNameSafe(ComponentOwner),
+			*ComponentInfo,
 			*BoundsSize.ToString());
 	};
 
-	UE_LOG(LogNavigationDirtyArea, VeryVerbose, TEXT("%s"), *DumpExtraInfo());
-
 	if (ShouldReportOversizedDirtyArea() && BoundsSize.GetMax() > DirtyAreaWarningSizeThreshold)
 	{
-		UE_LOG(LogNavigationDirtyArea, Warning, TEXT("Adding an oversized dirty area (object:%s size:%s threshold:%.2f)"),
-			*GetFullNameSafe(ObjectProviderFunc ? ObjectProviderFunc() : nullptr),
-			*BoundsSize.ToString(),
-			DirtyAreaWarningSizeThreshold);
+		UE_LOG(LogNavigationDirtyArea, Warning, TEXT("Adding an oversized dirty area: %s | Threshold: %.2f"), *DumpExtraInfo(), DirtyAreaWarningSizeThreshold);
+	}
+	else
+	{
+		UE_LOG(LogNavigationDirtyArea, VeryVerbose, TEXT("Adding dirty area object: %s"), *DumpExtraInfo());
 	}
 #endif // !UE_BUILD_SHIPPING
 
