Improve navmesh generation debugging :
- add ways to see the heightfield bound (and max height)
- improve logging for components added to the NavOctree

diff --git a/Engine/Source/Runtime/NavigationSystem/Private/NavMesh/RecastNavMesh.cpp b/Engine/Source/Runtime/NavigationSystem/Private/NavMesh/RecastNavMesh.cpp
--- a/Engine/Source/Runtime/NavigationSystem/Private/NavMesh/RecastNavMesh.cpp
+++ b/Engine/Source/Runtime/NavigationSystem/Private/NavMesh/RecastNavMesh.cpp
@@ -448,6 +448,7 @@ FRecastNavMeshTileGenerationDebug::FRecastNavMeshTileGenerationDebug()
 	bHeightfieldSolidFromRasterization = false;
 	bHeightfieldSolidPostInclusionBoundsFiltering = false;
 	bHeightfieldSolidPostHeightFiltering = false;
+	bHeightfieldSolidBounds = false;
 	bCompactHeightfield = false;
 	bCompactHeightfieldEroded = false;
 	bCompactHeightfieldRegions = false;
diff --git a/Engine/Source/Runtime/NavigationSystem/Private/NavMesh/RecastNavMeshGenerator.cpp b/Engine/Source/Runtime/NavigationSystem/Private/NavMesh/RecastNavMeshGenerator.cpp
--- a/Engine/Source/Runtime/NavigationSystem/Private/NavMesh/RecastNavMeshGenerator.cpp
+++ b/Engine/Source/Runtime/NavigationSystem/Private/NavMesh/RecastNavMeshGenerator.cpp
@@ -3332,9 +3332,16 @@ bool FRecastTileGenerator::GenerateCompressedLayers(FNavMeshBuildContext& BuildC
 	}
 
 #if RECAST_INTERNAL_DEBUG_DATA
-	if (IsTileDebugActive() && TileDebugSettings.bHeightfieldSolidFromRasterization)
+	if (IsTileDebugActive())
 	{
-		duDebugDrawHeightfieldSolid(&BuildContext.InternalDebugData, *RasterContext.SolidHF);
+		if (TileDebugSettings.bHeightfieldSolidFromRasterization)
+		{
+			duDebugDrawHeightfieldSolid(&BuildContext.InternalDebugData, *RasterContext.SolidHF);
+		}
+		if (TileDebugSettings.bHeightfieldSolidBounds)
+		{
+			duDebugDrawHeightfieldBounds(&BuildContext.InternalDebugData, *RasterContext.SolidHF);
+		}
 	}
 #endif
 
diff --git a/Engine/Source/Runtime/NavigationSystem/Private/NavigationDataHandler.cpp b/Engine/Source/Runtime/NavigationSystem/Private/NavigationDataHandler.cpp
--- a/Engine/Source/Runtime/NavigationSystem/Private/NavigationDataHandler.cpp
+++ b/Engine/Source/Runtime/NavigationSystem/Private/NavigationDataHandler.cpp
@@ -88,6 +88,18 @@ FSetElementId FNavigationDataHandler::RegisterNavOctreeElement(UObject& ElementO
 	return SetId;
 }
 
+namespace UE::NavigationHelper::Private
+{
+	FString GetElementName(const UObject* ElementOwner)
+	{
+		if (const UActorComponent* ActorComp = Cast<UActorComponent>(ElementOwner))
+		{
+			return FString::Format(TEXT("Comp {0} of Actor {1}"), {*GetNameSafe(ElementOwner), *GetNameSafe(ActorComp->GetOwner())});
+		}
+		return *GetNameSafe(ElementOwner);
+	}
+}
+
 void FNavigationDataHandler::AddElementToNavOctree(const FNavigationDirtyElement& DirtyElement)
 {
 	check(OctreeController.NavOctree.IsValid());
@@ -137,17 +149,17 @@ void FNavigationDataHandler::AddElementToNavOctree(const FNavigationDirtyElement
 		const FOctreeElementId2* ElementId = ParentId ? ParentId : OctreeController.GetObjectsNavOctreeId(*NavigationParent);
 		if (ElementId && ensure(OctreeController.IsValidElement(*ElementId)))
 		{
-			UE_LOG(LogNavOctree, Log, TEXT("ADD %s to %s"), *GetNameSafe(ElementOwner), *GetNameSafe(NavigationParent));
+			UE_LOG(LogNavOctree, Log, TEXT("ADD %s to %s"), *UE::NavigationHelper::Private::GetElementName(ElementOwner), *GetNameSafe(NavigationParent));
 			OctreeController.NavOctree->AppendToNode(*ElementId, DirtyElement.NavInterface, ElementBounds, GeneratedData);
 		}
 		else
 		{
-			UE_LOG(LogNavOctree, Warning, TEXT("Can't add node [%s] - parent [%s] not found in octree!"), *GetNameSafe(ElementOwner), *GetNameSafe(NavigationParent));
+			UE_LOG(LogNavOctree, Warning, TEXT("Can't add node [%s] - parent [%s] not found in octree!"), *UE::NavigationHelper::Private::GetElementName(ElementOwner), *UE::NavigationHelper::Private::GetElementName(NavigationParent));
 		}
 	}
 	else
 	{
-		UE_LOG(LogNavOctree, Log, TEXT("ADD %s"), *GetNameSafe(ElementOwner));
+		UE_LOG(LogNavOctree, Log, TEXT("ADD %s"), *UE::NavigationHelper::Private::GetElementName(ElementOwner));
 		OctreeController.NavOctree->AddNode(ElementOwner, DirtyElement.NavInterface, ElementBounds, GeneratedData);
 	}
 
diff --git a/Engine/Source/Runtime/NavigationSystem/Public/NavMesh/RecastNavMesh.h b/Engine/Source/Runtime/NavigationSystem/Public/NavMesh/RecastNavMesh.h
--- a/Engine/Source/Runtime/NavigationSystem/Public/NavMesh/RecastNavMesh.h
+++ b/Engine/Source/Runtime/NavigationSystem/Public/NavMesh/RecastNavMesh.h
@@ -442,6 +442,9 @@ struct NAVIGATIONSYSTEM_API FRecastNavMeshTileGenerationDebug
 	UPROPERTY(EditAnywhere, Category = Debug)
 	uint32 bHeightfieldSolidPostHeightFiltering : 1;
 
+	UPROPERTY(EditAnywhere, Category = Debug)
+	uint32 bHeightfieldSolidBounds : 1;
+
 	UPROPERTY(EditAnywhere, Category = Debug)
 	uint32 bCompactHeightfield : 1;
 
diff --git a/Engine/Source/Runtime/Navmesh/Private/DebugUtils/RecastDebugDraw.cpp b/Engine/Source/Runtime/Navmesh/Private/DebugUtils/RecastDebugDraw.cpp
--- a/Engine/Source/Runtime/Navmesh/Private/DebugUtils/RecastDebugDraw.cpp
+++ b/Engine/Source/Runtime/Navmesh/Private/DebugUtils/RecastDebugDraw.cpp
@@ -211,6 +211,34 @@ void duDebugDrawHeightfieldWalkable(duDebugDraw* dd, const rcHeightfield& hf)
 	dd->end();
 }
 
+void duDebugDrawHeightfieldBounds(duDebugDraw* dd, const rcHeightfield& hf)
+{
+	if (!dd) return;
+
+	const duReal hfMaxSpawnHeight = hf.bmin[1] + RC_SPAN_MAX_HEIGHT * hf.ch;
+
+	const unsigned int hfBoundCol = duRGBA(160,160,160,255);
+
+	dd->begin(DU_DRAW_LINES, 2.f);
+
+	duAppendBoxWire(dd, hf.bmin[0], hf.bmin[1], hf.bmin[2], hf.bmax[0], hf.bmax[1], hf.bmax[2], hfBoundCol);
+
+	if (hfMaxSpawnHeight < hf.bmax[1]) 
+	{
+		const unsigned int maxSpanBoundCol = duRGBA(255,20,20,255); 
+		dd->vertex(hf.bmin[0], hfMaxSpawnHeight, hf.bmin[2], maxSpanBoundCol);
+		dd->vertex(hf.bmax[0], hfMaxSpawnHeight, hf.bmin[2], maxSpanBoundCol);
+		dd->vertex(hf.bmax[0], hfMaxSpawnHeight, hf.bmin[2], maxSpanBoundCol);
+		dd->vertex(hf.bmax[0], hfMaxSpawnHeight, hf.bmax[2], maxSpanBoundCol);
+		dd->vertex(hf.bmax[0], hfMaxSpawnHeight, hf.bmax[2], maxSpanBoundCol);
+		dd->vertex(hf.bmin[0], hfMaxSpawnHeight, hf.bmax[2], maxSpanBoundCol);
+		dd->vertex(hf.bmin[0], hfMaxSpawnHeight, hf.bmax[2], maxSpanBoundCol);
+		dd->vertex(hf.bmin[0], hfMaxSpawnHeight, hf.bmin[2], maxSpanBoundCol);
+	}
+
+	dd->end();
+}
+
 void duDebugDrawCompactHeightfieldSolid(duDebugDraw* dd, const rcCompactHeightfield& chf)
 {
 	if (!dd) return;
diff --git a/Engine/Source/Runtime/Navmesh/Public/DebugUtils/RecastDebugDraw.h b/Engine/Source/Runtime/Navmesh/Public/DebugUtils/RecastDebugDraw.h
--- a/Engine/Source/Runtime/Navmesh/Public/DebugUtils/RecastDebugDraw.h
+++ b/Engine/Source/Runtime/Navmesh/Public/DebugUtils/RecastDebugDraw.h
@@ -26,12 +26,14 @@
 #include "DebugDrawLargeWorldCoordinates.h"
 
 struct duDebugDraw;
+struct rcHeightfield;
 
 NAVMESH_API void duDebugDrawTriMesh(struct duDebugDraw* dd, const duReal* verts, int nverts, const int* tris, const duReal* normals, int ntris, const unsigned char* flags, const duReal texScale);
 NAVMESH_API void duDebugDrawTriMeshSlope(struct duDebugDraw* dd, const duReal* verts, int nverts, const int* tris, const duReal* normals, int ntris, const duReal walkableSlopeAngle, const duReal texScale);
 
-NAVMESH_API void duDebugDrawHeightfieldSolid(struct duDebugDraw* dd, const struct rcHeightfield& hf);
-NAVMESH_API void duDebugDrawHeightfieldWalkable(struct duDebugDraw* dd, const struct rcHeightfield& hf);
+NAVMESH_API void duDebugDrawHeightfieldSolid(duDebugDraw* dd, const rcHeightfield& hf);
+NAVMESH_API void duDebugDrawHeightfieldWalkable(duDebugDraw* dd, const rcHeightfield& hf);
+NAVMESH_API void duDebugDrawHeightfieldBounds(duDebugDraw* dd, const rcHeightfield& hf);
 
 NAVMESH_API void duDebugDrawCompactHeightfieldSolid(struct duDebugDraw* dd, const struct rcCompactHeightfield& chf);
 NAVMESH_API void duDebugDrawCompactHeightfieldRegions(struct duDebugDraw* dd, const struct rcCompactHeightfield& chf);
