Add navmesh and navsystem informations on screen when navigation display is enabled
Add bDrawTileBuildTimes property to display last tile build time
Add bDrawTileBuildTimesHeatMap property to display the navmesh as a colored heat map of tile build time
Fix UNavMeshRenderingComponent not clearing labels when they where turned off
Allow DebugLabels to be displayed in 2D on screen
Add draw distance clipping for navmesh label rendering
Addition of profile tags in NavMeshRenderingComponent

diff --git a/Engine/Source/Runtime/NavigationSystem/Private/NavMesh/NavMeshRenderingComponent.cpp b/Engine/Source/Runtime/NavigationSystem/Private/NavMesh/NavMeshRenderingComponent.cpp
--- a/Engine/Source/Runtime/NavigationSystem/Private/NavMesh/NavMeshRenderingComponent.cpp
+++ b/Engine/Source/Runtime/NavigationSystem/Private/NavMesh/NavMeshRenderingComponent.cpp
@@ -229,7 +229,7 @@ namespace FNavMeshRenderingHelpers
 		MeshData.Indices.Add(V2);
 	}
 
-	void AddArea(TArray<FNavMeshSceneProxyData::FDebugMeshData>& MeshBuilders, const TArray<int32>& MeshIndices, const TArray<FVector>& MeshVerts, const FColor Color, const FVector DrawOffset)
+	void AddCluster(TArray<FNavMeshSceneProxyData::FDebugMeshData>& MeshBuilders, const TArray<int32>& MeshIndices, const TArray<FVector>& MeshVerts, const FColor Color, const FVector DrawOffset)
 	{
 		if (MeshIndices.Num() == 0)
 		{
@@ -292,7 +292,9 @@ namespace FNavMeshRenderingHelpers
 			(NavMesh->bDrawClusters ? (1 << static_cast<int32>(ENavMeshDetailFlags::Clusters)) : 0) |
 			(NavMesh->bDrawOctree ? (1 << static_cast<int32>(ENavMeshDetailFlags::NavOctree)) : 0) |
 			(NavMesh->bDrawOctreeDetails ? (1 << static_cast<int32>(ENavMeshDetailFlags::NavOctreeDetails)) : 0) |
-			(NavMesh->bDrawMarkedForbiddenPolys ? (1 << static_cast<int32>(ENavMeshDetailFlags::MarkForbiddenPolys)) : 0);
+			(NavMesh->bDrawMarkedForbiddenPolys ? (1 << static_cast<int32>(ENavMeshDetailFlags::MarkForbiddenPolys)) : 0) |
+			(NavMesh->bDrawTileBuildTimes ? (1 << static_cast<int32>(ENavMeshDetailFlags::TileBuildTimes)) : 0) |
+			(NavMesh->bDrawTileBuildTimesHeatMap ? (1 << static_cast<int32>(ENavMeshDetailFlags::TileBuildTimesHeatMap)) : 0);
 	}
 #endif // WITH_RECAST
 
@@ -492,11 +494,87 @@ void FNavMeshSceneProxyData::GatherData(const ARecastNavMesh* NavMesh, int32 InN
 
 		NavMeshDrawOffset.Z = NavMesh->DrawOffset;
 
+		{
+			// On screen information
+			QUICK_SCOPE_CYCLE_COUNTER(STAT_NavMesh_GatherDebugDrawing_2DLabels);
+			
+			// Navmesh
+			const ERuntimeGenerationType Mode = NavMesh->GetRuntimeGenerationMode();
+			const FString GenerationMode = Mode == ERuntimeGenerationType::Static ? TEXT("Static") :
+				(Mode == ERuntimeGenerationType::Dynamic ? TEXT("Dynamic") : (Mode == ERuntimeGenerationType::DynamicModifiersOnly ? TEXT("DynamicModifersOnly") : TEXT("Unknown")));
+			DebugLabels.Add(FDebugText(FString::Printf(TEXT("%s (%s%s)"), *NavMesh->GetName(), NavMesh->bIsWorldPartitioned ? TEXT("WP ") : TEXT(""), *GenerationMode)));
+			DebugLabels.Add(FDebugText(FString::Printf(TEXT("AgentRadius %0.1f, AgentHeight %0.1f, CellSize %0.1f, CellHeight %0.1f"), NavMesh->AgentRadius, NavMesh->AgentHeight, NavMesh->CellSize, NavMesh->CellHeight)));
+			DebugLabels.Add(FDebugText(TEXT(" "))); // empty line
+
+			if (NavMesh->GetActiveTiles().Num() != 0)
+			{
+				DebugLabels.Add(FDebugText(FString::Printf(TEXT("Active tiles: %i"), NavMesh->GetActiveTiles().Num())));	
+			}	
+			
+			DebugLabels.Add(FDebugText(TEXT(" "))); // empty line
+			
+			// Navigation system
+			if (const UNavigationSystemV1* NavSys = FNavigationSystem::GetCurrent<UNavigationSystemV1>(NavMesh->GetWorld()))
+			{
+				DebugLabels.Add(FDebugText(FString::Printf(TEXT("MainNavData: %s"), NavSys->MainNavData ? *NavSys->MainNavData->GetName() : TEXT("none"))));
+
+				if (NavSys->IsActiveTilesGenerationEnabled()) // Checks bGenerateNavigationOnlyAroundNavigationInvokers
+				{
+					DebugLabels.Add(FDebugText(FString::Printf(TEXT("Invoker Locations: %i"), NavSys->GetInvokerLocations().Num())));	
+				}
+				
+				const int32 Running = NavSys->GetNumRunningBuildTasks();
+				const int32 Remaining = NavSys->GetNumRemainingBuildTasks(); 
+				if (Running || Remaining)
+				{
+					DebugLabels.Add(FDebugText(FString::Printf(TEXT("Tile jobs running/remaining: %6d / %6d"), Running, Remaining)));	
+				}
+			}
+		}
+
 		FRecastDebugGeometry NavMeshGeometry;
 		NavMeshGeometry.bGatherPolyEdges = FNavMeshRenderingHelpers::HasFlag(NavDetailFlags, ENavMeshDetailFlags::PolyEdges);
 		NavMeshGeometry.bGatherNavMeshEdges = FNavMeshRenderingHelpers::HasFlag(NavDetailFlags, ENavMeshDetailFlags::BoundaryEdges);
 		NavMeshGeometry.bMarkForbiddenPolys = FNavMeshRenderingHelpers::HasFlag(NavDetailFlags, ENavMeshDetailFlags::MarkForbiddenPolys);
+		NavMeshGeometry.bGatherTileBuildTimesHeatMap = FNavMeshRenderingHelpers::HasFlag(NavDetailFlags, ENavMeshDetailFlags::TileBuildTimesHeatMap);
 
+#if RECAST_INTERNAL_DEBUG_DATA
+		// Colors for tile build times heat map
+		auto LerpColor = [](FColor A, FColor B, float T) -> FColor
+		{
+			return FColor(
+				FMath::RoundToInt(float(A.R) * (1.f - T) + float(B.R) * T),
+				FMath::RoundToInt(float(A.G) * (1.f - T) + float(B.G) * T),
+				FMath::RoundToInt(float(A.B) * (1.f - T) + float(B.B) * T),
+				FMath::RoundToInt(float(A.A) * (1.f - T) + float(B.A) * T));
+		};
+
+		TArray<FColor> TileBuildTimeColors;
+		const TMap<FIntPoint, FRecastInternalDebugData>* DebugDataMap = NavMesh->GetDebugDataMap();
+		
+		if (NavMeshGeometry.bGatherTileBuildTimesHeatMap)
+		{
+			QUICK_SCOPE_CYCLE_COUNTER(STAT_NavMesh_GatherDebugDrawing_MaxTileBuildTime);
+			
+			// Find MaxTileBuildTime 
+			NavMeshGeometry.MaxTileBuildTime = 0;
+			if (DebugDataMap)
+			{
+				for(const TPair<FIntPoint, FRecastInternalDebugData>& Pair : *DebugDataMap)
+				{
+					NavMeshGeometry.MaxTileBuildTime = FMath::Max(NavMeshGeometry.MaxTileBuildTime, Pair.Value.BuildTime);
+				}
+			}
+
+			TileBuildTimeColors.AddDefaulted(FRecastDebugGeometry::BuildTimeBucketsCount);
+			for (int32 Index = 0; Index < FRecastDebugGeometry::BuildTimeBucketsCount; Index++)
+			{
+				const float LerpValue = (float)Index / (FRecastDebugGeometry::BuildTimeBucketsCount-1);
+				TileBuildTimeColors[Index] = LerpColor(FColor::Cyan.WithAlpha(128), FColor::Red.WithAlpha(128), LerpValue);
+			}
+		}
+#endif // RECAST_INTERNAL_DEBUG_DATA
+		
 		const FNavDataConfig& NavConfig = NavMesh->GetConfig();
 		TArray<FColor> NavMeshColors;
 		NavMeshColors.AddDefaulted(RECAST_MAX_AREAS);
@@ -507,8 +585,12 @@ void FNavMeshSceneProxyData::GatherData(const ARecastNavMesh* NavMesh, int32 InN
 		}
 		NavMeshColors[RECAST_DEFAULT_AREA] = NavConfig.Color.DWColor() > 0 ? NavConfig.Color : NavMeshRenderColor_RecastMesh;
 
-		// just a little trick to make sure navmeshes with different sized are not drawn with same offset
-		NavMeshDrawOffset.Z += NavMesh->GetConfig().AgentRadius / 10.f;
+		// Just a little trick to make sure navmeshes with different sized are not drawn with same offset.
+		// When DrawOffset is 0, don't add any offset (usually used to debug navmesh height).
+		if(NavMesh->DrawOffset != 0.f)
+		{
+			NavMeshDrawOffset.Z += NavMesh->GetConfig().AgentRadius / 10.f;
+		}
 
 		NavMesh->BeginBatchQuery();
 		if (TileSet.Num() > 0)
@@ -530,6 +612,8 @@ void FNavMeshSceneProxyData::GatherData(const ARecastNavMesh* NavMesh, int32 InN
 		const bool bGatherTriEdges = FNavMeshRenderingHelpers::HasFlag(NavDetailFlags, ENavMeshDetailFlags::TriangleEdges);
 		if (bGatherTriEdges)
 		{
+			QUICK_SCOPE_CYCLE_COUNTER(STAT_NavMesh_GatherDebugDrawing_GatherTriEdges);
+			
 			for (int32 AreaIdx = 0; AreaIdx < RECAST_MAX_AREAS; ++AreaIdx)
 			{
 				const TArray<int32>& MeshIndices = NavMeshGeometry.AreaIndices[AreaIdx];
@@ -546,6 +630,8 @@ void FNavMeshSceneProxyData::GatherData(const ARecastNavMesh* NavMesh, int32 InN
 		const bool bGatherPolyEdges = FNavMeshRenderingHelpers::HasFlag(NavDetailFlags, ENavMeshDetailFlags::PolyEdges);
 		if (bGatherPolyEdges)
 		{
+			QUICK_SCOPE_CYCLE_COUNTER(STAT_NavMesh_GatherDebugDrawing_GatherPolyEdges);
+			
 			const TArray<FVector>& TileEdgeVerts = NavMeshGeometry.PolyEdges;
 			for (int32 Idx = 0; Idx < TileEdgeVerts.Num(); Idx += 2)
 			{
@@ -557,6 +643,8 @@ void FNavMeshSceneProxyData::GatherData(const ARecastNavMesh* NavMesh, int32 InN
 		const bool bGatherBoundaryEdges = FNavMeshRenderingHelpers::HasFlag(NavDetailFlags, ENavMeshDetailFlags::BoundaryEdges);
 		if (bGatherBoundaryEdges)
 		{
+			QUICK_SCOPE_CYCLE_COUNTER(STAT_NavMesh_GatherDebugDrawing_GatherBoundaryEdges);
+			
 			const FColor EdgesColor = FNavMeshRenderingHelpers::DarkenColor(NavMeshColors[RECAST_DEFAULT_AREA]);
 			const TArray<FVector>& NavMeshEdgeVerts = NavMeshGeometry.NavMeshEdges;
 			for (int32 Idx = 0; Idx < NavMeshEdgeVerts.Num(); Idx += 2)
@@ -574,6 +662,8 @@ void FNavMeshSceneProxyData::GatherData(const ARecastNavMesh* NavMesh, int32 InN
 		if (bGatherClusters == false)
 #endif // WITH_NAVMESH_CLUSTER_LINKS
 		{
+			QUICK_SCOPE_CYCLE_COUNTER(STAT_NavMesh_GatherDebugDrawing_OffMeshLinks);
+			
 			for (int32 OffMeshLineIndex = 0; OffMeshLineIndex < NavMeshGeometry.OffMeshLinks.Num(); ++OffMeshLineIndex)
 			{
 				FRecastDebugGeometry::FOffMeshLink& Link = NavMeshGeometry.OffMeshLinks[OffMeshLineIndex];
@@ -602,8 +692,11 @@ void FNavMeshSceneProxyData::GatherData(const ARecastNavMesh* NavMesh, int32 InN
 					}
 				}
 			}
+			
 			if (NavMeshGeometry.bMarkForbiddenPolys)
 			{
+				QUICK_SCOPE_CYCLE_COUNTER(STAT_NavMesh_GatherDebugDrawing_MarkForbiddenPolys);
+				
 				for (int32 OffMeshLineIndex = 0; OffMeshLineIndex < NavMeshGeometry.ForbiddenLinks.Num(); ++OffMeshLineIndex)
 				{
 					FRecastDebugGeometry::FOffMeshLink& Link = NavMeshGeometry.ForbiddenLinks[OffMeshLineIndex];
@@ -622,9 +715,13 @@ void FNavMeshSceneProxyData::GatherData(const ARecastNavMesh* NavMesh, int32 InN
 		const bool bGatherPolygonLabels = FNavMeshRenderingHelpers::HasFlag(NavDetailFlags, ENavMeshDetailFlags::PolygonLabels);
 		const bool bGatherPolygonCost = FNavMeshRenderingHelpers::HasFlag(NavDetailFlags, ENavMeshDetailFlags::PolygonCost);
 		const bool bGatherPolygonFlags = FNavMeshRenderingHelpers::HasFlag(NavDetailFlags, ENavMeshDetailFlags::PolygonFlags);
+		const bool bGatherTileBuildTimes = FNavMeshRenderingHelpers::HasFlag(NavDetailFlags, ENavMeshDetailFlags::TileBuildTimes);
+		const bool bGatherTileBuildHeatMap = FNavMeshRenderingHelpers::HasFlag(NavDetailFlags, ENavMeshDetailFlags::TileBuildTimesHeatMap);
 
-		if (bGatherTileLabels || bGatherTileBounds || bGatherPolygonLabels || bGatherPolygonCost || bGatherPolygonFlags)
+		if (bGatherTileLabels || bGatherTileBounds || bGatherPolygonLabels || bGatherPolygonCost || bGatherPolygonFlags || bGatherTileBuildTimes)
 		{
+			QUICK_SCOPE_CYCLE_COUNTER(STAT_NavMesh_GatherDebugDrawing_TileIterations);
+			
 			TArray<int32> UseTileIndices;
 			if (TileSet.Num() > 0)
 			{
@@ -639,6 +736,8 @@ void FNavMeshSceneProxyData::GatherData(const ARecastNavMesh* NavMesh, int32 InN
 				}
 			}
 
+			TMap<FIntPoint, FVector> TileBuildTimeLabelLocations;
+			
 			// calculate appropriate points for displaying debug labels
 			DebugLabels.Reserve(UseTileIndices.Num());
 			for (int32 TileSetIdx = 0; TileSetIdx < UseTileIndices.Num(); TileSetIdx++)
@@ -662,6 +761,21 @@ void FNavMeshSceneProxyData::GatherData(const ARecastNavMesh* NavMesh, int32 InN
 						DebugLabels.Add(FDebugText(NavLocation.Location + NavMeshDrawOffset, FString::Printf(TEXT("(%d,%d:%d)"), X, Y, Layer)));
 					}
 
+					if (bGatherTileBuildTimes)
+					{
+						// Just keep the highest layer location
+						const FIntPoint Coord(X, Y);
+						FVector* Location = TileBuildTimeLabelLocations.Find(Coord);
+						if (!Location)
+						{
+							TileBuildTimeLabelLocations.Add(Coord, NavLocation.Location);
+						}
+						else if(NavLocation.Location.Z > Location->Z)
+						{
+							*Location = NavLocation;
+						}
+					}	
+					
 					if (bGatherPolygonLabels || bGatherPolygonCost || bGatherPolygonFlags)
 					{
 						TArray<FNavPoly> Polys;
@@ -736,19 +850,41 @@ void FNavMeshSceneProxyData::GatherData(const ARecastNavMesh* NavMesh, int32 InN
 					}
 				}
 			}
+
+#if RECAST_INTERNAL_DEBUG_DATA
+			if (bGatherTileBuildTimes)
+			{
+				QUICK_SCOPE_CYCLE_COUNTER(STAT_NavMesh_GatherDebugDrawing_TileBuildTimes);
+				for(const TPair<FIntPoint, FVector>& Pair : TileBuildTimeLabelLocations)
+				{
+					if (const FRecastInternalDebugData* DebugData = DebugDataMap->Find(Pair.Key))
+					{
+						constexpr double SecToMs = 1000.;
+						const double BuildTimeMs = DebugData->BuildTime * SecToMs;
+						const double CompressedLayerTimeMs = DebugData->BuildCompressedLayerTime * SecToMs;
+						const double NavigationDataTimeMs = DebugData->BuildNavigationDataTime * SecToMs;
+						DebugLabels.Add(FDebugText(Pair.Value + NavMeshDrawOffset, FString::Printf(TEXT("%.2f ms\n%.2f + %.2f"), BuildTimeMs, CompressedLayerTimeMs, NavigationDataTimeMs)));
+					}
+				}
+			}
+#endif // RECAST_INTERNAL_DEBUG_DATA
 		}
 
 #if RECAST_INTERNAL_DEBUG_DATA
-		// Display internal debug data
-		for (const FIntPoint& point : NavMeshGeometry.TilesToDisplayInternalData)
 		{
-			if (NavMesh->GetDebugDataMap())
+			QUICK_SCOPE_CYCLE_COUNTER(STAT_NavMesh_GatherDebugDrawing_DisplayInternalData);
+
+			// Display internal debug data
+			for (const FIntPoint& point : NavMeshGeometry.TilesToDisplayInternalData)
 			{
-				const FRecastInternalDebugData& DebugData = NavMesh->GetDebugDataMap()->FindRef(point);
-				AddMeshForInternalData(DebugData);
+				if (NavMesh->GetDebugDataMap())
+				{
+					const FRecastInternalDebugData& DebugData = NavMesh->GetDebugDataMap()->FindRef(point);
+					AddMeshForInternalData(DebugData);
+				}
 			}
 		}
-#endif
+#endif // RECAST_INTERNAL_DEBUG_DATA
 
 		NavMesh->FinishBatchQuery();
 
@@ -756,6 +892,8 @@ void FNavMeshSceneProxyData::GatherData(const ARecastNavMesh* NavMesh, int32 InN
 		const bool bGatherFilledPolys = FNavMeshRenderingHelpers::HasFlag(NavDetailFlags, ENavMeshDetailFlags::FilledPolys);
 		if (bGatherFilledPolys)
 		{
+			QUICK_SCOPE_CYCLE_COUNTER(STAT_NavMesh_GatherDebugDrawing_GatherFilledPolys);
+			
 #if WITH_NAVMESH_CLUSTER_LINKS
 			if (bGatherClusters)
 			{
@@ -786,17 +924,31 @@ void FNavMeshSceneProxyData::GatherData(const ARecastNavMesh* NavMesh, int32 InN
 			{
 				for (int32 AreaType = 0; AreaType < RECAST_MAX_AREAS; ++AreaType)
 				{
-					FNavMeshRenderingHelpers::AddArea(MeshBuilders, NavMeshGeometry.AreaIndices[AreaType], NavMeshGeometry.MeshVerts, NavMeshColors[AreaType], NavMeshDrawOffset);
+					FNavMeshRenderingHelpers::AddCluster(MeshBuilders, NavMeshGeometry.AreaIndices[AreaType], NavMeshGeometry.MeshVerts, NavMeshColors[AreaType], NavMeshDrawOffset);
 				}
-				FNavMeshRenderingHelpers::AddArea(MeshBuilders, NavMeshGeometry.ForbiddenIndices, NavMeshGeometry.MeshVerts, NavMeshRenderColor_PolyForbidden, NavMeshDrawOffset);
+				FNavMeshRenderingHelpers::AddCluster(MeshBuilders, NavMeshGeometry.ForbiddenIndices, NavMeshGeometry.MeshVerts, NavMeshRenderColor_PolyForbidden, NavMeshDrawOffset);
 			}
 		}
 
+#if RECAST_INTERNAL_DEBUG_DATA		
+		if (bGatherTileBuildHeatMap)
+		{
+			QUICK_SCOPE_CYCLE_COUNTER(STAT_NavMesh_GatherDebugDrawing_GatherTileBuildHeatMap);
+			
+			for (int32 Index = 0; Index < FRecastDebugGeometry::BuildTimeBucketsCount; ++Index)
+			{
+				FNavMeshRenderingHelpers::AddCluster(MeshBuilders, NavMeshGeometry.TileBuildTimesIndices[Index], NavMeshGeometry.MeshVerts, TileBuildTimeColors[Index], NavMeshDrawOffset);
+			}
+		}
+#endif // RECAST_INTERNAL_DEBUG_DATA
+
 		const bool bGatherOctree = FNavMeshRenderingHelpers::HasFlag(NavDetailFlags, ENavMeshDetailFlags::NavOctree);
 		const bool bGatherOctreeDetails = FNavMeshRenderingHelpers::HasFlag(NavDetailFlags, ENavMeshDetailFlags::NavOctreeDetails);
 		const bool bGatherPathCollidingGeometry = FNavMeshRenderingHelpers::HasFlag(NavDetailFlags, ENavMeshDetailFlags::PathCollidingGeometry);
 		if (bGatherOctree || bGatherPathCollidingGeometry)
 		{
+			QUICK_SCOPE_CYCLE_COUNTER(STAT_NavMesh_GatherDebugDrawing_Octree);
+			
 			const UNavigationSystemV1* NavSys = FNavigationSystem::GetCurrent<UNavigationSystemV1>(NavMesh->GetWorld());
 			const FNavigationOctree* NavOctree = NavSys ? NavSys->GetNavOctree() : nullptr;
 			if (NavOctree)
@@ -864,6 +1016,8 @@ void FNavMeshSceneProxyData::GatherData(const ARecastNavMesh* NavMesh, int32 InN
 
 		if (NavMeshGeometry.BuiltMeshIndices.Num() > 0)
 		{
+			QUICK_SCOPE_CYCLE_COUNTER(STAT_NavMesh_GatherDebugDrawing_BuiltMeshIndices);
+			
 			FDebugMeshData DebugMeshData;
 			for (int32 VertIdx = 0; VertIdx < MeshVerts.Num(); ++VertIdx)
 			{
@@ -877,6 +1031,8 @@ void FNavMeshSceneProxyData::GatherData(const ARecastNavMesh* NavMesh, int32 InN
 #if WITH_NAVMESH_CLUSTER_LINKS
 		if (bGatherClusters)
 		{
+			QUICK_SCOPE_CYCLE_COUNTER(STAT_NavMesh_GatherDebugDrawing_GatherClusters);
+			
 			for (int32 Idx = 0; Idx < NavMeshGeometry.ClusterLinks.Num(); Idx++)
 			{
 				const FRecastDebugGeometry::FClusterLink& CLink = NavMeshGeometry.ClusterLinks[Idx];
@@ -894,6 +1050,8 @@ void FNavMeshSceneProxyData::GatherData(const ARecastNavMesh* NavMesh, int32 InN
 		// cache segment links
 		if (bGatherNavLinks)
 		{
+			QUICK_SCOPE_CYCLE_COUNTER(STAT_NavMesh_GatherDebugDrawing_NavLinks);
+			
 			for (int32 AreaIdx = 0; AreaIdx < RECAST_MAX_AREAS; AreaIdx++)
 			{
 				const TArray<int32>& Indices = NavMeshGeometry.OffMeshSegmentAreas[AreaIdx];
@@ -1343,12 +1501,22 @@ void FNavMeshDebugDrawDelegateHelper::DrawDebugLabels(UCanvas* Canvas, APlayerCo
 	const FSceneView* View = Canvas->SceneView;
 	const UFont* Font = GEngine->GetSmallFont();
 	const FNavMeshSceneProxyData::FDebugText* DebugText = DebugLabels.GetData();
+	float ScreenY = 40.f;
 	for (int32 Idx = 0; Idx < DebugLabels.Num(); ++Idx, ++DebugText)
 	{
-		if (View->ViewFrustum.IntersectSphere(DebugText->Location, 1.0f))
+		if (DebugText->Location == FNavigationSystem::InvalidLocation)
 		{
-			const FVector ScreenLoc = Canvas->Project(DebugText->Location);
-			Canvas->DrawText(Font, DebugText->Text, ScreenLoc.X, ScreenLoc.Y);
+			constexpr float ScreenX = 10.f;
+			Canvas->DrawText(Font, DebugText->Text, ScreenX, ScreenY);
+			ScreenY += Font->GetStringHeightSize(*DebugText->Text);
+		}
+		else
+		{
+			if (FNavMeshRenderingHelpers::PointInView(DebugText->Location, View))
+			{
+				const FVector ScreenLoc = Canvas->Project(DebugText->Location);
+				Canvas->DrawText(Font, DebugText->Text, ScreenLoc.X, ScreenLoc.Y);
+			}
 		}
 	}
 
@@ -1514,7 +1682,7 @@ void UNavMeshRenderingComponent::GatherData(const ARecastNavMesh& NavMesh, FNavM
 }
 
 #if UE_ENABLE_DEBUG_DRAWING
-  FDebugRenderSceneProxy* UNavMeshRenderingComponent::CreateDebugSceneProxy()
+FDebugRenderSceneProxy* UNavMeshRenderingComponent::CreateDebugSceneProxy()
 {
 #if WITH_RECAST
 	FNavMeshSceneProxy* NavMeshSceneProxy = nullptr;
@@ -1534,6 +1702,10 @@ void UNavMeshRenderingComponent::GatherData(const ARecastNavMesh& NavMesh, FNavM
 			NavMeshSceneProxy = new FNavMeshSceneProxy(this, &ProxyData);
 			NavMeshDebugDrawDelegateManager.SetupFromProxy(NavMeshSceneProxy);
 		}
+		else
+		{
+			NavMeshDebugDrawDelegateManager.Reset();
+		}
 	}
 
 	return NavMeshSceneProxy;
diff --git a/Engine/Source/Runtime/NavigationSystem/Private/NavMesh/PImplRecastNavMesh.cpp b/Engine/Source/Runtime/NavigationSystem/Private/NavMesh/PImplRecastNavMesh.cpp
--- a/Engine/Source/Runtime/NavigationSystem/Private/NavMesh/PImplRecastNavMesh.cpp
+++ b/Engine/Source/Runtime/NavigationSystem/Private/NavMesh/PImplRecastNavMesh.cpp
@@ -2631,6 +2631,8 @@ uint8 GetValidEnds(const dtNavMesh& NavMesh, const dtMeshTile& Tile, const dtPol
 
 bool FPImplRecastNavMesh::GetDebugGeometryForTile(FRecastDebugGeometry& OutGeometry, int32 TileIndex) const
 {
+	QUICK_SCOPE_CYCLE_COUNTER(STAT_FPImplRecastNavMesh_GetDebugGeometryForTile);
+	
 	bool bDone = false;
 	if (DetourNavMesh == nullptr || TileIndex >= DetourNavMesh->getMaxTiles())
 	{
@@ -2683,6 +2685,10 @@ bool FPImplRecastNavMesh::GetDebugGeometryForTile(FRecastDebugGeometry& OutGeome
 		OutGeometry.MeshVerts.Reserve(OutGeometry.MeshVerts.Num() + NumVertsToReserve);
 		OutGeometry.AreaIndices[0].Reserve(OutGeometry.AreaIndices[0].Num() + NumIndicesToReserve);
 		OutGeometry.BuiltMeshIndices.Reserve(OutGeometry.BuiltMeshIndices.Num() + NumIndicesToReserve);
+		for (int32 Index = 0; Index < FRecastDebugGeometry::BuildTimeBucketsCount; Index++)
+		{
+			OutGeometry.TileBuildTimesIndices[Index].Reserve(OutGeometry.TileBuildTimesIndices[Index].Num() + NumIndicesToReserve);	
+		}
 
 		uint32 VertBase = OutGeometry.MeshVerts.Num();
 		for (const FIntPoint& TileLocation : ActiveTiles)
@@ -2723,6 +2729,10 @@ bool FPImplRecastNavMesh::GetDebugGeometryForTile(FRecastDebugGeometry& OutGeome
 		OutGeometry.MeshVerts.Reserve(OutGeometry.MeshVerts.Num() + NumVertsToReserve);
 		OutGeometry.AreaIndices[0].Reserve(OutGeometry.AreaIndices[0].Num() + NumIndicesToReserve);
 		OutGeometry.BuiltMeshIndices.Reserve(OutGeometry.BuiltMeshIndices.Num() + NumIndicesToReserve);
+		for (int32 Index = 0; Index < FRecastDebugGeometry::BuildTimeBucketsCount; Index++)
+		{
+			OutGeometry.TileBuildTimesIndices[Index].Reserve(OutGeometry.TileBuildTimesIndices[Index].Num() + NumIndicesToReserve);	
+		}
 
 		uint32 VertBase = OutGeometry.MeshVerts.Num();
 		for (int32 TileIdx = StartingTile; TileIdx < NumTiles; ++TileIdx)
@@ -2788,6 +2798,11 @@ int32 FPImplRecastNavMesh::GetTilesDebugGeometry(const FRecastNavMeshGenerator*
 		F += 3;
 	}
 
+#if RECAST_INTERNAL_DEBUG_DATA	
+	const FIntPoint TileCoord(Header->x, Header->y);
+	const FRecastInternalDebugData* DebugData = DebugDataMap.Find(TileCoord);
+#endif // RECAST_INTERNAL_DEBUG_DATA	
+	
 	// add all the indices
 	for (int32 PolyIdx = 0; PolyIdx < Header->polyCount; ++PolyIdx)
 	{
@@ -2797,10 +2812,27 @@ int32 FPImplRecastNavMesh::GetTilesDebugGeometry(const FRecastNavMeshGenerator*
 		{
 			dtPolyDetail const* const DetailPoly = &Tile.detailMeshes[PolyIdx];
 
-			TArray<int32>* Indices = bIsBeingBuilt ? &OutGeometry.BuiltMeshIndices 
-				: ((Poly->flags & ForbiddenFlags) != 0
-					? &OutGeometry.ForbiddenIndices
-					: &OutGeometry.AreaIndices[Poly->getArea()]);
+			TArray<int32>* Indices = nullptr;
+			if (bIsBeingBuilt)
+			{
+				Indices = &OutGeometry.BuiltMeshIndices;
+			}
+			else if ((Poly->flags & ForbiddenFlags) != 0)
+			{
+				Indices = &OutGeometry.ForbiddenIndices;
+			}
+#if RECAST_INTERNAL_DEBUG_DATA			
+			else if (OutGeometry.bGatherTileBuildTimesHeatMap && DebugData)
+			{
+				int32 Rank = (FRecastDebugGeometry::BuildTimeBucketsCount-1) * (DebugData->BuildTime / OutGeometry.MaxTileBuildTime);
+				Rank = FMath::Clamp(Rank, 0, FRecastDebugGeometry::BuildTimeBucketsCount-1);
+				Indices = &OutGeometry.TileBuildTimesIndices[Rank];
+			}
+#endif // RECAST_INTERNAL_DEBUG_DATA
+			else
+			{
+				Indices = &OutGeometry.AreaIndices[Poly->getArea()];
+			}
 
 			// one triangle at a time
 			for (int32 TriIdx = 0; TriIdx < DetailPoly->triCount; ++TriIdx)
diff --git a/Engine/Source/Runtime/NavigationSystem/Private/NavMesh/RecastNavMesh.cpp b/Engine/Source/Runtime/NavigationSystem/Private/NavMesh/RecastNavMesh.cpp
--- a/Engine/Source/Runtime/NavigationSystem/Private/NavMesh/RecastNavMesh.cpp
+++ b/Engine/Source/Runtime/NavigationSystem/Private/NavMesh/RecastNavMesh.cpp
@@ -292,6 +292,13 @@ uint32 FRecastDebugGeometry::GetAllocatedSize() const
 		Size += AreaIndices[i].GetAllocatedSize();
 	}
 
+#if RECAST_INTERNAL_DEBUG_DATA
+	for (int i = 0; i < BuildTimeBucketsCount; ++i)
+	{
+		Size += TileBuildTimesIndices[i].GetAllocatedSize();
+	}
+#endif // RECAST_INTERNAL_DEBUG_DATA
+
 #if WITH_NAVMESH_CLUSTER_LINKS
 	Size += Clusters.GetAllocatedSize()	+ ClusterLinks.GetAllocatedSize();
 
@@ -3277,12 +3284,14 @@ bool ARecastNavMesh::HasCompleteDataInRadius(const FVector& TestLocation, FVecto
 //----------------------------------------------------------------------//
 void ARecastNavMesh::UpdateActiveTiles(const TArray<FNavigationInvokerRaw>& InvokerLocations)
 {
+	TRACE_CPUPROFILER_EVENT_SCOPE(ARecastNavMesh::UpdateActiveTiles);
+	
 	if (HasValidNavmesh() == false)
 	{
 		return;
 	}
 
-	FRecastNavMeshGenerator* MyGenerator = static_cast<FRecastNavMeshGenerator*>(GetGenerator());
+	const FRecastNavMeshGenerator* MyGenerator = static_cast<FRecastNavMeshGenerator*>(GetGenerator());
 	if (MyGenerator == nullptr)
 	{
 		return;
@@ -3302,14 +3311,11 @@ void ARecastNavMesh::UpdateActiveTiles(const TArray<FNavigationInvokerRaw>& Invo
 	TilesInMaxDistance.Reserve(ActiveTiles.Num());
 	ActiveTiles.Reset();
 
-	//const int32 TileRadius = FMath::CeilToInt(Radius / TileDim);
-	static const FVector::FReal SqareRootOf2 = FMath::Sqrt(2.f);
-
 	for (const FNavigationInvokerRaw& Invoker : InvokerLocations)
 	{
 		const FVector InvokerRelativeLocation = (NavmeshOrigin - Invoker.Location);
-		const  FVector::FReal TileCenterDistanceToRemoveSq = FMath::Square(TileDim * SqareRootOf2 / 2 + Invoker.RadiusMax);
-		const  FVector::FReal TileCenterDistanceToAddSq = FMath::Square(TileDim * SqareRootOf2 / 2 + Invoker.RadiusMin);
+		const  FVector::FReal TileCenterDistanceToRemoveSq = FMath::Square(TileDim * UE_SQRT_2 / 2 + Invoker.RadiusMax);
+		const  FVector::FReal TileCenterDistanceToAddSq = FMath::Square(TileDim * UE_SQRT_2 / 2 + Invoker.RadiusMin);
 
 		const int32 MinTileX = FMath::FloorToInt((InvokerRelativeLocation.X - Invoker.RadiusMax) / TileDim);
 		const int32 MaxTileX = FMath::CeilToInt((InvokerRelativeLocation.X + Invoker.RadiusMax) / TileDim);
@@ -3340,14 +3346,15 @@ void ARecastNavMesh::UpdateActiveTiles(const TArray<FNavigationInvokerRaw>& Invo
 	TilesToRemove.Reserve(OldActiveSet.Num());
 	for (int32 Index = OldActiveSet.Num() - 1; Index >= 0; --Index)
 	{
-		if (TilesInMaxDistance.Find(OldActiveSet[Index]) == INDEX_NONE)
+		const FIntPoint& Tile = OldActiveSet[Index];
+		if (TilesInMaxDistance.Find(Tile) == INDEX_NONE)
 		{
-			TilesToRemove.Add(OldActiveSet[Index]);
+			TilesToRemove.Add(Tile);
 			OldActiveSet.RemoveAtSwap(Index, 1, /*bAllowShrinking=*/false);
 		}
 		else
 		{
-			ActiveTiles.AddUnique(OldActiveSet[Index]);
+			ActiveTiles.AddUnique(Tile);
 		}
 	}
 
@@ -3356,9 +3363,10 @@ void ARecastNavMesh::UpdateActiveTiles(const TArray<FNavigationInvokerRaw>& Invo
 	for (int32 Index = TilesInMinDistance.Num() - 1; Index >= 0; --Index)
 	{
 		// check if it's a new tile
-		if (OldActiveSet.Find(TilesInMinDistance[Index]) == INDEX_NONE)
+		const FIntPoint& Tile = TilesInMinDistance[Index];
+		if (OldActiveSet.Find(Tile) == INDEX_NONE)
 		{
-			TilesToUpdate.Add(TilesInMinDistance[Index]);
+			TilesToUpdate.Add(Tile);
 		}
 	}
 
@@ -3417,7 +3425,7 @@ void ARecastNavMesh::DirtyTilesInBounds(const FBox& Bounds)
 	for (int32 TileIndex = 0; TileIndex < TileCount; ++TileIndex)
 	{
 		const dtMeshTile* Tile = DetourNavMesh->getTile(TileIndex);
-		dtMeshHeader* Header = Tile != nullptr ? Tile->header : nullptr;
+		const dtMeshHeader* Header = Tile != nullptr ? Tile->header : nullptr;
 		if (Header)
 		{
 			TilesToRemove.Add(FIntPoint(Header->x, Header->y));
diff --git a/Engine/Source/Runtime/NavigationSystem/Private/NavMesh/RecastNavMeshGenerator.cpp b/Engine/Source/Runtime/NavigationSystem/Private/NavMesh/RecastNavMeshGenerator.cpp
--- a/Engine/Source/Runtime/NavigationSystem/Private/NavMesh/RecastNavMeshGenerator.cpp
+++ b/Engine/Source/Runtime/NavigationSystem/Private/NavMesh/RecastNavMeshGenerator.cpp
@@ -73,6 +73,12 @@ static FAutoConsoleVariableRef NavmeshVarDebugTileY(TEXT("n.GNavmeshDebugTileY")
 static bool GNavmeshUseOodleCompression = true;
 static FAutoConsoleVariableRef NavmeshVarOodleCompression(TEXT("n.NavmeshUseOodleCompression"), GNavmeshUseOodleCompression, TEXT("Use Oodle for run-time tile cache compression/decompression. Optimized for size in editor, optimized for speed in standalone."), ECVF_Default);
 
+namespace UE::NavMesh::Private
+{
+	static float RecentlyBuildTileDisplayTime = 0.2f;
+	static FAutoConsoleVariableRef CVarRecentlyBuildTileDisplayTime(TEXT("n.RecentlyBuildTileDisplayTime"), RecentlyBuildTileDisplayTime, TEXT("Time (in seconds) to display tiles that have recently been built."), ECVF_Default);
+}
+
 #if	WITH_EDITOR	
 static FOodleDataCompression::ECompressor GNavmeshTileCacheCompressor = FOodleDataCompression::ECompressor::Kraken;
 static FOodleDataCompression::ECompressionLevel GNavmeshTileCacheCompressionLevel = FOodleDataCompression::ECompressionLevel::Optimal1;
@@ -2433,6 +2439,11 @@ ETimeSliceWorkResult FRecastTileGenerator::GenerateTileTimeSliced()
 
 bool FRecastTileGenerator::GenerateTile()
 {
+#if RECAST_INTERNAL_DEBUG_DATA
+	const double StartStamp = FPlatformTime::Seconds();
+	double PostCompressLayerStamp = StartStamp;
+#endif // RECAST_INTERNAL_DEBUG_DATA
+	
 	FNavMeshBuildContext BuildContext(*this);
 	bool bSuccess = true;
 
@@ -2442,6 +2453,10 @@ bool FRecastTileGenerator::GenerateTile()
 
 		bSuccess = GenerateCompressedLayers(BuildContext);
 
+#if RECAST_INTERNAL_DEBUG_DATA
+		PostCompressLayerStamp = FPlatformTime::Seconds();
+#endif // RECAST_INTERNAL_DEBUG_DATA
+
 		if (bSuccess)
 		{
 			// Mark all layers as dirty
@@ -2454,6 +2469,13 @@ bool FRecastTileGenerator::GenerateTile()
 		bSuccess = GenerateNavigationData(BuildContext);
 	}
 
+#if RECAST_INTERNAL_DEBUG_DATA	
+	const double EndStamp = FPlatformTime::Seconds();
+	BuildContext.InternalDebugData.BuildTime = EndStamp - StartStamp;
+	BuildContext.InternalDebugData.BuildCompressedLayerTime = PostCompressLayerStamp - StartStamp;
+	BuildContext.InternalDebugData.BuildNavigationDataTime = EndStamp - PostCompressLayerStamp;
+#endif // RECAST_INTERNAL_DEBUG_DATA
+	
 	// it's possible to have valid generation with empty resulting tile (no navigable geometry in tile)
 	return bSuccess;
 }
@@ -4831,7 +4853,7 @@ void FRecastNavMeshGenerator::TickAsyncBuild(float DeltaSeconds)
 		const double Timestamp = FPlatformTime::Seconds();
 		const int32 NumPreRemove = RecentlyBuiltTiles.Num();
 		
-		RecentlyBuiltTiles.RemoveAllSwap([&](const FTileTimestamp& Tile) { return (Timestamp - Tile.Timestamp) > 0.5; });
+		RecentlyBuiltTiles.RemoveAllSwap([&](const FTileTimestamp& Tile) { return (Timestamp - Tile.Timestamp) > UE::NavMesh::Private::RecentlyBuildTileDisplayTime; });
 
 		const int32 NumPostRemove = RecentlyBuiltTiles.Num();
 		bRequestDrawingUpdate = (NumPreRemove != NumPostRemove);
diff --git a/Engine/Source/Runtime/NavigationSystem/Public/NavMesh/NavMeshRenderingComponent.h b/Engine/Source/Runtime/NavigationSystem/Public/NavMesh/NavMeshRenderingComponent.h
--- a/Engine/Source/Runtime/NavigationSystem/Public/NavMesh/NavMeshRenderingComponent.h
+++ b/Engine/Source/Runtime/NavigationSystem/Public/NavMesh/NavMeshRenderingComponent.h
@@ -42,6 +42,8 @@ enum class ENavMeshDetailFlags : uint8
 	NavOctree,
 	NavOctreeDetails,
 	MarkForbiddenPolys,
+	TileBuildTimes,
+	TileBuildTimesHeatMap
 };
 
 // exported to API for GameplayDebugger module
@@ -80,6 +82,7 @@ struct NAVIGATIONSYSTEM_API FNavMeshSceneProxyData : public TSharedFromThis<FNav
 
 		FDebugText() {}
 		FDebugText(const FVector& InLocation, const FString& InText) : Location(InLocation), Text(InText) {}
+		FDebugText(const FString& InText) : Location(FNavigationSystem::InvalidLocation), Text(InText) {}
 	};
 	TArray<FDebugText> DebugLabels;
 	
@@ -164,6 +167,12 @@ public:
 		bNeedsNewData = InSceneProxy->ProxyData.bNeedsNewData;
 	}
 
+	void Reset()
+	{
+		DebugLabels.Reset();
+		bNeedsNewData = true;
+	}
+
 protected:
 	NAVIGATIONSYSTEM_API virtual void DrawDebugLabels(UCanvas* Canvas, APlayerController*) override;
 
diff --git a/Engine/Source/Runtime/NavigationSystem/Public/NavMesh/RecastInternalDebugData.h b/Engine/Source/Runtime/NavigationSystem/Public/NavMesh/RecastInternalDebugData.h
--- a/Engine/Source/Runtime/NavigationSystem/Public/NavMesh/RecastInternalDebugData.h
+++ b/Engine/Source/Runtime/NavigationSystem/Public/NavMesh/RecastInternalDebugData.h
@@ -24,6 +24,10 @@ struct FRecastInternalDebugData : public duDebugDraw
 	TArray<FVector> PointVertices;
 	TArray<FColor>  PointColors;
 
+	double BuildTime = 0.;
+	double BuildCompressedLayerTime = 0.;
+	double BuildNavigationDataTime = 0.;
+
 	FRecastInternalDebugData() {}
 	virtual ~FRecastInternalDebugData() override {}
 
diff --git a/Engine/Source/Runtime/NavigationSystem/Public/NavMesh/RecastNavMesh.h b/Engine/Source/Runtime/NavigationSystem/Public/NavMesh/RecastNavMesh.h
--- a/Engine/Source/Runtime/NavigationSystem/Public/NavMesh/RecastNavMesh.h
+++ b/Engine/Source/Runtime/NavigationSystem/Public/NavMesh/RecastNavMesh.h
@@ -213,10 +213,15 @@ struct FRecastDebugGeometry
 	};
 #endif // WITH_NAVMESH_SEGMENT_LINKS
 
+	static constexpr int32 BuildTimeBucketsCount = 5;
+	
 	TArray<FVector> MeshVerts;
+	
 	TArray<int32> AreaIndices[RECAST_MAX_AREAS];
 	TArray<int32> ForbiddenIndices;
 	TArray<int32> BuiltMeshIndices;
+	TArray<int32> TileBuildTimesIndices[BuildTimeBucketsCount];
+	
 	TArray<FVector> PolyEdges;
 	TArray<FVector> NavMeshEdges;
 	TArray<FOffMeshLink> OffMeshLinks;
@@ -239,8 +244,11 @@ struct FRecastDebugGeometry
 	int32 bGatherPolyEdges : 1;
 	int32 bGatherNavMeshEdges : 1;
 	int32 bMarkForbiddenPolys : 1;
+	int32 bGatherTileBuildTimesHeatMap : 1;
+
+	double MaxTileBuildTime = 0;
 	
-	FRecastDebugGeometry() : bGatherPolyEdges(false), bGatherNavMeshEdges(false), bMarkForbiddenPolys(false)
+	FRecastDebugGeometry() : bGatherPolyEdges(false), bGatherNavMeshEdges(false), bMarkForbiddenPolys(false), bGatherTileBuildTimesHeatMap(false)
 	{}
 
 	uint32 NAVIGATIONSYSTEM_API GetAllocatedSize() const;
@@ -609,6 +617,12 @@ class NAVIGATIONSYSTEM_API ARecastNavMesh : public ANavigationData
 	UPROPERTY(EditAnywhere, Category=Display)
 	uint32 bDrawTileLabels:1;
 
+	UPROPERTY(EditAnywhere, Category=Display)
+	uint32 bDrawTileBuildTimes:1;
+
+	UPROPERTY(EditAnywhere, Category=Display)
+	uint32 bDrawTileBuildTimesHeatMap:1;
+	
 	/** Draw a label for every poly that indicates its poly and tile indices */
 	UPROPERTY(EditAnywhere, Category=Display, meta = (DisplayName = "Draw Polygon Indices"))
 	uint32 bDrawPolygonLabels:1;
