Navigation display improvements
Addition of tile building statistics
Addition of partitioning type, navdata count and triangle count
Addition of profile tags in the tile generation

diff --git a/Engine/Source/Runtime/NavigationSystem/Private/NavMesh/NavMeshRenderingComponent.cpp b/Engine/Source/Runtime/NavigationSystem/Private/NavMesh/NavMeshRenderingComponent.cpp
--- a/Engine/Source/Runtime/NavigationSystem/Private/NavMesh/NavMeshRenderingComponent.cpp
+++ b/Engine/Source/Runtime/NavigationSystem/Private/NavMesh/NavMeshRenderingComponent.cpp
@@ -494,9 +494,77 @@
 
 		NavMeshDrawOffset.Z = NavMesh->DrawOffset;
 
+		FRecastDebugGeometry NavMeshGeometry;
+		NavMeshGeometry.bGatherPolyEdges = FNavMeshRenderingHelpers::HasFlag(NavDetailFlags, ENavMeshDetailFlags::PolyEdges);
+		NavMeshGeometry.bGatherNavMeshEdges = FNavMeshRenderingHelpers::HasFlag(NavDetailFlags, ENavMeshDetailFlags::BoundaryEdges);
+		NavMeshGeometry.bMarkForbiddenPolys = FNavMeshRenderingHelpers::HasFlag(NavDetailFlags, ENavMeshDetailFlags::MarkForbiddenPolys);
+		NavMeshGeometry.bGatherTileBuildTimesHeatMap = FNavMeshRenderingHelpers::HasFlag(NavDetailFlags, ENavMeshDetailFlags::TileBuildTimesHeatMap);
+		const bool bGatherTileBuildTimes = FNavMeshRenderingHelpers::HasFlag(NavDetailFlags, ENavMeshDetailFlags::TileBuildTimes);
+
+#if RECAST_INTERNAL_DEBUG_DATA
+		// Colors for tile build times heat map
+		auto LerpColor = [](FColor A, FColor B, float T) -> FColor
 		{
+			return FColor(
+				FMath::RoundToInt(float(A.R) * (1.f - T) + float(B.R) * T),
+				FMath::RoundToInt(float(A.G) * (1.f - T) + float(B.G) * T),
+				FMath::RoundToInt(float(A.B) * (1.f - T) + float(B.B) * T),
+				FMath::RoundToInt(float(A.A) * (1.f - T) + float(B.A) * T));
+		};
+
+		TArray<FColor> TileBuildTimeColors;
+		const TMap<FIntPoint, FRecastInternalDebugData>* DebugDataMap = NavMesh->GetDebugDataMap();
+
+		double TotalTileBuildTime = 0.;
+		double AverageBuildTime = 0.;
+		double AverageCompLayersBuildTime = 0.;
+		double AverageNavLayersBuildTime = 0.;
+		if (bGatherTileBuildTimes || NavMeshGeometry.bGatherTileBuildTimesHeatMap)
+		{
+			QUICK_SCOPE_CYCLE_COUNTER(STAT_NavMesh_GatherDebugDrawing_MaxTileBuildTime);
+			
+			// Find min and max tile build time
+			NavMeshGeometry.MinTileBuildTime = DBL_MAX;
+			NavMeshGeometry.MaxTileBuildTime = 0.;
+			if (DebugDataMap)
+			{
+				for(const TPair<FIntPoint, FRecastInternalDebugData>& Pair : *DebugDataMap)
+				{
+					NavMeshGeometry.MaxTileBuildTime = FMath::Max(NavMeshGeometry.MaxTileBuildTime, Pair.Value.BuildTime);
+					NavMeshGeometry.MinTileBuildTime = FMath::Min(NavMeshGeometry.MinTileBuildTime, Pair.Value.BuildTime);
+					
+					AverageBuildTime += Pair.Value.BuildTime;
+					AverageCompLayersBuildTime += Pair.Value.BuildCompressedLayerTime;
+					AverageNavLayersBuildTime += Pair.Value.BuildNavigationDataTime;
+				}
+			}
+
+			TotalTileBuildTime = AverageBuildTime;
+			AverageBuildTime /= DebugDataMap->Num();
+			AverageCompLayersBuildTime /= DebugDataMap->Num();
+			AverageNavLayersBuildTime /= DebugDataMap->Num();
+		}
+
+		if(NavMeshGeometry.bGatherTileBuildTimesHeatMap)
+		{
+			TileBuildTimeColors.AddDefaulted(FRecastDebugGeometry::BuildTimeBucketsCount);
+			for (int32 Index = 0; Index < FRecastDebugGeometry::BuildTimeBucketsCount; Index++)
+			{
+				const float LerpValue = (float)Index / (FRecastDebugGeometry::BuildTimeBucketsCount-1);
+				TileBuildTimeColors[Index] = LerpColor(FColor::Blue.WithAlpha(140), FColor::Red.WithAlpha(140), LerpValue);
+			}
+		}
+#endif // RECAST_INTERNAL_DEBUG_DATA
+
+		{
 			// On screen information
 			QUICK_SCOPE_CYCLE_COUNTER(STAT_NavMesh_GatherDebugDrawing_2DLabels);
+
+			auto GetPartitioningString = [](const ERecastPartitioning::Type Type) -> FString
+			{
+				return Type == ERecastPartitioning::Monotone ? TEXT("Monotone") :
+					Type == ERecastPartitioning::Watershed ? TEXT("Watershed") : Type == ERecastPartitioning::ChunkyMonotone ? TEXT("ChunkyMonotone") : TEXT("Unknown");
+			}; 
 			
 			// Navmesh
 			const ERuntimeGenerationType Mode = NavMesh->GetRuntimeGenerationMode();
@@ -504,20 +572,24 @@
 				(Mode == ERuntimeGenerationType::Dynamic ? TEXT("Dynamic") : (Mode == ERuntimeGenerationType::DynamicModifiersOnly ? TEXT("DynamicModifersOnly") : TEXT("Unknown")));
 			DebugLabels.Add(FDebugText(FString::Printf(TEXT("%s (%s%s)"), *NavMesh->GetName(), NavMesh->bIsWorldPartitioned ? TEXT("WP ") : TEXT(""), *GenerationMode)));
 			DebugLabels.Add(FDebugText(FString::Printf(TEXT("AgentRadius %0.1f, AgentHeight %0.1f, CellSize %0.1f, CellHeight %0.1f"), NavMesh->AgentRadius, NavMesh->AgentHeight, NavMesh->CellSize, NavMesh->CellHeight)));
-			DebugLabels.Add(FDebugText(TEXT(" "))); // empty line
+			DebugLabels.Add(FDebugText(FString::Printf(TEXT("Region part %s, Layer part %s"), *GetPartitioningString(NavMesh->RegionPartitioning), *GetPartitioningString(NavMesh->LayerPartitioning))));
+			DebugLabels.Add(FDebugText(TEXT(""))); // empty line
 
 			if (NavMesh->GetActiveTiles().Num() != 0)
 			{
 				DebugLabels.Add(FDebugText(FString::Printf(TEXT("Active tiles: %i"), NavMesh->GetActiveTiles().Num())));	
 			}	
 			
-			DebugLabels.Add(FDebugText(TEXT(" "))); // empty line
-			
 			// Navigation system
 			if (const UNavigationSystemV1* NavSys = FNavigationSystem::GetCurrent<UNavigationSystemV1>(NavMesh->GetWorld()))
 			{
+				DebugLabels.Add(FDebugText(FString::Printf(TEXT("NavData count: %i"), NavSys->NavDataSet.Num())));
 				DebugLabels.Add(FDebugText(FString::Printf(TEXT("MainNavData: %s"), NavSys->MainNavData ? *NavSys->MainNavData->GetName() : TEXT("none"))));
 
+#if WITH_NAVMESH_CLUSTER_LINKS
+				DebugLabels.Add(FDebugText(FString::Printf(TEXT("Using cluster links"))));
+#endif // WITH_NAVMESH_CLUSTER_LINKS
+
 				if (NavSys->IsActiveTilesGenerationEnabled()) // Checks bGenerateNavigationOnlyAroundNavigationInvokers
 				{
 					DebugLabels.Add(FDebugText(FString::Printf(TEXT("Invoker Locations: %i"), NavSys->GetInvokerLocations().Num())));	
@@ -529,51 +601,34 @@
 				{
 					DebugLabels.Add(FDebugText(FString::Printf(TEXT("Tile jobs running/remaining: %6d / %6d"), Running, Remaining)));	
 				}
-			}
-		}
 
-		FRecastDebugGeometry NavMeshGeometry;
-		NavMeshGeometry.bGatherPolyEdges = FNavMeshRenderingHelpers::HasFlag(NavDetailFlags, ENavMeshDetailFlags::PolyEdges);
-		NavMeshGeometry.bGatherNavMeshEdges = FNavMeshRenderingHelpers::HasFlag(NavDetailFlags, ENavMeshDetailFlags::BoundaryEdges);
-		NavMeshGeometry.bMarkForbiddenPolys = FNavMeshRenderingHelpers::HasFlag(NavDetailFlags, ENavMeshDetailFlags::MarkForbiddenPolys);
-		NavMeshGeometry.bGatherTileBuildTimesHeatMap = FNavMeshRenderingHelpers::HasFlag(NavDetailFlags, ENavMeshDetailFlags::TileBuildTimesHeatMap);
-
-#if RECAST_INTERNAL_DEBUG_DATA
-		// Colors for tile build times heat map
-		auto LerpColor = [](FColor A, FColor B, float T) -> FColor
-		{
-			return FColor(
-				FMath::RoundToInt(float(A.R) * (1.f - T) + float(B.R) * T),
-				FMath::RoundToInt(float(A.G) * (1.f - T) + float(B.G) * T),
-				FMath::RoundToInt(float(A.B) * (1.f - T) + float(B.B) * T),
-				FMath::RoundToInt(float(A.A) * (1.f - T) + float(B.A) * T));
-		};
-
-		TArray<FColor> TileBuildTimeColors;
-		const TMap<FIntPoint, FRecastInternalDebugData>* DebugDataMap = NavMesh->GetDebugDataMap();
-		
-		if (NavMeshGeometry.bGatherTileBuildTimesHeatMap)
-		{
-			QUICK_SCOPE_CYCLE_COUNTER(STAT_NavMesh_GatherDebugDrawing_MaxTileBuildTime);
-			
-			// Find MaxTileBuildTime 
-			NavMeshGeometry.MaxTileBuildTime = 0;
-			if (DebugDataMap)
-			{
-				for(const TPair<FIntPoint, FRecastInternalDebugData>& Pair : *DebugDataMap)
-				{
-					NavMeshGeometry.MaxTileBuildTime = FMath::Max(NavMeshGeometry.MaxTileBuildTime, Pair.Value.BuildTime);
-				}
+				DebugLabels.Add(FDebugText(TEXT(""))); // empty line
 			}
 
-			TileBuildTimeColors.AddDefaulted(FRecastDebugGeometry::BuildTimeBucketsCount);
-			for (int32 Index = 0; Index < FRecastDebugGeometry::BuildTimeBucketsCount; Index++)
+#if RECAST_INTERNAL_DEBUG_DATA			
+			// Tile build time statistics
+			if (bGatherTileBuildTimes || NavMeshGeometry.bGatherTileBuildTimesHeatMap)
 			{
-				const float LerpValue = (float)Index / (FRecastDebugGeometry::BuildTimeBucketsCount-1);
-				TileBuildTimeColors[Index] = LerpColor(FColor::Cyan.WithAlpha(128), FColor::Red.WithAlpha(128), LerpValue);
+				DebugLabels.Add(FDebugText(FString::Printf(TEXT("Tile count: %i"), DebugDataMap->Num())));
+				DebugLabels.Add(FDebugText(FString::Printf(TEXT("Avg tile build time: %0.2f ms"), AverageBuildTime*1000.)));
+				DebugLabels.Add(FDebugText(FString::Printf(TEXT("   Avg comp layers time: %0.1f ms"), AverageCompLayersBuildTime*1000.)));
+				DebugLabels.Add(FDebugText(FString::Printf(TEXT("   Avg nav layers time: %0.1f ms"), AverageNavLayersBuildTime*1000.)));
+				DebugLabels.Add(FDebugText(FString::Printf(TEXT("Min: %0.2f ms  Max: %0.2f ms"), NavMeshGeometry.MinTileBuildTime*1000., NavMeshGeometry.MaxTileBuildTime*1000.)));
+				DebugLabels.Add(FDebugText(FString::Printf(TEXT("Total: %0.3f s"), TotalTileBuildTime)));
+
+				DebugLabels.Add(FDebugText(TEXT(""))); // empty line
+				const double TileAreaM2 = FMath::Square(NavMesh->TileSizeUU) / 10000.;
+				const double TotalAreaM2 = DebugDataMap->Num() * TileAreaM2;
+				const double TimePer100M2Ms = (TotalTileBuildTime / (TotalAreaM2/100.)) * 1000.;
+				DebugLabels.Add(FDebugText(FString::Printf(TEXT("Time per 100m2: %0.2f ms"), TimePer100M2Ms)));
+
+				const double TimePerSqKmS = TotalTileBuildTime / (TotalAreaM2/1000000.);
+				DebugLabels.Add(FDebugText(FString::Printf(TEXT("Time per km2: %0.3f s"), TimePerSqKmS)));
+
+				DebugLabels.Add(FDebugText(TEXT(""))); // empty line
 			}
+#endif // RECAST_INTERNAL_DEBUG_DATA		
 		}
-#endif // RECAST_INTERNAL_DEBUG_DATA
 		
 		const FNavDataConfig& NavConfig = NavMesh->GetConfig();
 		TArray<FColor> NavMeshColors;
@@ -715,8 +770,6 @@
 		const bool bGatherPolygonLabels = FNavMeshRenderingHelpers::HasFlag(NavDetailFlags, ENavMeshDetailFlags::PolygonLabels);
 		const bool bGatherPolygonCost = FNavMeshRenderingHelpers::HasFlag(NavDetailFlags, ENavMeshDetailFlags::PolygonCost);
 		const bool bGatherPolygonFlags = FNavMeshRenderingHelpers::HasFlag(NavDetailFlags, ENavMeshDetailFlags::PolygonFlags);
-		const bool bGatherTileBuildTimes = FNavMeshRenderingHelpers::HasFlag(NavDetailFlags, ENavMeshDetailFlags::TileBuildTimes);
-		const bool bGatherTileBuildHeatMap = FNavMeshRenderingHelpers::HasFlag(NavDetailFlags, ENavMeshDetailFlags::TileBuildTimesHeatMap);
 
 		if (bGatherTileLabels || bGatherTileBounds || bGatherPolygonLabels || bGatherPolygonCost || bGatherPolygonFlags || bGatherTileBuildTimes)
 		{
@@ -863,7 +916,19 @@
 						const double BuildTimeMs = DebugData->BuildTime * SecToMs;
 						const double CompressedLayerTimeMs = DebugData->BuildCompressedLayerTime * SecToMs;
 						const double NavigationDataTimeMs = DebugData->BuildNavigationDataTime * SecToMs;
-						DebugLabels.Add(FDebugText(Pair.Value + NavMeshDrawOffset, FString::Printf(TEXT("%.2f ms\n%.2f + %.2f"), BuildTimeMs, CompressedLayerTimeMs, NavigationDataTimeMs)));
+
+						const double Range = NavMeshGeometry.MaxTileBuildTime - NavMeshGeometry.MinTileBuildTime;
+						int32 Rank = 0;
+						if (Range != 0.)
+						{
+							Rank = FRecastDebugGeometry::BuildTimeBucketsCount * ((DebugData->BuildTime - NavMeshGeometry.MinTileBuildTime) / Range);
+							Rank = FMath::Clamp(Rank, 0, FRecastDebugGeometry::BuildTimeBucketsCount-1);
+						}
+
+						DebugLabels.Add(FDebugText(Pair.Value + NavMeshDrawOffset, FString::Printf(TEXT("%.2f ms\n%.2f comp + %.2f nav\n%i tri\nrank %i"),
+							BuildTimeMs, CompressedLayerTimeMs, NavigationDataTimeMs,
+							DebugData->TriangleCount,
+							Rank)));
 					}
 				}
 			}
@@ -931,7 +996,7 @@
 		}
 
 #if RECAST_INTERNAL_DEBUG_DATA		
-		if (bGatherTileBuildHeatMap)
+		if (NavMeshGeometry.bGatherTileBuildTimesHeatMap)
 		{
 			QUICK_SCOPE_CYCLE_COUNTER(STAT_NavMesh_GatherDebugDrawing_GatherTileBuildHeatMap);
 			
@@ -1508,7 +1573,14 @@
 		{
 			constexpr float ScreenX = 10.f;
 			Canvas->DrawText(Font, DebugText->Text, ScreenX, ScreenY);
-			ScreenY += Font->GetStringHeightSize(*DebugText->Text);
+			if (DebugText->Text.IsEmpty())
+			{
+				ScreenY += Font->GetMaxCharHeight();
+			}
+			else
+			{
+				ScreenY += Font->GetStringHeightSize(*DebugText->Text);
+			}
 		}
 		else
 		{
diff --git a/Engine/Source/Runtime/NavigationSystem/Private/NavMesh/PImplRecastNavMesh.cpp b/Engine/Source/Runtime/NavigationSystem/Private/NavMesh/PImplRecastNavMesh.cpp
--- a/Engine/Source/Runtime/NavigationSystem/Private/NavMesh/PImplRecastNavMesh.cpp
+++ b/Engine/Source/Runtime/NavigationSystem/Private/NavMesh/PImplRecastNavMesh.cpp
@@ -2824,8 +2824,13 @@ int32 FPImplRecastNavMesh::GetTilesDebugGeometry(const FRecastNavMeshGenerator*
 #if RECAST_INTERNAL_DEBUG_DATA			
 			else if (OutGeometry.bGatherTileBuildTimesHeatMap && DebugData)
 			{
-				int32 Rank = (FRecastDebugGeometry::BuildTimeBucketsCount-1) * (DebugData->BuildTime / OutGeometry.MaxTileBuildTime);
-				Rank = FMath::Clamp(Rank, 0, FRecastDebugGeometry::BuildTimeBucketsCount-1);
+				const double Range = OutGeometry.MaxTileBuildTime - OutGeometry.MinTileBuildTime;
+				int32 Rank = 0;
+				if (Range != 0.)
+				{
+					Rank = FRecastDebugGeometry::BuildTimeBucketsCount * ((DebugData->BuildTime - OutGeometry.MinTileBuildTime) / Range);
+					Rank = FMath::Clamp(Rank, 0, FRecastDebugGeometry::BuildTimeBucketsCount-1);
+				}
 				Indices = &OutGeometry.TileBuildTimesIndices[Rank];
 			}
 #endif // RECAST_INTERNAL_DEBUG_DATA
diff --git a/Engine/Source/Runtime/NavigationSystem/Private/NavMesh/RecastNavMeshGenerator.cpp b/Engine/Source/Runtime/NavigationSystem/Private/NavMesh/RecastNavMeshGenerator.cpp
--- a/Engine/Source/Runtime/NavigationSystem/Private/NavMesh/RecastNavMeshGenerator.cpp
+++ b/Engine/Source/Runtime/NavigationSystem/Private/NavMesh/RecastNavMeshGenerator.cpp
@@ -2578,12 +2578,16 @@ ETimeSliceWorkResult FRecastTileGenerator::RasterizeGeometryRecastTimeSliced(FNa
 
 		QUICK_SCOPE_CYCLE_COUNTER(STAT_Navigation_RasterizeGeomRecastRasterizeTriangles);
 
-		TInlineMaskArray::ElementType* MaskArray = RasterContext.RasterizationMasks.Num() > 0 ? RasterContext.RasterizationMasks.GetData() : nullptr;
+		const TInlineMaskArray::ElementType* MaskArray = RasterContext.RasterizationMasks.Num() > 0 ? RasterContext.RasterizationMasks.GetData() : nullptr;
 		rcRasterizeTriangles(&BuildContext,
 			Coords.GetData(), NumVerts,
 			Indices.GetData(), RasterizeGeomRecastTriAreas.GetData(), NumFaces,
 			*RasterContext.SolidHF, TileConfig.walkableClimb, RasterizationFlags, MaskArray);
 
+#if RECAST_INTERNAL_DEBUG_DATA	
+		BuildContext.InternalDebugData.TriangleCount += NumFaces;
+#endif // RECAST_INTERNAL_DEBUG_DATA
+			
 		RasterizeGeomRecastTriAreas.Reset();
 
 		//reset this so next call we start by marking walkable triangles
@@ -2632,13 +2636,17 @@ void FRecastTileGenerator::RasterizeGeometryRecast(FNavMeshBuildContext& BuildCo
 	{
 		QUICK_SCOPE_CYCLE_COUNTER(STAT_Navigation_RasterizeGeomRecastRasterizeTriangles);
 
-		TInlineMaskArray::ElementType* MaskArray = RasterContext.RasterizationMasks.Num() > 0 ? RasterContext.RasterizationMasks.GetData() : nullptr;
+		const TInlineMaskArray::ElementType* MaskArray = RasterContext.RasterizationMasks.Num() > 0 ? RasterContext.RasterizationMasks.GetData() : nullptr;
 		rcRasterizeTriangles(&BuildContext,
 			Coords.GetData(), NumVerts,
 			Indices.GetData(), RasterizeGeomRecastTriAreas.GetData(), NumFaces,
 			*RasterContext.SolidHF, TileConfig.walkableClimb, RasterizationFlags, MaskArray);
 	}
 
+#if RECAST_INTERNAL_DEBUG_DATA	
+	BuildContext.InternalDebugData.TriangleCount += NumFaces;
+#endif // RECAST_INTERNAL_DEBUG_DATA
+
 	RasterizeGeomRecastTriAreas.Reset();
 }
 
@@ -3881,7 +3889,10 @@ bool FRecastTileGenerator::GenerateNavigationData(FNavMeshBuildContext& BuildCon
 		NavigationData = MoveTemp(GenerationContext.NavigationData);
 	}
 	
-	GenerationContext.ResetIntermediateData();
+	{
+		QUICK_SCOPE_CYCLE_COUNTER(FRecastTileGenerator_GenerateNavigationData_Free);
+		GenerationContext.ResetIntermediateData();
+	}
 
 	return bGenDataLayer;
 }
diff --git a/Engine/Source/Runtime/NavigationSystem/Public/NavMesh/RecastInternalDebugData.h b/Engine/Source/Runtime/NavigationSystem/Public/NavMesh/RecastInternalDebugData.h
--- a/Engine/Source/Runtime/NavigationSystem/Public/NavMesh/RecastInternalDebugData.h
+++ b/Engine/Source/Runtime/NavigationSystem/Public/NavMesh/RecastInternalDebugData.h
@@ -28,6 +28,8 @@ struct FRecastInternalDebugData : public duDebugDraw
 	double BuildCompressedLayerTime = 0.;
 	double BuildNavigationDataTime = 0.;
 
+	uint32 TriangleCount = 0;
+	
 	FRecastInternalDebugData() {}
 	virtual ~FRecastInternalDebugData() override {}
 
diff --git a/Engine/Source/Runtime/NavigationSystem/Public/NavMesh/RecastNavMesh.h b/Engine/Source/Runtime/NavigationSystem/Public/NavMesh/RecastNavMesh.h
--- a/Engine/Source/Runtime/NavigationSystem/Public/NavMesh/RecastNavMesh.h
+++ b/Engine/Source/Runtime/NavigationSystem/Public/NavMesh/RecastNavMesh.h
@@ -246,7 +246,8 @@ struct FRecastDebugGeometry
 	int32 bMarkForbiddenPolys : 1;
 	int32 bGatherTileBuildTimesHeatMap : 1;
 
-	double MaxTileBuildTime = 0;
+	double MinTileBuildTime = DBL_MAX;
+	double MaxTileBuildTime = 0.;
 	
 	FRecastDebugGeometry() : bGatherPolyEdges(false), bGatherNavMeshEdges(false), bMarkForbiddenPolys(false), bGatherTileBuildTimesHeatMap(false)
 	{}
diff --git a/Engine/Source/Runtime/Navmesh/Private/DetourTileCache/DetourTileCacheBuilder.cpp b/Engine/Source/Runtime/Navmesh/Private/DetourTileCache/DetourTileCacheBuilder.cpp
--- a/Engine/Source/Runtime/Navmesh/Private/DetourTileCache/DetourTileCacheBuilder.cpp
+++ b/Engine/Source/Runtime/Navmesh/Private/DetourTileCache/DetourTileCacheBuilder.cpp
@@ -22,6 +22,7 @@
 #include "DetourTileCache/DetourTileCacheBuilder.h"
 #include "Detour/DetourCommon.h"
 #include "Detour/DetourAssert.h"
+#include "Stats/Stats.h"
 
 static const int MAX_VERTS_PER_POLY = 6;	// TODO: use the DT_VERTS_PER_POLYGON
 static const int MAX_REM_EDGES = 48;		// TODO: make this an expression.
@@ -38,6 +39,8 @@ dtTileCacheContourSet* dtAllocTileCacheContourSet(dtTileCacheAlloc* alloc)
 
 void dtFreeTileCacheContourSet(dtTileCacheAlloc* alloc, dtTileCacheContourSet* cset)
 {
+	QUICK_SCOPE_CYCLE_COUNTER(dtFreeTileCacheContourSet);
+	
 	dtAssert(alloc);
 
 	if (!cset) return;
@@ -60,6 +63,8 @@ dtTileCacheClusterSet* dtAllocTileCacheClusterSet(dtTileCacheAlloc* alloc)
 
 void dtFreeTileCacheClusterSet(dtTileCacheAlloc* alloc, dtTileCacheClusterSet* clusters)
 {
+	QUICK_SCOPE_CYCLE_COUNTER(dtFreeTileCacheClusterSet);
+	
 	dtAssert(alloc);
 
 	if (!clusters) return;
@@ -81,6 +86,8 @@ dtTileCachePolyMesh* dtAllocTileCachePolyMesh(dtTileCacheAlloc* alloc)
 
 void dtFreeTileCachePolyMesh(dtTileCacheAlloc* alloc, dtTileCachePolyMesh* lmesh)
 {
+	QUICK_SCOPE_CYCLE_COUNTER(dtFreeTileCachePolyMesh);
+	
 	dtAssert(alloc);
 	
 	if (!lmesh) return;
@@ -103,6 +110,8 @@ dtTileCachePolyMeshDetail* dtAllocTileCachePolyMeshDetail(dtTileCacheAlloc* allo
 
 void dtFreeTileCachePolyMeshDetail(dtTileCacheAlloc* alloc, dtTileCachePolyMeshDetail* dmesh)
 {
+	QUICK_SCOPE_CYCLE_COUNTER(dtFreeTileCachePolyMeshDetail);
+	
 	dtAssert(alloc);
 
 	if (!dmesh) return;
@@ -123,6 +132,8 @@ dtTileCacheDistanceField* dtAllocTileCacheDistanceField(dtTileCacheAlloc* alloc)
 
 void dtFreeTileCacheDistanceField(dtTileCacheAlloc* alloc, dtTileCacheDistanceField* dfield)
 {
+	QUICK_SCOPE_CYCLE_COUNTER(dtFreeTileCacheDistanceField);
+
 	dtAssert(alloc);
 
 	if (!dfield) return;
@@ -2658,6 +2669,8 @@ dtStatus dtBuildTileCacheLayer(dtTileCacheCompressor* comp,
 
 void dtFreeTileCacheLayer(dtTileCacheAlloc* alloc, dtTileCacheLayer* layer)
 {
+	QUICK_SCOPE_CYCLE_COUNTER(dtFreeTileCacheLayer);
+	
 	dtAssert(alloc);
 	// The layer is allocated as one contiguous blob of data.
 	alloc->free(layer);
diff --git a/Engine/Source/Runtime/Navmesh/Private/Recast/RecastLayers.cpp b/Engine/Source/Runtime/Navmesh/Private/Recast/RecastLayers.cpp
--- a/Engine/Source/Runtime/Navmesh/Private/Recast/RecastLayers.cpp
+++ b/Engine/Source/Runtime/Navmesh/Private/Recast/RecastLayers.cpp
@@ -20,6 +20,7 @@
 //
 
 #include "CoreMinimal.h"
+#include "Stats/Stats.h"
 #define _USE_MATH_DEFINES
 #include "Recast/Recast.h"
 #include "Recast/RecastAlloc.h"
@@ -897,6 +898,8 @@ bool rcBuildHeightfieldLayersMonotone(rcContext* ctx, rcCompactHeightfield& chf,
 							  const int borderSize, const int walkableHeight,
 							  rcHeightfieldLayerSet& lset)
 {
+	QUICK_SCOPE_CYCLE_COUNTER(STAT_Navigation_BuildHeightfieldLayersMonotone);
+	
 	rcAssert(ctx);
 	
 	ctx->startTimer(RC_TIMER_BUILD_LAYERS);
@@ -938,6 +941,8 @@ bool rcBuildHeightfieldLayersChunky(rcContext* ctx, rcCompactHeightfield& chf,
 									const int chunkSize,
 									rcHeightfieldLayerSet& lset)
 {
+	QUICK_SCOPE_CYCLE_COUNTER(STAT_Navigation_BuildHeightfieldLayersChunky);
+	
 	rcAssert(ctx);
 
 	ctx->startTimer(RC_TIMER_BUILD_LAYERS);
@@ -1107,6 +1112,8 @@ bool rcBuildHeightfieldLayers(rcContext* ctx, rcCompactHeightfield& chf,
 	const int borderSize, const int walkableHeight,
 	rcHeightfieldLayerSet& lset)
 {
+	QUICK_SCOPE_CYCLE_COUNTER(STAT_Navigation_BuildHeightfieldLayers);
+	
 	rcAssert(ctx);
 
 	ctx->startTimer(RC_TIMER_BUILD_LAYERS);
diff --git a/Engine/Source/Runtime/Navmesh/Private/Recast/RecastRegion.cpp b/Engine/Source/Runtime/Navmesh/Private/Recast/RecastRegion.cpp
--- a/Engine/Source/Runtime/Navmesh/Private/Recast/RecastRegion.cpp
+++ b/Engine/Source/Runtime/Navmesh/Private/Recast/RecastRegion.cpp
@@ -20,6 +20,7 @@
 //
 
 #include "CoreMinimal.h"
+#include "Stats/Stats.h"
 #define _USE_MATH_DEFINES
 #include "Recast/Recast.h"
 #include "Recast/RecastAlloc.h"
@@ -954,6 +955,8 @@ static bool filterSmallRegions(rcContext* ctx, int minRegionArea, int mergeRegio
 /// @see rcCompactHeightfield, rcBuildRegions, rcBuildRegionsMonotone
 bool rcBuildDistanceField(rcContext* ctx, rcCompactHeightfield& chf)
 {
+	QUICK_SCOPE_CYCLE_COUNTER(STAT_Navigation_BuildDistanceField);
+	
 	rcAssert(ctx);
 	
 	ctx->startTimer(RC_TIMER_BUILD_DISTANCEFIELD);
