Fix crash when playing in viewport while navmesh is visible
(with other minor fixes)
#rnx

diff --git a/Engine/Source/Runtime/NavigationSystem/Private/NavMesh/NavMeshRenderingComponent.cpp b/Engine/Source/Runtime/NavigationSystem/Private/NavMesh/NavMeshRenderingComponent.cpp
--- a/Engine/Source/Runtime/NavigationSystem/Private/NavMesh/NavMeshRenderingComponent.cpp
+++ b/Engine/Source/Runtime/NavigationSystem/Private/NavMesh/NavMeshRenderingComponent.cpp
@@ -540,9 +540,12 @@ void FNavMeshSceneProxyData::GatherData(const ARecastNavMesh* NavMesh, int32 InN
 			}
 
 			TotalTileBuildTime = AverageBuildTime;
-			AverageBuildTime /= DebugDataMap->Num();
-			AverageCompLayersBuildTime /= DebugDataMap->Num();
-			AverageNavLayersBuildTime /= DebugDataMap->Num();
+			if (DebugDataMap->Num() != 0)
+			{
+				AverageBuildTime /= DebugDataMap->Num();
+				AverageCompLayersBuildTime /= DebugDataMap->Num();
+				AverageNavLayersBuildTime /= DebugDataMap->Num();
+			}
 		}
 
 		if(NavMeshGeometry.bGatherTileBuildTimesHeatMap)
@@ -575,7 +578,7 @@ void FNavMeshSceneProxyData::GatherData(const ARecastNavMesh* NavMesh, int32 InN
 			DebugLabels.Add(FDebugText(FString::Printf(TEXT("Region part %s, Layer part %s"), *GetPartitioningString(NavMesh->RegionPartitioning), *GetPartitioningString(NavMesh->LayerPartitioning))));
 			DebugLabels.Add(FDebugText(TEXT(""))); // empty line
 
-			if (NavMesh->GetActiveTiles().Num() != 0)
+			if (NavMesh->GetGenerator() && NavMesh->GetActiveTiles().Num() != 0)
 			{
 				DebugLabels.Add(FDebugText(FString::Printf(TEXT("Active tiles: %i"), NavMesh->GetActiveTiles().Num())));	
 			}	
@@ -925,10 +928,9 @@ void FNavMeshSceneProxyData::GatherData(const ARecastNavMesh* NavMesh, int32 InN
 							Rank = FMath::Clamp(Rank, 0, FRecastDebugGeometry::BuildTimeBucketsCount-1);
 						}
 
-						DebugLabels.Add(FDebugText(Pair.Value + NavMeshDrawOffset, FString::Printf(TEXT("%.2f ms\n%.2f comp + %.2f nav\n%i tri\nrank %i"),
+						DebugLabels.Add(FDebugText(Pair.Value + NavMeshDrawOffset, FString::Printf(TEXT("%.2f ms\n%.2f comp + %.2f nav\n%i tri"),
 							BuildTimeMs, CompressedLayerTimeMs, NavigationDataTimeMs,
-							DebugData->TriangleCount,
-							Rank)));
+							DebugData->TriangleCount)));
 					}
 				}
 			}
