[Navmesh] Keep the same Oodle compression in editor as in game to (avoid being slower in PIE)

diff --git a/Engine/Source/Runtime/NavigationSystem/Private/NavMesh/RecastNavMeshGenerator.cpp b/Engine/Source/Runtime/NavigationSystem/Private/NavMesh/RecastNavMeshGenerator.cpp
--- a/Engine/Source/Runtime/NavigationSystem/Private/NavMesh/RecastNavMeshGenerator.cpp
+++ b/Engine/Source/Runtime/NavigationSystem/Private/NavMesh/RecastNavMeshGenerator.cpp
@@ -79,13 +79,8 @@ namespace UE::NavMesh::Private
 	static FAutoConsoleVariableRef CVarRecentlyBuildTileDisplayTime(TEXT("n.RecentlyBuildTileDisplayTime"), RecentlyBuildTileDisplayTime, TEXT("Time (in seconds) to display tiles that have recently been built."), ECVF_Default);
 }
 
-#if	WITH_EDITOR	
-static FOodleDataCompression::ECompressor GNavmeshTileCacheCompressor = FOodleDataCompression::ECompressor::Kraken;
-static FOodleDataCompression::ECompressionLevel GNavmeshTileCacheCompressionLevel = FOodleDataCompression::ECompressionLevel::Optimal1;
-#else
 static FOodleDataCompression::ECompressor GNavmeshTileCacheCompressor = FOodleDataCompression::ECompressor::Mermaid;
 static FOodleDataCompression::ECompressionLevel GNavmeshTileCacheCompressionLevel = FOodleDataCompression::ECompressionLevel::HyperFast1;
-#endif
 
 FORCEINLINE bool DoesBoxContainOrOverlapVector(const FBox& BigBox, const FVector& In)
 {
