Fix misleading local variable name.

diff --git a/Engine/Source/Runtime/NavigationSystem/Private/NavMesh/RecastNavMeshGenerator.cpp b/Engine/Source/Runtime/NavigationSystem/Private/NavMesh/RecastNavMeshGenerator.cpp
--- a/Engine/Source/Runtime/NavigationSystem/Private/NavMesh/RecastNavMeshGenerator.cpp
+++ b/Engine/Source/Runtime/NavigationSystem/Private/NavMesh/RecastNavMeshGenerator.cpp
@@ -1662,11 +1662,11 @@ void FRecastTileGenerator::Setup(const FRecastNavMeshGenerator& ParentGenerator,
 {
 	const FVector RcNavMeshOrigin = ParentGenerator.GetRcNavMeshOrigin();
 	const FBox NavTotalBounds = ParentGenerator.GetTotalBounds();
-	const FVector::FReal TileCellSize = (TileConfig.tileSize * TileConfig.cs);
+	const FVector::FReal TileSizeUU = (TileConfig.tileSize * TileConfig.cs);
 
 	NavDataConfig = ParentGenerator.GetOwner()->GetConfig();
 
-	TileBB = CalculateTileBounds(TileX, TileY, RcNavMeshOrigin, NavTotalBounds, TileCellSize);
+	TileBB = CalculateTileBounds(TileX, TileY, RcNavMeshOrigin, NavTotalBounds, TileSizeUU);
 	TileBBExpandedForAgent = TileBB.ExpandBy(NavDataConfig.AgentRadius * 2 + TileConfig.cs);
 	const FBox RCBox = Unreal2RecastBox(TileBB);
 	FVector Min32(RCBox.Min), Max32(RCBox.Max);
