add more debug logs when building tiles

diff --git a/Engine/Source/Runtime/NavigationSystem/Private/NavMesh/RecastNavMesh.cpp b/Engine/Source/Runtime/NavigationSystem/Private/NavMesh/RecastNavMesh.cpp
--- a/Engine/Source/Runtime/NavigationSystem/Private/NavMesh/RecastNavMesh.cpp
+++ b/Engine/Source/Runtime/NavigationSystem/Private/NavMesh/RecastNavMesh.cpp
@@ -3372,6 +3372,8 @@ void ARecastNavMesh::UpdateActiveTiles(const TArray<FNavigationInvokerRaw>& Invo
 		}
 	}
 
+	UE_VLOG(this, LogNavigation, Log, TEXT("Updating active tiles: %d to remove, %d to update"), TilesToRemove.Num(), TilesToUpdate.Num());
+
 	RemoveTiles(TilesToRemove);
 	RebuildTile(TilesToUpdate);
 
