Avoid bounds fully included in the Max tile computation for navmesh generation

diff --git a/Engine/Source/Runtime/NavigationSystem/Private/NavMesh/RecastNavMeshGenerator.cpp b/Engine/Source/Runtime/NavigationSystem/Private/NavMesh/RecastNavMeshGenerator.cpp
--- a/Engine/Source/Runtime/NavigationSystem/Private/NavMesh/RecastNavMeshGenerator.cpp
+++ b/Engine/Source/Runtime/NavigationSystem/Private/NavMesh/RecastNavMeshGenerator.cpp
@@ -4397,8 +4397,34 @@ void CheckTileIndicesInValidRange(const TNavStatArray<FBox>& NavigableAreas, con
 int32 CalculateMaxTilesCount(const TNavStatArray<FBox>& NavigableAreas, FVector::FReal TileSizeInWorldUnits, FVector::FReal AvgLayersPerGridCell, const uint32 NavMeshVersion)
 {
 	int32 GridCellsCount = 0;
-	for (const FBox& AreaBounds : NavigableAreas)
+	for (int32 Index = 0; Index < NavigableAreas.Num(); ++Index)
 	{
+		const FBox& AreaBounds = NavigableAreas[Index];
+		if (NavMeshVersion >= NAVMESHVER_MAXTILES_COUNT_SKIP_INCLUSION)
+		{
+			bool bIsInsideAnotherArea = false;
+			for (int32 OtherIndex = 0; OtherIndex < NavigableAreas.Num(); ++OtherIndex)
+			{
+				if (Index == OtherIndex)
+				{
+					continue;
+				}
+				const FBox& OtherBox = NavigableAreas[OtherIndex];
+				if (OtherBox.IsInsideXY(AreaBounds))
+				{
+					// The current Area Bounds is fully contained in another Area, so we don't need to count it.
+					// This doesn't take into account partial overlaps
+					bIsInsideAnotherArea = true;
+					break;
+				}
+			}
+
+			if (bIsInsideAnotherArea)
+			{
+				continue;
+			}
+		}
+
 		// TODO: need more precise calculation, currently we don't take into account that volumes can be overlapped
 		const FBox RCBox = Unreal2RecastBox(AreaBounds);
 
@@ -4938,6 +4964,7 @@ void FRecastNavMeshGenerator::OnNavigationBoundsChanged()
 		const int32 MaxRequestedTiles = UE::NavMesh::Private::CalculateMaxTilesCount(InclusionBounds, Config.tileSize * Config.cs, AvgLayersPerTile, DestNavMesh->NavMeshVersion);
 		if (DetourMesh->getMaxTiles() != MaxRequestedTiles)
 		{
+			UE_LOG(LogNavigation, Log, TEXT("%s> Navigation bounds changed, rebuilding navmesh"), *DestNavMesh->GetName());
 			// Destroy current NavMesh
 			DestNavMesh->GetRecastNavMeshImpl()->SetRecastMesh(nullptr);
 
diff --git a/Engine/Source/Runtime/NavigationSystem/Public/NavMesh/RecastVersion.h b/Engine/Source/Runtime/NavigationSystem/Public/NavMesh/RecastVersion.h
--- a/Engine/Source/Runtime/NavigationSystem/Public/NavMesh/RecastVersion.h
+++ b/Engine/Source/Runtime/NavigationSystem/Public/NavMesh/RecastVersion.h
@@ -20,6 +20,7 @@
 #define NAVMESHVER_MAXTILES_COUNT_CHANGE 		19
 #define NAVMESHVER_LWCOORDS_OPTIMIZATION		20
 #define NAVMESHVER_OPTIM_FIX_SERIALIZE_PARAMS	21 // Fix, serialize params that used to be in the tile and are now in the navmesh.
+#define NAVMESHVER_MAXTILES_COUNT_SKIP_INCLUSION 22
 
-#define NAVMESHVER_LATEST				NAVMESHVER_OPTIM_FIX_SERIALIZE_PARAMS
+#define NAVMESHVER_LATEST				NAVMESHVER_MAXTILES_COUNT_SKIP_INCLUSION
 #define NAVMESHVER_MIN_COMPATIBLE		NAVMESHVER_LWCOORDS_OPTIMIZATION
