[Navmesh] Optimization: remove unnecessary extra agent radius tile expansion when computing tile bounds.
This avoids unnecessary collision gathering as well a preventing unnecessary tile dirtiness.
Also resolves issue in Z where an NavModifierVolume would incorrectly dirty a tile below itself if within agent radius.
Addition of n.UseTightBoundExpansion console command to revert back to previous behavior as a fallback precaution.

diff --git a/Engine/Source/Runtime/NavigationSystem/Private/NavMesh/RecastNavMeshGenerator.cpp b/Engine/Source/Runtime/NavigationSystem/Private/NavMesh/RecastNavMeshGenerator.cpp
--- a/Engine/Source/Runtime/NavigationSystem/Private/NavMesh/RecastNavMeshGenerator.cpp
+++ b/Engine/Source/Runtime/NavigationSystem/Private/NavMesh/RecastNavMeshGenerator.cpp
@@ -77,6 +77,9 @@ namespace UE::NavMesh::Private
 {
 	static float RecentlyBuildTileDisplayTime = 0.2f;
 	static FAutoConsoleVariableRef CVarRecentlyBuildTileDisplayTime(TEXT("n.RecentlyBuildTileDisplayTime"), RecentlyBuildTileDisplayTime, TEXT("Time (in seconds) to display tiles that have recently been built."), ECVF_Default);
+
+	static bool bUseTightBoundExpansion = true;
+	static FAutoConsoleVariableRef CVarUseTightBoundExpansion(TEXT("n.UseTightBoundExpansion"), bUseTightBoundExpansion, TEXT("Active by default. Use an expansion of one AgentRadius. Set to false to revert to the previous behavior (2 AgentRadius)."), ECVF_Default);
 }
 
 static FOodleDataCompression::ECompressor GNavmeshTileCacheCompressor = FOodleDataCompression::ECompressor::Mermaid;
@@ -1667,9 +1670,20 @@ void FRecastTileGenerator::Setup(const FRecastNavMeshGenerator& ParentGenerator,
 	NavDataConfig = ParentGenerator.GetOwner()->GetConfig();
 
 	TileBB = CalculateTileBounds(TileX, TileY, RcNavMeshOrigin, NavTotalBounds, TileSizeUU);
-	TileBBExpandedForAgent = TileBB.ExpandBy(NavDataConfig.AgentRadius * 2 + TileConfig.cs);
+
+	if (UE::NavMesh::Private::bUseTightBoundExpansion)
+	{
+		TileBBExpandedForAgent = ParentGenerator.GrowBoundingBox(TileBB, /*bIncludeAgentHeight*/ false);
+	}
+	else
+	{
+		// Deprecated
+		TileBBExpandedForAgent = TileBB.ExpandBy(NavDataConfig.AgentRadius * 2 + TileConfig.cs);
+	}
+	
 	const FBox RCBox = Unreal2RecastBox(TileBB);
-	FVector Min32(RCBox.Min), Max32(RCBox.Max);
+	const FVector Min32(RCBox.Min);
+	const FVector Max32(RCBox.Max);
 	rcVcopy(TileConfig.bmin, &Min32.X);
 	rcVcopy(TileConfig.bmax, &Max32.X);
 			
@@ -1705,7 +1719,7 @@ void FRecastTileGenerator::Setup(const FRecastNavMeshGenerator& ParentGenerator,
 	// We have to regenerate layers data in case geometry is changed or tile cache is missing
 	bRegenerateCompressedLayers = (bGeometryChanged || CompressedLayers.Num() == 0);
 	
-	// Gather geometry for tile if it inside navigable bounds
+	// Gather geometry for tile if it's inside navigable bounds
 	if (InclusionBounds.Num())
 	{
 		if (!bRegenerateCompressedLayers)
@@ -4556,7 +4570,16 @@ void FRecastNavMeshGenerator::Init()
 
 	ConfigureBuildProperties(Config);
 
-	BBoxGrowth = FVector(2.0f * Config.borderSize * Config.cs);
+	if (UE::NavMesh::Private::bUseTightBoundExpansion)
+	{
+		const rcReal HorizontalGrowth = Config.borderSize * Config.cs;
+		BBoxGrowth = FVector(HorizontalGrowth, HorizontalGrowth, Config.cs);
+	}
+	else
+	{
+		// Deprecated
+		BBoxGrowth = FVector(2.0f * Config.borderSize * Config.cs);
+	}
 	RcNavMeshOrigin = Unreal2RecastPoint(DestNavMesh->NavMeshOriginOffset);
 	
 	AdditionalCachedData = FRecastNavMeshCachedData::Construct(DestNavMesh);
diff --git a/Engine/Source/Runtime/NavigationSystem/Public/NavMesh/RecastNavMeshGenerator.h b/Engine/Source/Runtime/NavigationSystem/Public/NavMesh/RecastNavMeshGenerator.h
--- a/Engine/Source/Runtime/NavigationSystem/Public/NavMesh/RecastNavMeshGenerator.h
+++ b/Engine/Source/Runtime/NavigationSystem/Public/NavMesh/RecastNavMeshGenerator.h
@@ -478,7 +478,7 @@ protected:
 	/** Tile's bounding box, Unreal coords */
 	FBox TileBB;
 
-	/** Tile's bounding box expanded by Agent's radius *2 + CellSize, Unreal coords */
+	/** Tile's bounding box expanded by Agent's radius horizontally, Unreal coords */
 	FBox TileBBExpandedForAgent;
 	
 	/** Layers dirty flags */
