Recast debug improvement :
- add way to generate only the debug tile
- add way to visualize the geometry received for navmesh generation

diff --git a/Engine/Source/Runtime/NavigationSystem/Private/NavMesh/RecastNavMesh.cpp b/Engine/Source/Runtime/NavigationSystem/Private/NavMesh/RecastNavMesh.cpp
--- a/Engine/Source/Runtime/NavigationSystem/Private/NavMesh/RecastNavMesh.cpp
+++ b/Engine/Source/Runtime/NavigationSystem/Private/NavMesh/RecastNavMesh.cpp
@@ -452,6 +452,8 @@ FRecastNavMeshTileGenerationDebug::FRecastNavMeshTileGenerationDebug()
 {
 	bEnabled = false;
 	TileCoordinate = FIntVector::ZeroValue;
+	bGenerateDebugTileOnly = false;
+	bCollisionGeometry = false;
 	bHeightfieldSolidFromRasterization = false;
 	bHeightfieldSolidPostInclusionBoundsFiltering = false;
 	bHeightfieldSolidPostHeightFiltering = false;
diff --git a/Engine/Source/Runtime/NavigationSystem/Private/NavMesh/RecastNavMeshGenerator.cpp b/Engine/Source/Runtime/NavigationSystem/Private/NavMesh/RecastNavMeshGenerator.cpp
--- a/Engine/Source/Runtime/NavigationSystem/Private/NavMesh/RecastNavMeshGenerator.cpp
+++ b/Engine/Source/Runtime/NavigationSystem/Private/NavMesh/RecastNavMeshGenerator.cpp
@@ -65,8 +65,10 @@ static FAutoConsoleVariableRef NavmeshVarSynchronous(TEXT("n.GNavmeshSynchronous
 #if RECAST_INTERNAL_DEBUG_DATA
 static int32 GNavmeshDebugTileX = MAX_int32;
 static int32 GNavmeshDebugTileY = MAX_int32;
+static bool GNavmeshGenerateDebugTileOnly = false;
 static FAutoConsoleVariableRef NavmeshVarDebugTileX(TEXT("n.GNavmeshDebugTileX"), GNavmeshDebugTileX, TEXT(""), ECVF_Default);
 static FAutoConsoleVariableRef NavmeshVarDebugTileY(TEXT("n.GNavmeshDebugTileY"), GNavmeshDebugTileY, TEXT(""), ECVF_Default);
+static FAutoConsoleVariableRef NavmeshVarGenerateDebugTileOnly(TEXT("n.GNavmeshGenerateDebugTileOnly"), GNavmeshGenerateDebugTileOnly, TEXT(""), ECVF_Default);
 #endif //RECAST_INTERNAL_DEBUG_DATA
 
 // Hotfixing this flag without rebuilding the data will cause decompression errors, equivalent to not having prebuilt navmesh data at all.
@@ -1993,12 +1995,16 @@ void FRecastTileGenerator::GatherNavigationDataGeometry(const TSharedRef<FNaviga
 	bool bDumpGeometryData = false;
 
 #if RECAST_INTERNAL_DEBUG_DATA
+	if (!IsTileDebugAllowingGeneration())
+	{
+		return;
+	}
 	if (IsTileDebugActive())
 	{
 		UE_LOG(LogNavigation, Log, TEXT("Gathering geometry for tile (%i,%i): %s."), TileX, TileY, *GetFullNameSafe(ElementData->GetOwner()));
 		UE_LOG(LogNavigation, Log, TEXT("                       bounds: %s"), *ElementData->Bounds.ToString());
 	}
-#endif
+#endif // RECAST_INTERNAL_DEBUG_DATA
 
 	if (ElementData->IsPendingLazyGeometryGathering() || ElementData->NeedAnyPendingLazyModifiersGathering())
 	{
@@ -2517,6 +2523,13 @@ private:
 
 bool FRecastTileGenerator::CreateHeightField(FNavMeshBuildContext& BuildContext, FTileRasterizationContext& RasterContext)
 {
+#if RECAST_INTERNAL_DEBUG_DATA
+	if (!IsTileDebugAllowingGeneration())
+	{
+		return false;
+	}
+#endif // RECAST_INTERNAL_DEBUG_DATA
+
 	SCOPE_CYCLE_COUNTER(STAT_Navigation_RecastCreateHeightField);
 
 	TileConfig.width = TileConfig.tileSize + TileConfig.borderSize * 2;
@@ -2652,7 +2665,18 @@ void FRecastTileGenerator::RasterizeGeometryRecast(FNavMeshBuildContext& BuildCo
 			*RasterContext.SolidHF, TileConfig.walkableClimb, RasterizationFlags, MaskArray);
 	}
 
-#if RECAST_INTERNAL_DEBUG_DATA	
+#if RECAST_INTERNAL_DEBUG_DATA
+	if (IsTileDebugActive() && TileDebugSettings.bCollisionGeometry) 
+	{
+		TArray<FVector::FReal> Normals;
+		Normals.AddZeroed(NumFaces*3);
+
+		rcCalcTriNormals(Coords.GetData(), NumVerts, Indices.GetData(), NumFaces, Normals.GetData());
+
+		constexpr FVector::FReal TextureScale = 1.;
+		duDebugDrawTriMesh(&BuildContext.InternalDebugData, Coords.GetData(), NumVerts, Indices.GetData(), Normals.GetData(), NumFaces, RasterizeGeomRecastTriAreas.GetData(), TextureScale);
+	}
+
 	BuildContext.InternalDebugData.TriangleCount += NumFaces;
 #endif // RECAST_INTERNAL_DEBUG_DATA
 
@@ -7138,9 +7162,26 @@ public:
 #endif // WITH_RECAST
 
 #if RECAST_INTERNAL_DEBUG_DATA
-bool FRecastTileGenerator::IsTileDebugActive()
+bool FRecastTileGenerator::IsTileDebugActive() const
 {
-	const FIntVector& Coord = TileDebugSettings.TileCoordinate; 
+	const FIntVector& Coord = TileDebugSettings.TileCoordinate;
 	return (TileDebugSettings.bEnabled && TileX == Coord.X && TileY == Coord.Y) || (TileX == GNavmeshDebugTileX && TileY == GNavmeshDebugTileY);
 }
+
+bool FRecastTileGenerator::IsTileDebugAllowingGeneration() const
+{
+	const FIntVector& Coord = TileDebugSettings.TileCoordinate;
+	if (TileDebugSettings.bEnabled && TileDebugSettings.bGenerateDebugTileOnly)
+	{
+		return TileX == Coord.X && TileY == Coord.Y;
+	}
+	else if (GNavmeshGenerateDebugTileOnly)
+	{
+		return TileX == GNavmeshDebugTileX && TileY == GNavmeshDebugTileY;
+	}
+	else
+	{
+		return true;
+	}
+}
 #endif //RECAST_INTERNAL_DEBUG_DATA
diff --git a/Engine/Source/Runtime/NavigationSystem/Public/NavMesh/RecastNavMesh.h b/Engine/Source/Runtime/NavigationSystem/Public/NavMesh/RecastNavMesh.h
--- a/Engine/Source/Runtime/NavigationSystem/Public/NavMesh/RecastNavMesh.h
+++ b/Engine/Source/Runtime/NavigationSystem/Public/NavMesh/RecastNavMesh.h
@@ -446,6 +446,15 @@ struct NAVIGATIONSYSTEM_API FRecastNavMeshTileGenerationDebug
 	UPROPERTY(EditAnywhere, Category = Debug)
 	FIntVector TileCoordinate = FIntVector::ZeroValue;
 
+	/** If set, the generator will only generate the tile selected to debug (set in TileCoordinate).*/
+	UPROPERTY(EditAnywhere, Category = Debug)
+	uint32 bGenerateDebugTileOnly : 1;
+
+	/** Display the collision used for the navmesh rasterization.
+	 * Note: The visualization is affected by the DrawOffset of the RecastNavmesh owner*/
+	UPROPERTY(EditAnywhere, Category = Debug)
+	uint32 bCollisionGeometry : 1;
+
 	UPROPERTY(EditAnywhere, Category = Debug)
 	uint32 bHeightfieldSolidFromRasterization : 1;
 
diff --git a/Engine/Source/Runtime/NavigationSystem/Public/NavMesh/RecastNavMeshGenerator.h b/Engine/Source/Runtime/NavigationSystem/Public/NavMesh/RecastNavMeshGenerator.h
--- a/Engine/Source/Runtime/NavigationSystem/Public/NavMesh/RecastNavMeshGenerator.h
+++ b/Engine/Source/Runtime/NavigationSystem/Public/NavMesh/RecastNavMeshGenerator.h
@@ -440,7 +440,8 @@ protected:
 	void DumpAsyncData();
 
 #if RECAST_INTERNAL_DEBUG_DATA
-	bool IsTileDebugActive();
+	bool IsTileDebugActive() const;
+	bool IsTileDebugAllowingGeneration() const;
 #endif
 
 protected:
diff --git a/Engine/Source/Runtime/Navmesh/Private/Recast/Recast.cpp b/Engine/Source/Runtime/Navmesh/Private/Recast/Recast.cpp
--- a/Engine/Source/Runtime/Navmesh/Private/Recast/Recast.cpp
+++ b/Engine/Source/Runtime/Navmesh/Private/Recast/Recast.cpp
@@ -321,6 +321,21 @@ static void calcTriNormal(const rcReal* v0, const rcReal* v1, const rcReal* v2,
 	rcVnormalize(norm);
 }
 
+//@UE BEGIN
+void rcCalcTriNormals(const rcReal* verts, const int nv, const int* tris, const int nt, rcReal* norms)
+{
+	for (int i = 0; i < nt; ++i)
+	{
+		const int* tri = &tris[i*3];
+		const rcReal* v0 = &verts[tri[0]*3];
+		const rcReal* v1 = &verts[tri[1]*3];
+		const rcReal* v2 = &verts[tri[2]*3];
+
+		calcTriNormal(v0, v1, v2, &norms[i*3]);
+	}
+}
+//@UE END
+
 /// @par
 ///
 /// Only sets the aread id's for the walkable triangles.  Does not alter the
diff --git a/Engine/Source/Runtime/Navmesh/Public/Recast/Recast.h b/Engine/Source/Runtime/Navmesh/Public/Recast/Recast.h
--- a/Engine/Source/Runtime/Navmesh/Public/Recast/Recast.h
+++ b/Engine/Source/Runtime/Navmesh/Public/Recast/Recast.h
@@ -876,6 +876,17 @@ inline void rcVnormalize(rcReal* v)
 	v[2] *= d;
 }
 
+//@UE BEGIN
+/// Calculates the normals of each triangles in an array
+///  @ingroup recast
+///  @param[in]		verts	An array of vertices. [(x, y, z) * @p nv]
+///  @param[in]		nv		The number of vertices in the @p verts array.
+///  @param[in]		tris	The triangle vertex indices. [(vertA, vertB, vertC) * @p nt]
+///  @param[in]		nt		The number of triangles.
+///  @param[out]	norms	The normal vector of each triangle. [(x, y, z) * @p nt]
+NAVMESH_API void rcCalcTriNormals(const rcReal* verts, const int nv, const int* tris, const int nt, rcReal* norms);
+//@UE END
+
 /// @}
 /// @name Heightfield Functions
 /// @see rcHeightfield
