Addressing the remaining AI numeric conversion issues.

diff --git a/Engine/Plugins/AI/EnvironmentQueryEditor/Source/EnvironmentQueryEditor/Private/DetailCustomizations/EnvQueryTestDetails.cpp b/Engine/Plugins/AI/EnvironmentQueryEditor/Source/EnvironmentQueryEditor/Private/DetailCustomizations/EnvQueryTestDetails.cpp
--- a/Engine/Plugins/AI/EnvironmentQueryEditor/Source/EnvironmentQueryEditor/Private/DetailCustomizations/EnvQueryTestDetails.cpp
+++ b/Engine/Plugins/AI/EnvironmentQueryEditor/Source/EnvironmentQueryEditor/Private/DetailCustomizations/EnvQueryTestDetails.cpp
@@ -317,13 +317,13 @@ void FEnvQueryTestDetails::BuildScoreEquationValues()
 
 void FEnvQueryTestDetails::OnFilterTestChange(int32 Index)
 {
-	uint8 EnumValue = Index;
+	const uint8 EnumValue = IntCastChecked<uint8>(Index);
 	FilterTypeHandle->SetValue(EnumValue);
 }
 
 void FEnvQueryTestDetails::OnScoreEquationChange(int32 Index)
 {
-	uint8 EnumValue = Index;
+	const uint8 EnumValue = IntCastChecked<uint8>(Index);
 	ScoreEquationHandle->SetValue(EnumValue);
 }
 
@@ -361,14 +361,14 @@ void FEnvQueryTestDetails::BuildScoreClampingTypeValues(bool bBuildMinValues, TA
 void FEnvQueryTestDetails::OnClampMinTestChange(int32 Index)
 {
 	check(ClampMinTypeHandle.IsValid());
-	uint8 EnumValue = Index;
+	const uint8 EnumValue = IntCastChecked<uint8>(Index);
 	ClampMinTypeHandle->SetValue(EnumValue);
 }
 
 void FEnvQueryTestDetails::OnClampMaxTestChange(int32 Index)
 {
 	check(ClampMaxTypeHandle.IsValid());
-	uint8 EnumValue = Index;
+	const uint8 EnumValue = IntCastChecked<uint8>(Index);
 	ClampMaxTypeHandle->SetValue(EnumValue);
 }
 
diff --git a/Engine/Plugins/AI/EnvironmentQueryEditor/Source/EnvironmentQueryEditor/Private/DetailCustomizations/EnvTraceDataCustomization.cpp b/Engine/Plugins/AI/EnvironmentQueryEditor/Source/EnvironmentQueryEditor/Private/DetailCustomizations/EnvTraceDataCustomization.cpp
--- a/Engine/Plugins/AI/EnvironmentQueryEditor/Source/EnvironmentQueryEditor/Private/DetailCustomizations/EnvTraceDataCustomization.cpp
+++ b/Engine/Plugins/AI/EnvironmentQueryEditor/Source/EnvironmentQueryEditor/Private/DetailCustomizations/EnvTraceDataCustomization.cpp
@@ -173,14 +173,14 @@ void FEnvTraceDataCustomization::OnTraceChannelChanged()
 		ETraceTypeQuery TraceTypeValue = (ETraceTypeQuery)TraceChannelValue;
 		ECollisionChannel CollisionChannelValue = UEngineTypes::ConvertToCollisionChannel(TraceTypeValue);
 
-		uint8 SerializedChannelValue = CollisionChannelValue;
+		uint8 SerializedChannelValue = IntCastChecked<uint8>((int)CollisionChannelValue);
 		PropTraceChannelSerialized->SetValue(SerializedChannelValue);
 	}
 }
 
 void FEnvTraceDataCustomization::OnTraceModeChanged(int32 Index)
 {
-	ActiveMode = Index;
+	ActiveMode = IntCastChecked<uint8>(Index);
 	PropTraceMode->SetValue(ActiveMode);
 }
 
diff --git a/Engine/Plugins/AI/EnvironmentQueryEditor/Source/EnvironmentQueryEditor/Private/SEnvQueryProfiler.h b/Engine/Plugins/AI/EnvironmentQueryEditor/Source/EnvironmentQueryEditor/Private/SEnvQueryProfiler.h
--- a/Engine/Plugins/AI/EnvironmentQueryEditor/Source/EnvironmentQueryEditor/Private/SEnvQueryProfiler.h
+++ b/Engine/Plugins/AI/EnvironmentQueryEditor/Source/EnvironmentQueryEditor/Private/SEnvQueryProfiler.h
@@ -21,8 +21,8 @@ public:
 	struct FStatData : public TSharedFromThis<FStatData, ESPMode::ThreadSafe>
 	{
 		FName QueryName;
-		float WorstTime;
-		float AverageTime;
+		double WorstTime;
+		double AverageTime;
 		float AverageLoad;
 		int32 NumRuns;
 		bool bIsHighlighted;
diff --git a/Engine/Plugins/AI/EnvironmentQueryEditor/Source/EnvironmentQueryEditor/Private/SGraphNode_EnvironmentQuery.cpp b/Engine/Plugins/AI/EnvironmentQueryEditor/Source/EnvironmentQueryEditor/Private/SGraphNode_EnvironmentQuery.cpp
--- a/Engine/Plugins/AI/EnvironmentQueryEditor/Source/EnvironmentQueryEditor/Private/SGraphNode_EnvironmentQuery.cpp
+++ b/Engine/Plugins/AI/EnvironmentQueryEditor/Source/EnvironmentQueryEditor/Private/SGraphNode_EnvironmentQuery.cpp
@@ -238,7 +238,7 @@ void SGraphNode_EnvironmentQuery::UpdateGraphNode()
 									.AutoHeight()
 									.Padding(0,0,0,2)
 									[
-										SNew(SBox).HeightOverride(4)
+										SNew(SBox).HeightOverride(4.f)
 										[
 											// weight bar
 											SNew(SProgressBar)
diff --git a/Engine/Plugins/AI/HTNPlanner/Source/HTNPlanner/Private/HTNBuilder.cpp b/Engine/Plugins/AI/HTNPlanner/Source/HTNPlanner/Private/HTNBuilder.cpp
--- a/Engine/Plugins/AI/HTNPlanner/Source/HTNPlanner/Private/HTNBuilder.cpp
+++ b/Engine/Plugins/AI/HTNPlanner/Source/HTNPlanner/Private/HTNBuilder.cpp
@@ -76,10 +76,10 @@ bool FHTNBuilder_Domain::Compile()
 
 		Domain.TaskNameMap.Add(It->Key, (CurrentAddress - CompiledData) / sizeof(uint8));
 
-		Task->ActionID = PrimitiveTaskBuilder.ActionID;
+		Task->ActionID = IntCastChecked<FHTNPolicy::FActionID>(PrimitiveTaskBuilder.ActionID);
 		Task->Parameter = PrimitiveTaskBuilder.Parameter;
 		const int32 EffectsNum = PrimitiveTaskBuilder.Effects.Num();
-		Task->EffectsCount = EffectsNum;
+		Task->EffectsCount = IntCastChecked<uint8>(EffectsNum);
 
 		CurrentAddress += sizeof(FHTNPrimitiveTask);
 
@@ -101,7 +101,7 @@ bool FHTNBuilder_Domain::Compile()
 		Domain.TaskNameMap.Add(It->Key, (CurrentAddress - CompiledData) / sizeof(uint8));
 
 		const int32 MethodsCount = CompositeTaskBuilder.Methods.Num();
-		Task->MethodsCount = MethodsCount;
+		Task->MethodsCount = IntCastChecked<uint8>(MethodsCount);
 
 		CurrentAddress += sizeof(FHTNCompositeTask);
 
@@ -116,8 +116,8 @@ bool FHTNBuilder_Domain::Compile()
 				const FHTNBuilder_Method& MethodBuilder = CompositeTaskBuilder.Methods[MethodIndex];
 
 				FHTNMethod* Method = new (&Task->Methods[MethodIndex]) FHTNMethod();
-				Method->ConditionsCount = MethodBuilder.Conditions.Num();
-				Method->TasksCount = MethodBuilder.Tasks.Num();
+				Method->ConditionsCount = IntCastChecked<uint8>(MethodBuilder.Conditions.Num());
+				Method->TasksCount = IntCastChecked<uint8>(MethodBuilder.Tasks.Num());
 
 				if (Method->ConditionsCount > 0)
 				{
diff --git a/Engine/Plugins/AI/HTNPlanner/Source/HTNTestSuite/Private/HTNTest.cpp b/Engine/Plugins/AI/HTNPlanner/Source/HTNTestSuite/Private/HTNTest.cpp
--- a/Engine/Plugins/AI/HTNPlanner/Source/HTNTestSuite/Private/HTNTest.cpp
+++ b/Engine/Plugins/AI/HTNPlanner/Source/HTNTestSuite/Private/HTNTest.cpp
@@ -28,7 +28,7 @@ struct FHTNTestBase : public FAITestBase
 		for (int32 WSIndex = 0; WSIndex < int32(EMockHTNWorldState::MAX); ++WSIndex)
 		{
 			// setting every key to it's numerical index value
-			WorldState.SetValueUnsafe(WSIndex, WSIndex);
+			WorldState.SetValueUnsafe(static_cast<FHTNPolicy::FWSKey>(WSIndex), WSIndex);
 		}
 	}
 
@@ -369,7 +369,7 @@ struct FAITest_HTNWorldRepresentation : public FHTNTestBase
 		for (int32 WSIndex = 0; WSIndex < int32(EMockHTNWorldState::MAX); ++WSIndex)
 		{
 			FHTNPolicy::FWSValue Value;
-			AITEST_TRUE("Retrieving known values from the WorldState instance", WorldState.GetValue(WSIndex, Value) && Value == FHTNPolicy::DefaultValue);
+			AITEST_TRUE("Retrieving known values from the WorldState instance", WorldState.GetValue(static_cast<FHTNPolicy::FWSKey>(WSIndex), Value) && Value == FHTNPolicy::DefaultValue);
 		}
 
 		PopulateWorldState();
@@ -378,7 +378,7 @@ struct FAITest_HTNWorldRepresentation : public FHTNTestBase
 		for (int32 WSIndex = 0; WSIndex < int32(EMockHTNWorldState::MAX); ++WSIndex)
 		{
 			FHTNPolicy::FWSValue Value;
-			AITEST_TRUE("Retrieving known values from the WorldState instance", WorldState.GetValue(WSIndex, Value) && Value == WSIndex);
+			AITEST_TRUE("Retrieving known values from the WorldState instance", WorldState.GetValue(static_cast<FHTNPolicy::FWSKey>(WSIndex), Value) && Value == WSIndex);
 
 			for (int32 OpIndex = 0; OpIndex < int32(EHTNWorldStateCheck::MAX); ++OpIndex)
 			{
@@ -530,7 +530,7 @@ struct FAITest_HTNCustomWSCheck : public FHTNTestBase
 
 		FHTNBuilder_CompositeTask& CompositeTaskBuilder = DomainBuilder.AddCompositeTask(NAME_None);	// root
 		{
-			FHTNBuilder_Method& MethodsBuilder = CompositeTaskBuilder.AddMethod(FHTNCondition(0, CustomCheckID));
+			FHTNBuilder_Method& MethodsBuilder = CompositeTaskBuilder.AddMethod(FHTNCondition(0, IntCastChecked<FHTNPolicy::FWSOperationID>(CustomCheckID)));
 			MethodsBuilder.AddTask(TEXT("Task1"));
 		}
 		FHTNBuilder_PrimitiveTask& PrimitiveTask = DomainBuilder.AddPrimitiveTask(TEXT("Task1"));
@@ -568,8 +568,8 @@ struct FAITest_HTNCustomWSOperation : public FHTNTestBase
 			MethodsBuilder.AddTask(TEXT("Task1"));
 		}
 		FHTNBuilder_PrimitiveTask& PrimitiveTask = DomainBuilder.AddPrimitiveTask(TEXT("Task1"));
-		PrimitiveTask.AddEffect(FHTNEffect(0, CustomOperationID));
-		PrimitiveTask.AddEffect(FHTNEffect(2, CustomOperationID));
+		PrimitiveTask.AddEffect(FHTNEffect(0, IntCastChecked<FHTNPolicy::FWSOperationID>(CustomOperationID)));
+		PrimitiveTask.AddEffect(FHTNEffect(2, IntCastChecked<FHTNPolicy::FWSOperationID>(CustomOperationID)));
 
 		DomainBuilder.Compile();
 
diff --git a/Engine/Plugins/Runtime/NavCorridor/Source/NavCorridor/Private/NavCorridor.cpp b/Engine/Plugins/Runtime/NavCorridor/Source/NavCorridor/Private/NavCorridor.cpp
--- a/Engine/Plugins/Runtime/NavCorridor/Source/NavCorridor/Private/NavCorridor.cpp
+++ b/Engine/Plugins/Runtime/NavCorridor/Source/NavCorridor/Private/NavCorridor.cpp
@@ -1488,7 +1488,7 @@ FNavCorridorLocation FNavCorridor::FindNearestLocationOnPath(const FVector Locat
 		if (SectionDistanceSq < NearestDistanceSq)
 		{
 			NearestDistanceSq = SectionDistanceSq;
-			Result.T = UV.Y;
+			Result.T = static_cast<float>(UV.Y);
 			Result.Location = FMath::Lerp(CurrPortal.Location, NextPortal.Location, Result.T);
 			Result.PortalIndex = PortalIndex;
 
@@ -1539,7 +1539,7 @@ FNavCorridorLocation FNavCorridor::AdvancePathLocation(const FNavCorridorLocatio
 		const FReal SectionLength = FVector::Distance(NextPortal.Location, CurrPortal.Location);
 		const FReal LeftoverDistance = AdvanceDistance - DistanceSoFar;
 		Result.PortalIndex = Portals.Num() - 2;
-		Result.T = 1.0 + FloatCastChecked<float>(LeftoverDistance / SectionLength, /* Precision */ 1. / 256.); // T will be > 1
+		Result.T = static_cast<float>(1. + (LeftoverDistance / SectionLength)); // T will be > 1
 		Result.Location = FMath::Lerp(CurrPortal.Location, NextPortal.Location, Result.T);
 	}
 	
diff --git a/Engine/Plugins/Runtime/StateTree/Source/StateTreeEditorModule/Private/Customizations/StateTreeEditorDataDetails.cpp b/Engine/Plugins/Runtime/StateTree/Source/StateTreeEditorModule/Private/Customizations/StateTreeEditorDataDetails.cpp
--- a/Engine/Plugins/Runtime/StateTree/Source/StateTreeEditorModule/Private/Customizations/StateTreeEditorDataDetails.cpp
+++ b/Engine/Plugins/Runtime/StateTree/Source/StateTreeEditorModule/Private/Customizations/StateTreeEditorDataDetails.cpp
@@ -104,7 +104,7 @@ void FStateTreeEditorDataDetails::CustomizeDetails(IDetailLayoutBuilder& DetailB
 					[
 						SNew(SBorder)
 						.Padding(FMargin(6, 1))
-						.BorderImage(new FSlateRoundedBoxBrush(FStyleColors::Hover, 6))
+						.BorderImage(new FSlateRoundedBoxBrush(FStyleColors::Hover, 6.f))
 						[
 							SNew(STextBlock)
 							.TextStyle(FStateTreeEditorStyle::Get(), "StateTree.Node.Operand")
diff --git a/Engine/Plugins/Runtime/StateTree/Source/StateTreeEditorModule/Private/StateTreeEditorData.cpp b/Engine/Plugins/Runtime/StateTree/Source/StateTreeEditorModule/Private/StateTreeEditorData.cpp
--- a/Engine/Plugins/Runtime/StateTree/Source/StateTreeEditorModule/Private/StateTreeEditorData.cpp
+++ b/Engine/Plugins/Runtime/StateTree/Source/StateTreeEditorModule/Private/StateTreeEditorData.cpp
@@ -219,7 +219,7 @@ FStateTreeBindableStructDesc UStateTreeEditorData::FindContextData(const UStruct
 		{
 			return 1.0f;
 		}
-		const float WorstCase = Name.Len() + CandidateName.Len();
+		const float WorstCase = static_cast<float>(Name.Len() + CandidateName.Len());
 		return 1.0f - (Algo::LevenshteinDistance(Name, CandidateName) / WorstCase);
 	};
 	
diff --git a/Engine/Plugins/Runtime/StateTree/Source/StateTreeEditorModule/Private/StateTreePropertyBindingCompiler.cpp b/Engine/Plugins/Runtime/StateTree/Source/StateTreeEditorModule/Private/StateTreePropertyBindingCompiler.cpp
--- a/Engine/Plugins/Runtime/StateTree/Source/StateTreeEditorModule/Private/StateTreePropertyBindingCompiler.cpp
+++ b/Engine/Plugins/Runtime/StateTree/Source/StateTreeEditorModule/Private/StateTreePropertyBindingCompiler.cpp
@@ -137,8 +137,8 @@ bool FStateTreePropertyBindingCompiler::CompileBatch(const FStateTreeBindableStr
 	{
 		FStateTreePropCopyBatch& Batch = PropertyBindings->CopyBatches.AddDefaulted_GetRef();
 		Batch.TargetStruct = TargetStruct;
-		Batch.BindingsBegin = BindingsBegin;
-		Batch.BindingsEnd = BindingsEnd;
+		Batch.BindingsBegin = IntCastChecked<uint16>(BindingsBegin);
+		Batch.BindingsEnd = IntCastChecked<uint16>(BindingsEnd);
 		OutBatchIndex = PropertyBindings->CopyBatches.Num() - 1;
 	}
 
diff --git a/Engine/Plugins/Runtime/StateTree/Source/StateTreeModule/Private/StateTreeExecutionContext.cpp b/Engine/Plugins/Runtime/StateTree/Source/StateTreeModule/Private/StateTreeExecutionContext.cpp
--- a/Engine/Plugins/Runtime/StateTree/Source/StateTreeModule/Private/StateTreeExecutionContext.cpp
+++ b/Engine/Plugins/Runtime/StateTree/Source/StateTreeModule/Private/StateTreeExecutionContext.cpp
@@ -107,7 +107,7 @@ bool FStateTreeExecutionContext::AreExternalDataViewsValid() const
 	return bResult;
 }
 
-void FStateTreeExecutionContext::UpdateLinkedStateParameters(const FCompactStateTreeState& State, const uint16 ParameterInstanceIndex)
+void FStateTreeExecutionContext::UpdateLinkedStateParameters(const FCompactStateTreeState& State, const int32 ParameterInstanceIndex)
 {
 	FStateTreeDataView StateParamsInstance = InstanceData.GetMutableStruct(ParameterInstanceIndex);
 	FCompactStateTreeParameters& StateParams = StateParamsInstance.GetMutable<FCompactStateTreeParameters>();
diff --git a/Engine/Plugins/Runtime/StateTree/Source/StateTreeModule/Private/StateTreePropertyBindings.cpp b/Engine/Plugins/Runtime/StateTree/Source/StateTreeModule/Private/StateTreePropertyBindings.cpp
--- a/Engine/Plugins/Runtime/StateTree/Source/StateTreeModule/Private/StateTreePropertyBindings.cpp
+++ b/Engine/Plugins/Runtime/StateTree/Source/StateTreeModule/Private/StateTreePropertyBindings.cpp
@@ -100,7 +100,7 @@ bool FStateTreePropertyBindings::ResolvePath(const UStruct* Struct, const FState
 		}
 		Indirection.ArrayIndex = FStateTreeIndex16(Segment->ArrayIndex.IsValid() ? Segment->ArrayIndex.Get() : 0);
 		Indirection.Type = Segment->Type;
-		Indirection.Offset = Property->GetOffset_ForInternal() + Property->ElementSize * Indirection.ArrayIndex.Get();
+		Indirection.Offset = IntCastChecked<uint16>(Property->GetOffset_ForInternal() + Property->ElementSize * Indirection.ArrayIndex.Get());
 
 		// Check to see if it is an array access first.
 		if (const FArrayProperty* ArrayProperty = CastField<FArrayProperty>(Property))
diff --git a/Engine/Plugins/Runtime/StateTree/Source/StateTreeModule/Public/StateTreeExecutionContext.h b/Engine/Plugins/Runtime/StateTree/Source/StateTreeModule/Public/StateTreeExecutionContext.h
--- a/Engine/Plugins/Runtime/StateTree/Source/StateTreeModule/Public/StateTreeExecutionContext.h
+++ b/Engine/Plugins/Runtime/StateTree/Source/StateTreeModule/Public/StateTreeExecutionContext.h
@@ -383,7 +383,7 @@ protected:
 	}
 
 	/** Sets up parameter data view for a linked state and copies bound properties. */
-	void UpdateLinkedStateParameters(const FCompactStateTreeState& State, const uint16 ParameterInstanceIndex);
+	void UpdateLinkedStateParameters(const FCompactStateTreeState& State, const int32 ParameterInstanceIndex);
 
 	/** Sets up parameter data view for subtree state. */
 	void UpdateSubtreeStateParameters(const FCompactStateTreeState& State);
diff --git a/Engine/Source/Editor/BehaviorTreeEditor/Private/BehaviorTreeDebugger.cpp b/Engine/Source/Editor/BehaviorTreeEditor/Private/BehaviorTreeDebugger.cpp
--- a/Engine/Source/Editor/BehaviorTreeEditor/Private/BehaviorTreeDebugger.cpp
+++ b/Engine/Source/Editor/BehaviorTreeEditor/Private/BehaviorTreeDebugger.cpp
@@ -36,8 +36,8 @@ FBehaviorTreeDebugger::FBehaviorTreeDebugger()
 	StepBackIntoIdx = INDEX_NONE;
 	StepBackOverIdx = INDEX_NONE;
 	StepOutIdx = INDEX_NONE;
-	SavedTimestamp = 0.0f;
-	CurrentTimestamp = 0.0f;
+	SavedTimestamp = 0.;
+	CurrentTimestamp = 0.;
 
 	FEditorDelegates::BeginPIE.AddRaw(this, &FBehaviorTreeDebugger::OnBeginPIE);
 	FEditorDelegates::EndPIE.AddRaw(this, &FBehaviorTreeDebugger::OnEndPIE);
@@ -228,7 +228,7 @@ void FBehaviorTreeDebugger::Tick(float DeltaTime)
 		if (ActiveStepIndex == LatestStepIndex)
 		{
 			TArray<FString> RuntimeDescriptions;
-			TreeInstance->StoreDebuggerRuntimeValues(RuntimeDescriptions, ShowInstance.RootNode, DebuggerInstanceIndex);
+			TreeInstance->StoreDebuggerRuntimeValues(RuntimeDescriptions, ShowInstance.RootNode, IntCastChecked<uint16>(DebuggerInstanceIndex));
 
 			UpdateAssetRuntimeDescription(RuntimeDescriptions, RootNode.Get());
 		}
@@ -1189,7 +1189,7 @@ void FBehaviorTreeDebugger::UpdateDebuggerViewOnTick()
 #if USE_BEHAVIORTREE_DEBUGGER
 	if (IsDebuggerRunning() && TreeInstance.IsValid())
 	{
-		const float GameTime = GEditor && GEditor->PlayWorld ? GEditor->PlayWorld->GetTimeSeconds() : 0.0f;
+		const double GameTime = GEditor && GEditor->PlayWorld ? GEditor->PlayWorld->GetTimeSeconds() : 0.;
 		CurrentTimestamp = GameTime;
 
 		TreeInstance->StoreDebuggerBlackboard(CurrentValues);
@@ -1227,7 +1227,7 @@ FText FBehaviorTreeDebugger::FindValueForKey(const FName& InKeyName, bool bUseCu
 #endif
 }
 
-float FBehaviorTreeDebugger::GetTimeStamp(bool bUseCurrentState) const
+double FBehaviorTreeDebugger::GetTimeStamp(bool bUseCurrentState) const
 {
 	return bUseCurrentState ? CurrentTimestamp : SavedTimestamp;
 }
diff --git a/Engine/Source/Editor/BehaviorTreeEditor/Private/BehaviorTreeDebugger.h b/Engine/Source/Editor/BehaviorTreeEditor/Private/BehaviorTreeDebugger.h
--- a/Engine/Source/Editor/BehaviorTreeEditor/Private/BehaviorTreeDebugger.h
+++ b/Engine/Source/Editor/BehaviorTreeEditor/Private/BehaviorTreeDebugger.h
@@ -95,7 +95,7 @@ public:
 	FText FindValueForKey(const FName& InKeyName, bool bUseCurrentState) const;
 
 	/** Gets the timestamp to be displayed, either current or saved */
-	float GetTimeStamp(bool bUseCurrentState) const;
+	double GetTimeStamp(bool bUseCurrentState) const;
 
 	/** Delegate fired when the debugged blackboard is changed */
 	DECLARE_EVENT_OneParam(FBehaviorTreeDebugger, FOnDebuggedBlackboardChanged, UBlackboardData*);
@@ -154,8 +154,8 @@ private:
 	TMap<FName, FString> CurrentValues;
 
 	/** Debugger timestamps */
-	float SavedTimestamp;
-	float CurrentTimestamp;
+	double SavedTimestamp;
+	double CurrentTimestamp;
 
 	/** set value of DebuggerInstanceIndex variable */
 	void UpdateDebuggerInstance();
diff --git a/Engine/Source/Editor/BehaviorTreeEditor/Private/BehaviorTreeDecoratorGraph.cpp b/Engine/Source/Editor/BehaviorTreeEditor/Private/BehaviorTreeDecoratorGraph.cpp
--- a/Engine/Source/Editor/BehaviorTreeEditor/Private/BehaviorTreeDecoratorGraph.cpp
+++ b/Engine/Source/Editor/BehaviorTreeEditor/Private/BehaviorTreeDecoratorGraph.cpp
@@ -64,7 +64,7 @@ void UBehaviorTreeDecoratorGraph::CollectDecoratorDataWorker(const UBehaviorTree
 		}
 	}
 
-	FBTDecoratorLogic LogicOp(Node->GetOperationType(), LinkedNodes.Num());
+	FBTDecoratorLogic LogicOp(IntCastChecked<uint8>((int32)Node->GetOperationType()), IntCastChecked<uint16>(LinkedNodes.Num()));
 
 	// special case: invalid
 	if (LogicOp.Operation == EBTDecoratorLogic::Invalid)
@@ -79,7 +79,7 @@ void UBehaviorTreeDecoratorGraph::CollectDecoratorDataWorker(const UBehaviorTree
 		UBTDecorator* DecoratorInstance = DecoratorNode ? (UBTDecorator*)DecoratorNode->NodeInstance : NULL;
 		if (DecoratorInstance)
 		{
-			LogicOp.Number = DecoratorInstances.Add(DecoratorInstance);
+			LogicOp.Number = IntCastChecked<uint16>(DecoratorInstances.Add(DecoratorInstance));
 			DecoratorOperations.Add(LogicOp);
 		}
 	}
@@ -138,8 +138,8 @@ UBehaviorTreeDecoratorGraphNode* UBehaviorTreeDecoratorGraph::SpawnMissingNodeWo
 
 	if (GraphNode)
 	{
-		GraphNode->NodePosX = ParentGraphNode->NodePosX - 300.0f;
-		GraphNode->NodePosY = ParentGraphNode->NodePosY + ChildIdx * 100.0f;
+		GraphNode->NodePosX = ParentGraphNode->NodePosX - 300;
+		GraphNode->NodePosY = ParentGraphNode->NodePosY + ChildIdx * 100;
 	}
 
 	for (int32 Idx = 0; Idx < NumSteps; Idx++)
diff --git a/Engine/Source/Editor/BehaviorTreeEditor/Private/BehaviorTreeEditor.cpp b/Engine/Source/Editor/BehaviorTreeEditor/Private/BehaviorTreeEditor.cpp
--- a/Engine/Source/Editor/BehaviorTreeEditor/Private/BehaviorTreeEditor.cpp
+++ b/Engine/Source/Editor/BehaviorTreeEditor/Private/BehaviorTreeEditor.cpp
@@ -423,7 +423,7 @@ FText FBehaviorTreeEditor::HandleGetDebugKeyValue(const FName& InKeyName, bool b
 	return FText();
 }
 
-float FBehaviorTreeEditor::HandleGetDebugTimeStamp(bool bUseCurrentState) const
+double FBehaviorTreeEditor::HandleGetDebugTimeStamp(bool bUseCurrentState) const
 {
 	if(IsDebuggerReady())
 	{
diff --git a/Engine/Source/Editor/BehaviorTreeEditor/Private/BehaviorTreeEditor.h b/Engine/Source/Editor/BehaviorTreeEditor/Private/BehaviorTreeEditor.h
--- a/Engine/Source/Editor/BehaviorTreeEditor/Private/BehaviorTreeEditor.h
+++ b/Engine/Source/Editor/BehaviorTreeEditor/Private/BehaviorTreeEditor.h
@@ -185,7 +185,7 @@ public:
 	FText HandleGetDebugKeyValue(const FName& InKeyName, bool bUseCurrentState) const;
 
 	/** Delegate handler for retrieving timestamp to display */
-	float HandleGetDebugTimeStamp(bool bUseCurrentState) const;
+	double HandleGetDebugTimeStamp(bool bUseCurrentState) const;
 
 	/** Delegate handler for when the debugged blackboard changes */
 	void HandleDebuggedBlackboardChanged(UBlackboardData* InObject);
diff --git a/Engine/Source/Editor/BehaviorTreeEditor/Private/BehaviorTreeGraph.cpp b/Engine/Source/Editor/BehaviorTreeEditor/Private/BehaviorTreeGraph.cpp
--- a/Engine/Source/Editor/BehaviorTreeEditor/Private/BehaviorTreeGraph.cpp
+++ b/Engine/Source/Editor/BehaviorTreeEditor/Private/BehaviorTreeGraph.cpp
@@ -346,7 +346,7 @@ namespace BTGraphHelpers
 			}
 			else if (MyCompositeNode)
 			{
-				MyCompositeNode->SetDecoratorData(RootNode, ChildIdx);
+				MyCompositeNode->SetDecoratorData(RootNode, IntCastChecked<uint8>(ChildIdx));
 
 				FIntIntPair RangeData;
 				RangeData.FirstIdx = DecoratorInstances.Num();
@@ -369,7 +369,7 @@ namespace BTGraphHelpers
 			DecoratorInstances[i]->InitializeNode(RootNode, *ExecutionIndex, 0, TreeDepth);
 			if (bInitializeNodes)
 			{
-				DecoratorInstances[i]->InitializeParentLink(ChildIdx);
+				DecoratorInstances[i]->InitializeParentLink(IntCastChecked<uint8>(ChildIdx));
 				*ExecutionIndex += 1;
 
 				// make sure that flow abort mode matches - skip for root level nodes
@@ -399,7 +399,7 @@ namespace BTGraphHelpers
 		{
 			if (NumNodes > 1)
 			{
-				FBTDecoratorLogic LogicOp(EBTDecoratorLogic::And, NumNodes);
+				FBTDecoratorLogic LogicOp(EBTDecoratorLogic::And, IntCastChecked<uint16>(NumNodes));
 				DecoratorOperations.Insert(LogicOp, 0);
 			}
 		}
@@ -432,7 +432,7 @@ namespace BTGraphHelpers
 			}
 			else if (MyCompositeNode)
 			{
-				MyCompositeNode->SetDecoratorData(RootNode, ChildIdx);
+				MyCompositeNode->SetDecoratorData(RootNode, IntCastChecked<uint8>(ChildIdx));
 
 				FIntIntPair RangeData;
 				RangeData.FirstIdx = DecoratorInstances.Num();
@@ -447,7 +447,7 @@ namespace BTGraphHelpers
 		for (int32 i = 0; i < DecoratorInstances.Num(); i++)
 		{
 			DecoratorInstances[i]->InitializeNode(RootNode, ExecutionIndex, 0, TreeDepth);
-			DecoratorInstances[i]->InitializeParentLink(ChildIdx);
+			DecoratorInstances[i]->InitializeParentLink(IntCastChecked<uint8>(ChildIdx));
 			ExecutionIndex++;
 		}
 
@@ -594,7 +594,7 @@ namespace BTGraphHelpers
 				UBTTask_RunBehavior* SubtreeTask = Cast<UBTTask_RunBehavior>(TaskInstance);
 				if (SubtreeTask)
 				{
-					*ExecutionIndex += SubtreeTask->GetInjectedNodesCount();
+					*ExecutionIndex += IntCastChecked<uint16>(SubtreeTask->GetInjectedNodesCount());
 				}
 
 				// collect services (task node)
@@ -613,7 +613,7 @@ namespace BTGraphHelpers
 							}
 
 							ServiceInstance->InitializeNode(RootNode, *ExecutionIndex, 0, TreeDepth);
-							ServiceInstance->InitializeParentLink(ChildIdx);
+							ServiceInstance->InitializeParentLink(IntCastChecked<uint8>(ChildIdx));
 							*ExecutionIndex += 1;
 
 							TaskInstance->Services.Add(ServiceInstance);
@@ -722,7 +722,7 @@ namespace BTGraphHelpers
 				UBTTask_RunBehavior* SubtreeTask = Cast<UBTTask_RunBehavior>(TaskInstance);
 				if (SubtreeTask)
 				{
-					*ExecutionIndex += SubtreeTask->GetInjectedNodesCount();
+					*ExecutionIndex += IntCastChecked<uint16>(SubtreeTask->GetInjectedNodesCount());
 				}
 
 				// collect services: task
@@ -843,8 +843,8 @@ namespace BTGraphHelpers
 		if (GraphNode)
 		{
 			const int32 ParentSubNodes = ParentGraphNode->Services.Num() + ParentGraphNode->Decorators.Num();
-			GraphNode->NodePosX = ParentGraphNode->NodePosX + ChildIdx * 400.0f;
-			GraphNode->NodePosY = ParentGraphNode->NodePosY + (ParentDecoratorCount + ParentSubNodes + 1) * 75.0f;
+			GraphNode->NodePosX = ParentGraphNode->NodePosX + ChildIdx * 400;
+			GraphNode->NodePosY = ParentGraphNode->NodePosY + (ParentDecoratorCount + ParentSubNodes + 1) * 75;
 			GraphNode->NodeInstance = Node;
 		}
 
@@ -1028,7 +1028,7 @@ void UBehaviorTreeGraph::SpawnMissingNodesForParallel()
 					UBehaviorTreeGraphNode* MainTaskNode = Cast<UBehaviorTreeGraphNode>(MainTaskPin->LinkedTo[0]->GetOwningNode());
 					if (MainTaskNode)
 					{
-						const int32 Width = MainTaskNode->DEPRECATED_NodeWidget.IsValid() ? MainTaskNode->DEPRECATED_NodeWidget.Pin()->GetDesiredSize().X : 200;
+						const int32 Width = MainTaskNode->DEPRECATED_NodeWidget.IsValid() ? static_cast<int32>(MainTaskNode->DEPRECATED_NodeWidget.Pin()->GetDesiredSize().X) : 200;
 						XOffset = MainTaskNode->NodePosX - ParallelNode->NodePosX + Width + 20;
 					}
 				}
@@ -1038,7 +1038,7 @@ void UBehaviorTreeGraph::SpawnMissingNodesForParallel()
 				WaitTaskNode->ClassData = FGraphNodeClassData(UBTTask_Wait::StaticClass(), "");
 				NodeBuilder.Finalize();
 
-				const int32 ParentHeight = ParallelNode->DEPRECATED_NodeWidget.IsValid() ? ParallelNode->DEPRECATED_NodeWidget.Pin()->GetDesiredSize().Y : 200;
+				const int32 ParentHeight = ParallelNode->DEPRECATED_NodeWidget.IsValid() ? static_cast<int32>(ParallelNode->DEPRECATED_NodeWidget.Pin()->GetDesiredSize().Y) : 200;
 				WaitTaskNode->NodePosX = ParallelNode->NodePosX + XOffset;
 				WaitTaskNode->NodePosY = ParallelNode->NodePosY + ParentHeight + 20;
 
@@ -1217,11 +1217,10 @@ namespace BTAutoArrangeHelpers
 				if (GraphNode && BBoxTree.Children.Num() > 0)
 				{
 					AutoArrangeNodes(GraphNode, BBoxTree.Children[BBoxIndex], PosX, PosY + GraphNode->DEPRECATED_NodeWidget.Pin()->GetDesiredSize().Y * 2.5f);
-					GraphNode->DEPRECATED_NodeWidget.Pin()->MoveTo(FVector2D(BBoxTree.Children[BBoxIndex].SubGraphBBox.X / 2 - GraphNode->DEPRECATED_NodeWidget.Pin()->GetDesiredSize().X / 2 + PosX, PosY), NodeFilter);
-					PosX += BBoxTree.Children[BBoxIndex].SubGraphBBox.X + 20;
+					GraphNode->DEPRECATED_NodeWidget.Pin()->MoveTo(FVector2D(BBoxTree.Children[BBoxIndex].SubGraphBBox.X / 2.f - GraphNode->DEPRECATED_NodeWidget.Pin()->GetDesiredSize().X / 2.f + PosX, PosY), NodeFilter);
+					PosX += BBoxTree.Children[BBoxIndex].SubGraphBBox.X + 20.f;
 					BBoxIndex++;
 				}
-
 			}
 		}
 	}
@@ -1246,7 +1245,7 @@ namespace BTAutoArrangeHelpers
 
 					GetNodeSizeInfo(GraphNode, ChildBounds);
 
-					LevelWidth += ChildBounds.SubGraphBBox.X + 20;
+					LevelWidth += ChildBounds.SubGraphBBox.X + 20.f;
 					if (ChildBounds.SubGraphBBox.Y > LevelHeight)
 					{
 						LevelHeight = ChildBounds.SubGraphBBox.Y;
@@ -1285,7 +1284,7 @@ void UBehaviorTreeGraph::AutoArrange()
 	BTAutoArrangeHelpers::GetNodeSizeInfo(RootNode, BBoxTree);
 	BTAutoArrangeHelpers::AutoArrangeNodes(RootNode, BBoxTree, 0, RootNode->DEPRECATED_NodeWidget.Pin()->GetDesiredSize().Y * 2.5f);
 
-	RootNode->NodePosX = BBoxTree.SubGraphBBox.X / 2 - RootNode->DEPRECATED_NodeWidget.Pin()->GetDesiredSize().X / 2;
+	RootNode->NodePosX = static_cast<int32>(BBoxTree.SubGraphBBox.X / 2.f - RootNode->DEPRECATED_NodeWidget.Pin()->GetDesiredSize().X / 2.f);
 	RootNode->NodePosY = 0;
 
 	RootNode->DEPRECATED_NodeWidget.Pin()->GetOwnerPanel()->ZoomToFit(/*bOnlySelection=*/ false);
diff --git a/Engine/Source/Editor/BehaviorTreeEditor/Private/BehaviorTreeGraphNode_Decorator.cpp b/Engine/Source/Editor/BehaviorTreeEditor/Private/BehaviorTreeGraphNode_Decorator.cpp
--- a/Engine/Source/Editor/BehaviorTreeEditor/Private/BehaviorTreeGraphNode_Decorator.cpp
+++ b/Engine/Source/Editor/BehaviorTreeEditor/Private/BehaviorTreeGraphNode_Decorator.cpp
@@ -47,6 +47,6 @@ void UBehaviorTreeGraphNode_Decorator::CollectDecoratorData(TArray<UBTDecorator*
 	{
 		UBTDecorator* DecoratorNode = (UBTDecorator*)NodeInstance;
 		const int32 InstanceIdx = NodeInstances.Add(DecoratorNode);
-		Operations.Add(FBTDecoratorLogic(EBTDecoratorLogic::Test, InstanceIdx));
+		Operations.Add(FBTDecoratorLogic(EBTDecoratorLogic::Test, IntCastChecked<uint16>(InstanceIdx)));
 	}
 }
diff --git a/Engine/Source/Editor/BehaviorTreeEditor/Private/DetailCustomizations/BlackboardDecoratorDetails.cpp b/Engine/Source/Editor/BehaviorTreeEditor/Private/DetailCustomizations/BlackboardDecoratorDetails.cpp
--- a/Engine/Source/Editor/BehaviorTreeEditor/Private/DetailCustomizations/BlackboardDecoratorDetails.cpp
+++ b/Engine/Source/Editor/BehaviorTreeEditor/Private/DetailCustomizations/BlackboardDecoratorDetails.cpp
@@ -146,7 +146,7 @@ void FBlackboardDecoratorDetails::OnKeyIDChanged()
 		if(KeyEntry && KeyEntry->KeyType)
 		{
 			CachedKeyType = KeyEntry->KeyType->GetClass();
-			CachedOperationType = KeyEntry->KeyType->GetTestOperation();
+			CachedOperationType = IntCastChecked<uint8>((int32)KeyEntry->KeyType->GetTestOperation());
 		}
 	}
 
diff --git a/Engine/Source/Editor/BehaviorTreeEditor/Private/SBehaviorTreeBlackboardView.h b/Engine/Source/Editor/BehaviorTreeEditor/Private/SBehaviorTreeBlackboardView.h
--- a/Engine/Source/Editor/BehaviorTreeEditor/Private/SBehaviorTreeBlackboardView.h
+++ b/Engine/Source/Editor/BehaviorTreeEditor/Private/SBehaviorTreeBlackboardView.h
@@ -42,7 +42,7 @@ DECLARE_DELEGATE_RetVal(bool, FOnIsDebuggerPaused);
 DECLARE_DELEGATE_RetVal(bool, FOnGetDisplayCurrentState);
 
 /** Delegate used to get the debugger's current timestamp */
-DECLARE_DELEGATE_RetVal_OneParam(float, FOnGetDebugTimeStamp, bool /* bUseCurrentState */);
+DECLARE_DELEGATE_RetVal_OneParam(double, FOnGetDebugTimeStamp, bool /* bUseCurrentState */);
 
 /** Delegate for when a blackboard key changes (added, removed, renamed) */
 DECLARE_DELEGATE_TwoParams(FOnBlackboardKeyChanged, UBlackboardData* /*InBlackboardData*/, FBlackboardEntry* const /*InKey*/);
diff --git a/Engine/Source/Editor/BehaviorTreeEditor/Private/SGraphNode_BehaviorTree.cpp b/Engine/Source/Editor/BehaviorTreeEditor/Private/SGraphNode_BehaviorTree.cpp
--- a/Engine/Source/Editor/BehaviorTreeEditor/Private/SGraphNode_BehaviorTree.cpp
+++ b/Engine/Source/Editor/BehaviorTreeEditor/Private/SGraphNode_BehaviorTree.cpp
@@ -558,7 +558,7 @@ void SGraphNode_BehaviorTree::UpdateGraphNode()
 					.Visibility(this, &SGraphNode_BehaviorTree::GetDragOverMarkerVisibility)
 					[
 						SNew(SBox)
-						.HeightOverride(4)
+						.HeightOverride(4.f)
 					]
 				]
 
diff --git a/Engine/Source/Runtime/Engine/Private/AI/Navigation/AvoidanceManager.cpp b/Engine/Source/Runtime/Engine/Private/AI/Navigation/AvoidanceManager.cpp
--- a/Engine/Source/Runtime/Engine/Private/AI/Navigation/AvoidanceManager.cpp
+++ b/Engine/Source/Runtime/Engine/Private/AI/Navigation/AvoidanceManager.cpp
@@ -247,7 +247,7 @@ void UAvoidanceManager::UpdateRVO_Internal(int32 inAvoidanceUID, const FNavAvoid
 {
 	if (FNavAvoidanceData* existingData = AvoidanceObjects.Find(inAvoidanceUID))
 	{
-		float OverrideWeightTime = existingData->OverrideWeightTime;		//Hold onto this one value
+		double OverrideWeightTime = existingData->OverrideWeightTime;		//Hold onto this one value
 		*existingData = inAvoidanceData;
 		existingData->OverrideWeightTime = OverrideWeightTime;
 	}
diff --git a/Engine/Source/Runtime/Engine/Private/VisualLogger/VisualLoggerTypes.cpp b/Engine/Source/Runtime/Engine/Private/VisualLogger/VisualLoggerTypes.cpp
--- a/Engine/Source/Runtime/Engine/Private/VisualLogger/VisualLoggerTypes.cpp
+++ b/Engine/Source/Runtime/Engine/Private/VisualLogger/VisualLoggerTypes.cpp
@@ -315,11 +315,11 @@ void FVisualLogEntry::AddElement(const FVector& Base, float HalfHeight, float Ra
 	ElementsToDraw.Add(Element);
 }
 
-void FVisualLogEntry::AddElement(const TArray<FVector>& ConvexPoints, float MinZ, float MaxZ, const FName& CategoryName, ELogVerbosity::Type Verbosity, const FColor& Color, const FString& Description)
+void FVisualLogEntry::AddElement(const TArray<FVector>& ConvexPoints, FVector::FReal MinZ, FVector::FReal MaxZ, const FName& CategoryName, ELogVerbosity::Type Verbosity, const FColor& Color, const FString& Description)
 {
 	FVisualLogShapeElement Element(Description, Color, 0, CategoryName);
 	Element.Points.Reserve(1 + ConvexPoints.Num());
-	Element.Points.Add(FVector(MinZ, MaxZ, 0));
+	Element.Points.Add(FVector(MinZ, MaxZ, 0.));
 	Element.Points.Append(ConvexPoints);
 	Element.Type = EVisualLoggerShapeElement::NavAreaMesh;
 	Element.Verbosity = Verbosity;
diff --git a/Engine/Source/Runtime/Engine/Public/VisualLogger/VisualLoggerTypes.h b/Engine/Source/Runtime/Engine/Public/VisualLogger/VisualLoggerTypes.h
--- a/Engine/Source/Runtime/Engine/Public/VisualLogger/VisualLoggerTypes.h
+++ b/Engine/Source/Runtime/Engine/Public/VisualLogger/VisualLoggerTypes.h
@@ -222,7 +222,7 @@ struct ENGINE_API FVisualLogEntry
 	// custom element
 	void AddElement(const FVisualLogShapeElement& Element);
 	// NavAreaMesh
-	void AddElement(const TArray<FVector>& ConvexPoints, float MinZ, float MaxZ, const FName& CategoryName, ELogVerbosity::Type Verbosity, const FColor& Color = FColor::White, const FString& Description = TEXT(""));
+	void AddElement(const TArray<FVector>& ConvexPoints, FVector::FReal MinZ, FVector::FReal MaxZ, const FName& CategoryName, ELogVerbosity::Type Verbosity, const FColor& Color = FColor::White, const FString& Description = TEXT(""));
 	// 3d Mesh
 	void AddElement(const TArray<FVector>& Vertices, const TArray<int32>& Indices, const FName& CategoryName, ELogVerbosity::Type Verbosity, const FColor& Color = FColor::White, const FString& Description = TEXT(""));
 	// 2d convex
diff --git a/Engine/Source/Runtime/GameplayDebugger/Private/GameplayDebuggerCategory.cpp b/Engine/Source/Runtime/GameplayDebugger/Private/GameplayDebuggerCategory.cpp
--- a/Engine/Source/Runtime/GameplayDebugger/Private/GameplayDebuggerCategory.cpp
+++ b/Engine/Source/Runtime/GameplayDebugger/Private/GameplayDebuggerCategory.cpp
@@ -17,7 +17,7 @@ FGameplayDebuggerCategory::FGameplayDebuggerCategory() :
 	bHasAuthority(true),
 	bIsEnabled(true),
 	CategoryId(INDEX_NONE),
-	LastCollectDataTime(-FLT_MAX)
+	LastCollectDataTime(-MAX_dbl)
 {
 }
 
@@ -81,7 +81,7 @@ void FGameplayDebuggerCategory::DrawCategory(APlayerController* OwnerPC, FGamepl
 
 	if (bShowUpdateTimer && bHasAuthority)
 	{
-		const float GameTime = World->GetTimeSeconds();
+		const double GameTime = World->GetTimeSeconds();
 		CanvasContext.Printf(TEXT("%sNext update in: {yellow}%.0fs"), *CategoryPrefix, CollectDataInterval - (GameTime - LastCollectDataTime));
 	}
 
@@ -137,14 +137,14 @@ bool FGameplayDebuggerCategory::IsLocationInViewCone(const FVector& ViewLocation
 {
 	const UGameplayDebuggerUserSettings* Settings = GetDefault<UGameplayDebuggerUserSettings>();
 	const FVector DirToEntity = TargetLocation - ViewLocation;
-	const float DistanceToEntitySq = DirToEntity.SquaredLength();
+	const FVector::FReal DistanceToEntitySq = DirToEntity.SquaredLength();
 	if (DistanceToEntitySq > FMath::Square(Settings->MaxViewDistance))
 	{
 		return false;
 	}
 
-	const float ViewDot = FVector::DotProduct(DirToEntity.GetSafeNormal(), ViewDirection);
-	const float MinViewDirDot = FMath::Cos(FMath::DegreesToRadians(Settings->MaxViewAngle));
+	const FVector::FReal ViewDot = FVector::DotProduct(DirToEntity.GetSafeNormal(), ViewDirection);
+	const FVector::FReal MinViewDirDot = FMath::Cos(FMath::DegreesToRadians(Settings->MaxViewAngle));
 	if (ViewDot < MinViewDirDot)
 	{
 		return false;
diff --git a/Engine/Source/Runtime/GameplayDebugger/Private/GameplayDebuggerCategoryReplicator.cpp b/Engine/Source/Runtime/GameplayDebugger/Private/GameplayDebuggerCategoryReplicator.cpp
--- a/Engine/Source/Runtime/GameplayDebugger/Private/GameplayDebuggerCategoryReplicator.cpp
+++ b/Engine/Source/Runtime/GameplayDebugger/Private/GameplayDebuggerCategoryReplicator.cpp
@@ -354,7 +354,7 @@ bool FGameplayDebuggerNetPack::NetDeltaSerialize(FNetDeltaSerializeInfo& DeltaPa
 			uint8 BaseFlags = SavedCategory.bIsEnabled;
 			uint8 ShouldUpdateTextLines = bMissingOldState || (OldState->CategoryStates[Idx].TextLinesRepCounter != NewState->CategoryStates[Idx].TextLinesRepCounter);
 			uint8 ShouldUpdateShapes = bMissingOldState || (OldState->CategoryStates[Idx].ShapesRepCounter != NewState->CategoryStates[Idx].ShapesRepCounter);
-			uint8 NumDataPacks = SavedCategory.DataPacks.Num();
+			uint8 NumDataPacks = IntCastChecked<uint8>(SavedCategory.DataPacks.Num());
 
 			Writer << SavedCategory.CategoryName;
 
@@ -848,7 +848,7 @@ void AGameplayDebuggerCategoryReplicator::PostNetReceive()
 
 void AGameplayDebuggerCategoryReplicator::CollectCategoryData(bool bForce)
 {
-	const float GameTime = GetWorld()->GetTimeSeconds();
+	const double GameTime = GetWorld()->GetTimeSeconds();
 
 #if ENABLE_VISUAL_LOG
 	const FVisualLogger& VLogger = FVisualLogger::Get();
diff --git a/Engine/Source/Runtime/GameplayDebugger/Private/GameplayDebuggerLocalController.cpp b/Engine/Source/Runtime/GameplayDebugger/Private/GameplayDebuggerLocalController.cpp
--- a/Engine/Source/Runtime/GameplayDebugger/Private/GameplayDebuggerLocalController.cpp
+++ b/Engine/Source/Runtime/GameplayDebugger/Private/GameplayDebuggerLocalController.cpp
@@ -344,7 +344,7 @@ void UGameplayDebuggerLocalController::DrawCategoryHeader(int32 CategoryId, TSha
 		Category->GetNumDataPacks() > 0)
 	{
 		// collect brief data pack status, detailed info is displayed only when ShouldDrawReplicationStatus is true
-		const int32 CurrentSyncCounter = CachedReplicator->GetDebugActorCounter();
+		const int16 CurrentSyncCounter = CachedReplicator->GetDebugActorCounter();
 
 		DataPackDesc = TEXT("{white} ver[");
 		bool bIsPrevOutdated = false;
@@ -514,9 +514,9 @@ void UGameplayDebuggerLocalController::OnActivationPressed()
 {
 	if (CachedReplicator)
 	{
-		const float HoldTimeThr = 0.2f * (FApp::UseFixedTimeStep() ? (FApp::GetFixedDeltaTime() * 60.0f) : 1.0f);
+		const double HoldTimeThr = 0. * (FApp::UseFixedTimeStep() ? (FApp::GetFixedDeltaTime() * 60.) : 1.);
 
-		CachedReplicator->GetWorldTimerManager().SetTimer(StartSelectingActorHandle, this, &UGameplayDebuggerLocalController::OnStartSelectingActor, HoldTimeThr);
+		CachedReplicator->GetWorldTimerManager().SetTimer(StartSelectingActorHandle, this, &UGameplayDebuggerLocalController::OnStartSelectingActor, static_cast<float>(HoldTimeThr));
 	}
 }
 
@@ -684,11 +684,11 @@ void UGameplayDebuggerLocalController::OnSelectActorTick()
 		}
 
 		const UGameplayDebuggerUserSettings* Settings = GetDefault<UGameplayDebuggerUserSettings>();
-		const float MaxScanDistance = Settings->MaxViewDistance;
-		const float MinViewDirDot = FMath::Cos(FMath::DegreesToRadians(Settings->MaxViewAngle));
+		const FVector::FReal MaxScanDistance = Settings->MaxViewDistance;
+		const FVector::FReal MinViewDirDot = FMath::Cos(FMath::DegreesToRadians(Settings->MaxViewAngle));
 
 		AActor* BestCandidate = nullptr;
-		float BestScore = MinViewDirDot;
+		FVector::FReal BestScore = MinViewDirDot;
 
 		for (APawn* TestPawn : TActorRange<APawn>(OwnerPC->GetWorld()))
 		{
@@ -697,18 +697,18 @@ void UGameplayDebuggerLocalController::OnSelectActorTick()
 				TestPawn != OwnerPC->GetPawn())
 			{
 				FVector DirToPawn = (TestPawn->GetActorLocation() - ViewLocation);
-				float DistToPawn = DirToPawn.Size();
+				FVector::FReal DistToPawn = DirToPawn.Size();
 				if (FMath::IsNearlyZero(DistToPawn))
 				{
 					DirToPawn = ViewDirection;
-					DistToPawn = 1.0f;
+					DistToPawn = 1.;
 				}
 				else
 				{
 					DirToPawn /= DistToPawn;
 				}
 
-				const float ViewDot = FVector::DotProduct(ViewDirection, DirToPawn);
+				const FVector::FReal ViewDot = FVector::DotProduct(ViewDirection, DirToPawn);
 				if (DistToPawn < MaxScanDistance && ViewDot > BestScore)
 				{
 					BestScore = ViewDot;
diff --git a/Engine/Source/Runtime/GameplayDebugger/Private/GameplayDebuggerRenderingComponent.cpp b/Engine/Source/Runtime/GameplayDebugger/Private/GameplayDebuggerRenderingComponent.cpp
--- a/Engine/Source/Runtime/GameplayDebugger/Private/GameplayDebuggerRenderingComponent.cpp
+++ b/Engine/Source/Runtime/GameplayDebugger/Private/GameplayDebuggerRenderingComponent.cpp
@@ -67,13 +67,13 @@ public:
 private:
 	uint32 GetAllocatedSizeInternal(void) const
 	{
-		uint32 Size = FDebugRenderSceneProxy::GetAllocatedSize() + ChildProxies.GetAllocatedSize();
+		SIZE_T Size = FDebugRenderSceneProxy::GetAllocatedSize() + ChildProxies.GetAllocatedSize();
 		for (int32 Idx = 0; Idx < ChildProxies.Num(); Idx++)
 		{
 			Size += ChildProxies[Idx]->GetMemoryFootprint();
 		}
 
-		return Size;
+		return IntCastChecked<uint32>(Size);
 	}
 
 protected:
diff --git a/Engine/Source/Runtime/GameplayDebugger/Private/GameplayDebuggerTypes.cpp b/Engine/Source/Runtime/GameplayDebugger/Private/GameplayDebuggerTypes.cpp
--- a/Engine/Source/Runtime/GameplayDebugger/Private/GameplayDebuggerTypes.cpp
+++ b/Engine/Source/Runtime/GameplayDebugger/Private/GameplayDebuggerTypes.cpp
@@ -159,7 +159,7 @@ void FGameplayDebuggerShape::Draw(UWorld* World, FGameplayDebuggerCanvasContext&
 	case EGameplayDebuggerShape::Point:
 		if (ShapeData.Num() == 2 && ShapeData[1].X > 0)
 		{
-			DrawDebugSphere(World, ShapeData[0], ShapeData[1].X, 16, Color, bPersistent, LifeTime, DepthPriority);
+			DrawDebugSphere(World, ShapeData[0], static_cast<float>(ShapeData[1].X), 16, Color, bPersistent, LifeTime, DepthPriority);
 			DescLocation = ShapeData[0];
 		}
 		break;
@@ -167,7 +167,7 @@ void FGameplayDebuggerShape::Draw(UWorld* World, FGameplayDebuggerCanvasContext&
 	case EGameplayDebuggerShape::Segment:
 		if (ShapeData.Num() == 3 && ShapeData[2].X > 0)
 		{
-			DrawDebugLine(World, ShapeData[0], ShapeData[1], Color, bPersistent, LifeTime, DepthPriority, ShapeData[2].X);
+			DrawDebugLine(World, ShapeData[0], ShapeData[1], Color, bPersistent, LifeTime, DepthPriority, static_cast<float>(ShapeData[2].X));
 			DescLocation = (ShapeData[0] + ShapeData[1]) * 0.5f;
 		}
 		break;
@@ -175,7 +175,7 @@ void FGameplayDebuggerShape::Draw(UWorld* World, FGameplayDebuggerCanvasContext&
 	case EGameplayDebuggerShape::Arrow:
 		if (ShapeData.Num() == 3 && ShapeData[2].X > 0)
 		{
-			DrawDebugDirectionalArrow(World, ShapeData[0], ShapeData[1], ShapeData[2].Y, Color, bPersistent, LifeTime, DepthPriority, ShapeData[2].X);
+			DrawDebugDirectionalArrow(World, ShapeData[0], ShapeData[1], ShapeData[2].Y, Color, bPersistent, LifeTime, DepthPriority, static_cast<float>(ShapeData[2].X));
 			DescLocation = (ShapeData[0] + ShapeData[1]) * 0.5f;
 		}
 		break;
@@ -192,7 +192,7 @@ void FGameplayDebuggerShape::Draw(UWorld* World, FGameplayDebuggerCanvasContext&
 		if (ShapeData.Num() == 3 && ShapeData[2].X > 0)
 		{
 			const float DefaultConeAngle = 0.25f; // ~ 15 degrees
-			DrawDebugCone(World, ShapeData[0], ShapeData[1], ShapeData[2].X, DefaultConeAngle, DefaultConeAngle, 16, Color, bPersistent, LifeTime, DepthPriority);
+			DrawDebugCone(World, ShapeData[0], ShapeData[1], static_cast<float>(ShapeData[2].X), DefaultConeAngle, DefaultConeAngle, 16, Color, bPersistent, LifeTime, DepthPriority);
 			DescLocation = ShapeData[0];
 		}
 		break;
@@ -200,7 +200,7 @@ void FGameplayDebuggerShape::Draw(UWorld* World, FGameplayDebuggerCanvasContext&
 	case EGameplayDebuggerShape::Cylinder:
 		if (ShapeData.Num() == 2)
 		{
-			DrawDebugCylinder(World, ShapeData[0] - FVector(0, 0, ShapeData[1].Z), ShapeData[0] + FVector(0, 0, ShapeData[1].Z), ShapeData[1].X, 16, Color, bPersistent, LifeTime, DepthPriority);
+			DrawDebugCylinder(World, ShapeData[0] - FVector(0, 0, ShapeData[1].Z), ShapeData[0] + FVector(0, 0, ShapeData[1].Z), static_cast<float>(ShapeData[1].X), 16, Color, bPersistent, LifeTime, DepthPriority);
 			DescLocation = ShapeData[0];
 		}
 		break;
@@ -209,7 +209,7 @@ void FGameplayDebuggerShape::Draw(UWorld* World, FGameplayDebuggerCanvasContext&
 		if (ShapeData.Num() == 3)
 		{
 			const FMatrix Axes = FRotationMatrix::MakeFromX(ShapeData[1]);
-			DrawDebugCircle(World, ShapeData[0], ShapeData[2].X, 32, Color, bPersistent, LifeTime, DepthPriority, ShapeData[2].Y, Axes.GetUnitAxis(EAxis::Y), Axes.GetUnitAxis(EAxis::Z), false);
+			DrawDebugCircle(World, ShapeData[0], static_cast<float>(ShapeData[2].X), 32, Color, bPersistent, LifeTime, DepthPriority, ShapeData[2].Y, Axes.GetUnitAxis(EAxis::Y), Axes.GetUnitAxis(EAxis::Z), false);
 			DescLocation = ShapeData[0];
 		}
 		break;
@@ -217,7 +217,7 @@ void FGameplayDebuggerShape::Draw(UWorld* World, FGameplayDebuggerCanvasContext&
 	case EGameplayDebuggerShape::Capsule:
 		if (ShapeData.Num() == 2)
 		{
-			DrawDebugCapsule(World, ShapeData[0], ShapeData[1].Z, ShapeData[1].X, FQuat::Identity, Color, bPersistent, LifeTime, DepthPriority);
+			DrawDebugCapsule(World, ShapeData[0], static_cast<float>(ShapeData[1].Z), static_cast<float>(ShapeData[1].X), FQuat::Identity, Color, bPersistent, LifeTime, DepthPriority);
 			DescLocation = ShapeData[0];
 		}
 		break;
@@ -415,8 +415,8 @@ FGameplayDebuggerCanvasContext::FGameplayDebuggerCanvasContext(UCanvas* InCanvas
 	{
 		Canvas = InCanvas;
 		Font = InFont;
-		CursorX = DefaultX = InCanvas->SafeZonePadX;
-		CursorY = DefaultY = InCanvas->SafeZonePadY;
+		CursorX = DefaultX = static_cast<float>(InCanvas->SafeZonePadX);
+		CursorY = DefaultY = static_cast<float>(InCanvas->SafeZonePadY);
 	}
 	else
 	{
diff --git a/Engine/Source/Runtime/GameplayDebugger/Public/GameplayDebuggerCategory.h b/Engine/Source/Runtime/GameplayDebugger/Public/GameplayDebuggerCategory.h
--- a/Engine/Source/Runtime/GameplayDebugger/Public/GameplayDebuggerCategory.h
+++ b/Engine/Source/Runtime/GameplayDebugger/Public/GameplayDebuggerCategory.h
@@ -146,7 +146,7 @@ protected:
 	}
 
 	/** [AUTH] force data collection on next update */
-	void ForceImmediateCollect() { LastCollectDataTime = -FLT_MAX; }
+	void ForceImmediateCollect() { LastCollectDataTime = -MAX_dbl; }
 
 	/** update interval, 0 = each tick */
 	float CollectDataInterval;
@@ -182,7 +182,7 @@ private:
 	int32 CategoryId;
 
 	/** timestamp of last update */
-	float LastCollectDataTime;
+	double LastCollectDataTime;
 
 	/** name of debugger category (auto assigned during category registration) */
 	FName CategoryName;
diff --git a/Engine/Source/Runtime/GameplayDebugger/Public/GameplayDebuggerCategoryReplicator.h b/Engine/Source/Runtime/GameplayDebugger/Public/GameplayDebuggerCategoryReplicator.h
--- a/Engine/Source/Runtime/GameplayDebugger/Public/GameplayDebuggerCategoryReplicator.h
+++ b/Engine/Source/Runtime/GameplayDebugger/Public/GameplayDebuggerCategoryReplicator.h
@@ -102,7 +102,7 @@ struct FGameplayDebuggerDebugActor
 	FName ActorName;
 
 	UPROPERTY()
-	int32 SyncCounter = 0;
+	int16 SyncCounter = 0;
 };
 
 USTRUCT()
@@ -178,7 +178,7 @@ public:
 	FName GetDebugActorName() const { return DebugActor.ActorName; }
 
 	/** get sync counter, increased with every change of DebugActor */
-	int32 GetDebugActorCounter() const { return DebugActor.SyncCounter; }
+	int16 GetDebugActorCounter() const { return DebugActor.SyncCounter; }
 
 	/** get view point information */
 	bool GetViewPoint(FVector& OutViewLocation, FVector& OutViewDirection) const;
diff --git a/Engine/Source/Runtime/NavigationSystem/Private/NavMesh/NavMeshRenderingComponent.cpp b/Engine/Source/Runtime/NavigationSystem/Private/NavMesh/NavMeshRenderingComponent.cpp
--- a/Engine/Source/Runtime/NavigationSystem/Private/NavMesh/NavMeshRenderingComponent.cpp
+++ b/Engine/Source/Runtime/NavigationSystem/Private/NavMesh/NavMeshRenderingComponent.cpp
@@ -507,10 +507,10 @@ void FNavMeshSceneProxyData::GatherData(const ARecastNavMesh* NavMesh, int32 InN
 		auto LerpColor = [](FColor A, FColor B, float T) -> FColor
 		{
 			return FColor(
-				FMath::RoundToInt(float(A.R) * (1.f - T) + float(B.R) * T),
-				FMath::RoundToInt(float(A.G) * (1.f - T) + float(B.G) * T),
-				FMath::RoundToInt(float(A.B) * (1.f - T) + float(B.B) * T),
-				FMath::RoundToInt(float(A.A) * (1.f - T) + float(B.A) * T));
+				static_cast<uint8>(FMath::RoundToInt(float(A.R) * (1.f - T) + float(B.R) * T)),
+				static_cast<uint8>(FMath::RoundToInt(float(A.G) * (1.f - T) + float(B.G) * T)),
+				static_cast<uint8>(FMath::RoundToInt(float(A.B) * (1.f - T) + float(B.B) * T)),
+				static_cast<uint8>(FMath::RoundToInt(float(A.A) * (1.f - T) + float(B.A) * T)));
 		};
 
 		TArray<FColor> TileBuildTimeColors;
@@ -925,8 +925,8 @@ void FNavMeshSceneProxyData::GatherData(const ARecastNavMesh* NavMesh, int32 InN
 						int32 Rank = 0;
 						if (Range != 0.)
 						{
-							Rank = FRecastDebugGeometry::BuildTimeBucketsCount * ((DebugData->BuildTime - NavMeshGeometry.MinTileBuildTime) / Range);
-							Rank = FMath::Clamp(Rank, 0, FRecastDebugGeometry::BuildTimeBucketsCount-1);
+							const double RankCalc = FRecastDebugGeometry::BuildTimeBucketsCount * ((DebugData->BuildTime - NavMeshGeometry.MinTileBuildTime) / Range);
+							Rank = static_cast<int32>(FMath::Clamp<double>(RankCalc, 0., FRecastDebugGeometry::BuildTimeBucketsCount-1));
 						}
 
 						DebugLabels.Add(FDebugText(Pair.Value + NavMeshDrawOffset, FString::Printf(TEXT("%.2f ms\n%.2f comp + %.2f nav\n%i tri"),
diff --git a/Engine/Source/Runtime/NavigationSystem/Private/NavMesh/RecastNavMeshGenerator.cpp b/Engine/Source/Runtime/NavigationSystem/Private/NavMesh/RecastNavMeshGenerator.cpp
--- a/Engine/Source/Runtime/NavigationSystem/Private/NavMesh/RecastNavMeshGenerator.cpp
+++ b/Engine/Source/Runtime/NavigationSystem/Private/NavMesh/RecastNavMeshGenerator.cpp
@@ -4264,13 +4264,12 @@ void FRecastTileGenerator::MarkDynamicArea(const FAreaNavModifier& Modifier, con
 			if (ReplaceIDPtr)
 			{
 				dtReplaceCylinderArea(Layer, LayerRecastOrig, TileConfig.cs, TileConfig.ch,
-					&(RecastPos.X), CylinderData.Radius, CylinderData.Height, AreaID, *ReplaceIDPtr);
+					&(RecastPos.X), CylinderData.Radius, CylinderData.Height, IntCastChecked<unsigned char>(AreaID), IntCastChecked<unsigned char>(*ReplaceIDPtr));
 			}
 			else
 			{
 				dtMarkCylinderArea(Layer, LayerRecastOrig, TileConfig.cs, TileConfig.ch,
-					&(RecastPos.X), CylinderData.Radius, CylinderData.Height, AreaID);
-			}
+					&(RecastPos.X), CylinderData.Radius, CylinderData.Height, IntCastChecked<unsigned char>(AreaID));			}
 		}
 		break;
 
@@ -4291,12 +4290,12 @@ void FRecastTileGenerator::MarkDynamicArea(const FAreaNavModifier& Modifier, con
 			if (ReplaceIDPtr)
 			{
 				dtReplaceBoxArea(Layer, LayerRecastOrig, TileConfig.cs, TileConfig.ch,
-					&(RecastPos.X), &(RecastExtent.X), AreaID, *ReplaceIDPtr);
+					&(RecastPos.X), &(RecastExtent.X), IntCastChecked<unsigned char>(AreaID), IntCastChecked<unsigned char>(*ReplaceIDPtr));
 			}
 			else
 			{
 				dtMarkBoxArea(Layer, LayerRecastOrig, TileConfig.cs, TileConfig.ch,
-					&(RecastPos.X), &(RecastExtent.X), AreaID);
+					&(RecastPos.X), &(RecastExtent.X), IntCastChecked<unsigned char>(AreaID));
 			}
 		}
 		break;
@@ -4338,12 +4337,12 @@ void FRecastTileGenerator::MarkDynamicArea(const FAreaNavModifier& Modifier, con
 				if (ReplaceIDPtr)
 				{
 					dtReplaceConvexArea(Layer, LayerRecastOrig, TileConfig.cs, TileConfig.ch,
-						ConvexCoords.GetData(), ConvexVerts.Num(), ConvexData.MinZ, ConvexData.MaxZ, AreaID, *ReplaceIDPtr);
+						ConvexCoords.GetData(), ConvexVerts.Num(), ConvexData.MinZ, ConvexData.MaxZ, IntCastChecked<unsigned char>(AreaID), IntCastChecked<unsigned char>(*ReplaceIDPtr));
 				}
 				else
 				{
 					dtMarkConvexArea(Layer, LayerRecastOrig, TileConfig.cs, TileConfig.ch,
-						ConvexCoords.GetData(), ConvexVerts.Num(), ConvexData.MinZ, ConvexData.MaxZ, AreaID);
+						ConvexCoords.GetData(), ConvexVerts.Num(), ConvexData.MinZ, ConvexData.MaxZ, IntCastChecked<unsigned char>(AreaID));
 				}
 			}
 		}
@@ -4355,7 +4354,7 @@ void FRecastTileGenerator::MarkDynamicArea(const FAreaNavModifier& Modifier, con
 
 uint32 FRecastTileGenerator::GetUsedMemCount() const
 {
-	uint32 TotalMemory = 0;
+	SIZE_T TotalMemory = 0;
 	TotalMemory += InclusionBounds.GetAllocatedSize();
 	TotalMemory += Modifiers.GetAllocatedSize();
 	TotalMemory += OffmeshLinks.GetAllocatedSize();
@@ -4392,7 +4391,7 @@ uint32 FRecastTileGenerator::GetUsedMemCount() const
 		TotalMemory += NavigationData[i].DataSize;
 	}
 
-	return TotalMemory;
+	return IntCastChecked<uint32>(TotalMemory);
 }
 
 void FRecastTileGenerator::AddReferencedObjects(FReferenceCollector& Collector)
@@ -4434,7 +4433,7 @@ void CheckTileIndicesInValidRange(const TNavStatArray<FBox>& NavigableAreas, con
 
 int32 CalculateMaxTilesCount(const TNavStatArray<FBox>& NavigableAreas, FVector::FReal TileSizeInWorldUnits, FVector::FReal AvgLayersPerGridCell, const uint32 NavMeshVersion)
 {
-	int32 GridCellsCount = 0;
+	int64 GridCellsCount = 0;
 	for (int32 Index = 0; Index < NavigableAreas.Num(); ++Index)
 	{
 		const FBox& AreaBounds = NavigableAreas[Index];
@@ -4469,21 +4468,21 @@ int32 CalculateMaxTilesCount(const TNavStatArray<FBox>& NavigableAreas, FVector:
 		if (NavMeshVersion >= NAVMESHVER_MAXTILES_COUNT_CHANGE)
 		{
 			// Keep this as an integer division to avoid imprecision between platforms and targets (since MaxTilesCount is compared with stored data).
-			const int32 TileSizeUU = (int32)TileSizeInWorldUnits;
-			const int32 XSize = (FMath::CeilToInt(RCBox.GetSize().X) / TileSizeUU) + 1;
-			const int32 YSize = (FMath::CeilToInt(RCBox.GetSize().Z) / TileSizeUU) + 1;
+			const int64 TileSizeUU = (int64)TileSizeInWorldUnits;
+			const int64 XSize = (FMath::CeilToInt(RCBox.GetSize().X) / TileSizeUU) + 1;
+			const int64 YSize = (FMath::CeilToInt(RCBox.GetSize().Z) / TileSizeUU) + 1;
 			GridCellsCount += (XSize*YSize);
 		}
 		else
 		{
 			// Support old navmesh versions
-			int32 XSize = FMath::CeilToInt(RCBox.GetSize().X/TileSizeInWorldUnits) + 1;
-			int32 YSize = FMath::CeilToInt(RCBox.GetSize().Z/TileSizeInWorldUnits) + 1;
+			int64 XSize = FMath::CeilToInt(RCBox.GetSize().X/TileSizeInWorldUnits) + 1;
+			int64 YSize = FMath::CeilToInt(RCBox.GetSize().Z/TileSizeInWorldUnits) + 1;
 			GridCellsCount+= (XSize*YSize);
 		}
 	}
 	
-	return FMath::CeilToInt(GridCellsCount * AvgLayersPerGridCell);
+	return IntCastChecked<int32>(FMath::CeilToInt(GridCellsCount * AvgLayersPerGridCell));
 }
 } // UE::NavMesh::Private
 
@@ -4818,7 +4817,7 @@ void FRecastNavMeshGenerator::CalcPolyRefBits(ARecastNavMesh* NavMeshOwner, int3
 {
 	static const int32 TotalBits = (sizeof(dtPolyRef) * 8);
 #if USE_64BIT_ADDRESS
-	MaxTileBits = NavMeshOwner ? FMath::CeilToFloat(FMath::Log2(static_cast<float>(NavMeshOwner->GetTileNumberHardLimit()))) : 20;
+	MaxTileBits = NavMeshOwner ? static_cast<int32>(FMath::CeilToFloat(FMath::Log2(static_cast<float>(NavMeshOwner->GetTileNumberHardLimit())))) : 20;
 	MaxPolyBits = FMath::Min<int32>(32, (TotalBits - DT_MIN_SALT_BITS) - MaxTileBits);
 #else
 	MaxTileBits = 14;
@@ -6805,7 +6804,7 @@ void FRecastNavMeshGenerator::GrabDebugSnapshot(struct FVisualLogEntry* Snapshot
 							continue;
 						}
 
-						const uint8 AreaId = NavData->GetAreaID(AreaMod.GetAreaClass());
+						const uint8 AreaId = IntCastChecked<uint8>(NavData->GetAreaID(AreaMod.GetAreaClass()));
 						const UClass* AreaClass = NavData->GetAreaClass(AreaId);
 						const UNavArea* DefArea = AreaClass ? ((UClass*)AreaClass)->GetDefaultObject<UNavArea>() : NULL;
 						const FColor PolygonColor = AreaClass != FNavigationSystem::GetDefaultWalkableArea() ? (DefArea ? DefArea->DrawColor : NavData->GetConfig().Color) : FColorList::Cyan;
@@ -6958,7 +6957,7 @@ void FRecastNavMeshGenerator::ExportNavigationData(const FString& FileName) cons
 						if (ShapeType == ENavigationShapeType::Convex || ShapeType == ENavigationShapeType::InstancedConvex)
 						{
 							FAreaExportData ExportInfo;
-							ExportInfo.AreaId = NavData->GetAreaID(AreaMod.GetAreaClass());
+							ExportInfo.AreaId = IntCastChecked<uint8>(NavData->GetAreaID(AreaMod.GetAreaClass()));
 
 							auto AddAreaExportDataFunc = [&](const FConvexNavAreaData& InConvexNavAreaData)
 							{
diff --git a/Engine/Source/Runtime/NavigationSystem/Public/NavMesh/RecastNavMeshGenerator.h b/Engine/Source/Runtime/NavigationSystem/Public/NavMesh/RecastNavMeshGenerator.h
--- a/Engine/Source/Runtime/NavigationSystem/Public/NavMesh/RecastNavMeshGenerator.h
+++ b/Engine/Source/Runtime/NavigationSystem/Public/NavMesh/RecastNavMeshGenerator.h
@@ -108,8 +108,8 @@ struct FRecastVoxelCache
 {
 	struct FTileInfo
 	{
-		int16 TileX;
-		int16 TileY;
+		int32 TileX;
+		int32 TileY;
 		int32 NumSpans;
 		FTileInfo* NextTile;
 		rcSpanCache* SpanData;
