First pass at allowing a mix of low, default and high resolution tiles on the navmesh.
CellSize is now deprecated and CellSizes for the various resolutions are stored in the NavMeshResolutionParams array of ARecastNavMesh.
Added NavMeshResolution property to ANavModifierVolume.
Remove duplicated data in dtNavMesh (m_walkableHeight, etc. already in dtNavMeshParams).
Addition of a resolution index to dtMeshHeader.
Increase navmesh version for addition of tile resolutions.

diff --git a/Engine/Config/BaseEngine.ini b/Engine/Config/BaseEngine.ini
--- a/Engine/Config/BaseEngine.ini
+++ b/Engine/Config/BaseEngine.ini
@@ -2501,7 +2501,10 @@ MaxTileGridHeight=256
 DefaultDrawDistance=5000.0
 ; navmesh generation parameters
 TileSizeUU=1000.f
-CellSize=19.f
+CellSize=20.f
+NavMeshResolutionParams[0]=(CellSize=40.000000)
+NavMeshResolutionParams[1]=(CellSize=20.000000)
+NavMeshResolutionParams[2]=(CellSize=10.000000)
 CellHeight=10.f
 AgentRadius=34.f
 AgentHeight=144.f
diff --git a/Engine/Source/Runtime/Engine/Classes/AI/Navigation/NavRelevantInterface.h b/Engine/Source/Runtime/Engine/Classes/AI/Navigation/NavRelevantInterface.h
--- a/Engine/Source/Runtime/Engine/Classes/AI/Navigation/NavRelevantInterface.h
+++ b/Engine/Source/Runtime/Engine/Classes/AI/Navigation/NavRelevantInterface.h
@@ -104,7 +104,10 @@ struct ENGINE_API FNavigationRelevantData : public TSharedFromThis<FNavigationRe
 	FORCEINLINE SIZE_T GetGeometryAllocatedSize() const { return CollisionData.GetAllocatedSize() + VoxelData.GetAllocatedSize(); }
 	FORCEINLINE int32 GetDirtyFlag() const
 	{
-		return ((HasGeometry() || IsPendingLazyGeometryGathering() || Modifiers.GetMaskFillCollisionUnderneathForNavmesh()) ? ENavigationDirtyFlag::Geometry : 0) |
+		const bool bSetGeometryFlag = HasGeometry() || IsPendingLazyGeometryGathering() || Modifiers.GetMaskFillCollisionUnderneathForNavmesh() ||
+			(Modifiers.GetNavMeshResolution() != ENavigationDataResolution::Invalid);
+		
+		return (bSetGeometryFlag ? ENavigationDirtyFlag::Geometry : 0) |
 			((HasModifiers() || NeedAnyPendingLazyModifiersGathering()) ? ENavigationDirtyFlag::DynamicModifier : 0) |
 			(Modifiers.HasAgentHeightAdjust() ? ENavigationDirtyFlag::UseAgentHeight : 0);
 	}
diff --git a/Engine/Source/Runtime/Engine/Classes/AI/Navigation/NavigationTypes.h b/Engine/Source/Runtime/Engine/Classes/AI/Navigation/NavigationTypes.h
--- a/Engine/Source/Runtime/Engine/Classes/AI/Navigation/NavigationTypes.h
+++ b/Engine/Source/Runtime/Engine/Classes/AI/Navigation/NavigationTypes.h
@@ -65,6 +65,16 @@ namespace FNavigationSystem
 	ENGINE_API bool IsInBaseNavmesh(const UObject* Object);
 }
 
+UENUM()
+enum class ENavigationDataResolution : uint8
+{
+	Low = 0,
+	Default = 1,
+	High = 2,
+	Invalid = 3 UMETA(DisplayName="None"),
+	MAX = 3,
+};
+
 UENUM()
 namespace ENavigationOptionFlag
 {
diff --git a/Engine/Source/Runtime/Engine/Private/AI/Navigation/NavigationModifier.cpp b/Engine/Source/Runtime/Engine/Private/AI/Navigation/NavigationModifier.cpp
--- a/Engine/Source/Runtime/Engine/Private/AI/Navigation/NavigationModifier.cpp
+++ b/Engine/Source/Runtime/Engine/Private/AI/Navigation/NavigationModifier.cpp
@@ -755,6 +755,7 @@ void FCompositeNavModifier::Reset()
 	bIsPerInstanceModifier = false;
 	bFillCollisionUnderneathForNavmesh = false;
 	bMaskFillCollisionUnderneathForNavmesh = false;
+	NavMeshResolution = ENavigationDataResolution::Invalid;
 }
 
 void FCompositeNavModifier::Empty()
@@ -766,6 +767,7 @@ void FCompositeNavModifier::Empty()
 	bAdjustHeight = false;
 	bFillCollisionUnderneathForNavmesh = false;
 	bMaskFillCollisionUnderneathForNavmesh = false;
+	NavMeshResolution = ENavigationDataResolution::Invalid;
 }
 
 FCompositeNavModifier FCompositeNavModifier::GetInstantiatedMetaModifier(const FNavAgentProperties* NavAgent, TWeakObjectPtr<UObject> WeakOwnerPtr) const
diff --git a/Engine/Source/Runtime/Engine/Public/AI/NavigationModifier.h b/Engine/Source/Runtime/Engine/Public/AI/NavigationModifier.h
--- a/Engine/Source/Runtime/Engine/Public/AI/NavigationModifier.h
+++ b/Engine/Source/Runtime/Engine/Public/AI/NavigationModifier.h
@@ -263,6 +263,7 @@ struct ENGINE_API FCompositeNavModifier : public FNavigationModifier
 		, bIsPerInstanceModifier(false)
 		, bFillCollisionUnderneathForNavmesh(false)
 		, bMaskFillCollisionUnderneathForNavmesh(false)
+		, NavMeshResolution(ENavigationDataResolution::Invalid)
 	{}
 
 	void Shrink();
@@ -306,6 +307,10 @@ struct ENGINE_API FCompositeNavModifier : public FNavigationModifier
 		bHasLowAreaModifiers |= Modifiers.HasLowAreaModifiers();
 		bFillCollisionUnderneathForNavmesh |= Modifiers.GetFillCollisionUnderneathForNavmesh();
 		bMaskFillCollisionUnderneathForNavmesh |= Modifiers.GetMaskFillCollisionUnderneathForNavmesh();
+		if (Modifiers.GetNavMeshResolution() != ENavigationDataResolution::Invalid)
+		{
+			NavMeshResolution = FMath::Max(NavMeshResolution, Modifiers.GetNavMeshResolution());	// Pick the highest resolution
+		}
 	}
 
 	void CreateAreaModifiers(const UPrimitiveComponent* PrimComp, const TSubclassOf<UNavAreaBase> AreaClass);
@@ -325,6 +330,8 @@ struct ENGINE_API FCompositeNavModifier : public FNavigationModifier
 	FORCEINLINE void SetFillCollisionUnderneathForNavmesh(bool bValue) { bFillCollisionUnderneathForNavmesh = bValue; }
 	FORCEINLINE bool GetMaskFillCollisionUnderneathForNavmesh() const { return bMaskFillCollisionUnderneathForNavmesh; }
 	FORCEINLINE void SetMaskFillCollisionUnderneathForNavmesh(bool bValue) { bMaskFillCollisionUnderneathForNavmesh = bValue; }
+	FORCEINLINE ENavigationDataResolution GetNavMeshResolution() const { return NavMeshResolution; }
+	FORCEINLINE void SetNavMeshResolution(ENavigationDataResolution Resolution) { NavMeshResolution = Resolution; }
 	FORCEINLINE void ReserveForAdditionalAreas(int32 AdditionalElementsCount) { Areas.Reserve(Areas.Num() + AdditionalElementsCount); }
 
 	void MarkPotentialLinks() { bHasPotentialLinks = true; }
@@ -357,4 +364,5 @@ private:
     uint32 bIsPerInstanceModifier : 1;
 	uint32 bFillCollisionUnderneathForNavmesh : 1;
 	uint32 bMaskFillCollisionUnderneathForNavmesh : 1;
+	ENavigationDataResolution NavMeshResolution;
 };
diff --git a/Engine/Source/Runtime/NavigationSystem/Private/NavMesh/NavMeshRenderingComponent.cpp b/Engine/Source/Runtime/NavigationSystem/Private/NavMesh/NavMeshRenderingComponent.cpp
--- a/Engine/Source/Runtime/NavigationSystem/Private/NavMesh/NavMeshRenderingComponent.cpp
+++ b/Engine/Source/Runtime/NavigationSystem/Private/NavMesh/NavMeshRenderingComponent.cpp
@@ -575,7 +575,11 @@ void FNavMeshSceneProxyData::GatherData(const ARecastNavMesh* NavMesh, int32 InN
 			const FString GenerationMode = Mode == ERuntimeGenerationType::Static ? TEXT("Static") :
 				(Mode == ERuntimeGenerationType::Dynamic ? TEXT("Dynamic") : (Mode == ERuntimeGenerationType::DynamicModifiersOnly ? TEXT("DynamicModifersOnly") : TEXT("Unknown")));
 			DebugLabels.Add(FDebugText(FString::Printf(TEXT("%s (%s%s)"), *NavMesh->GetName(), NavMesh->bIsWorldPartitioned ? TEXT("WP ") : TEXT(""), *GenerationMode)));
-			DebugLabels.Add(FDebugText(FString::Printf(TEXT("AgentRadius %0.1f, AgentHeight %0.1f, CellSize %0.1f, CellHeight %0.1f"), NavMesh->AgentRadius, NavMesh->AgentHeight, NavMesh->CellSize, NavMesh->CellHeight)));
+			DebugLabels.Add(FDebugText(FString::Printf(TEXT("AgentRadius %0.1f, AgentHeight %0.1f, CellSizes %0.1f/%0.1f/%0.1f (low/default/high), CellHeight %0.1f"),
+				NavMesh->AgentRadius,
+				NavMesh->AgentHeight,
+				NavMesh->GetCellSize(ENavigationDataResolution::Low), NavMesh->GetCellSize(ENavigationDataResolution::Default), NavMesh->GetCellSize(ENavigationDataResolution::High),
+				NavMesh->CellHeight)));
 			DebugLabels.Add(FDebugText(FString::Printf(TEXT("Region part %s, Layer part %s"), *GetPartitioningString(NavMesh->RegionPartitioning), *GetPartitioningString(NavMesh->LayerPartitioning))));
 			DebugLabels.Add(FDebugText(TEXT(""))); // empty line
 
diff --git a/Engine/Source/Runtime/NavigationSystem/Private/NavMesh/PImplRecastNavMesh.cpp b/Engine/Source/Runtime/NavigationSystem/Private/NavMesh/PImplRecastNavMesh.cpp
--- a/Engine/Source/Runtime/NavigationSystem/Private/NavMesh/PImplRecastNavMesh.cpp
+++ b/Engine/Source/Runtime/NavigationSystem/Private/NavMesh/PImplRecastNavMesh.cpp
@@ -557,19 +557,39 @@ void FPImplRecastNavMesh::Serialize( FArchive& Ar, int32 NavMeshVersion )
 	Ar << Params.maxTiles;				///< The maximum number of tiles the navigation mesh can contain.
 	Ar << Params.maxPolys;
 
-	if(NavMeshOwner->NavMeshVersion >= NAVMESHVER_OPTIM_FIX_SERIALIZE_PARAMS)
+	if (NavMeshOwner->NavMeshVersion >= NAVMESHVER_TILE_RESOLUTIONS)
 	{
 		Ar << Params.walkableHeight;
 		Ar << Params.walkableRadius;
 		Ar << Params.walkableClimb;
-		Ar << Params.bvQuantFactor;
+		Ar << Params.resolutionParams[(uint8)ENavigationDataResolution::Low].bvQuantFactor;
+		Ar << Params.resolutionParams[(uint8)ENavigationDataResolution::Default].bvQuantFactor;
+		Ar << Params.resolutionParams[(uint8)ENavigationDataResolution::High].bvQuantFactor;
+	}
+	else if (NavMeshOwner->NavMeshVersion >= NAVMESHVER_OPTIM_FIX_SERIALIZE_PARAMS)
+	{
+		// Load previous version navmesh data into new struct
+		Ar << Params.walkableHeight;
+		Ar << Params.walkableRadius;
+		Ar << Params.walkableClimb;
+		Ar << Params.resolutionParams[(uint8)ENavigationDataResolution::Default].bvQuantFactor;
+		if (Ar.IsLoading())
+		{
+			const dtReal DefaultQuantFactor = Params.resolutionParams[(uint8)ENavigationDataResolution::Default].bvQuantFactor; 
+			Params.resolutionParams[(uint8)ENavigationDataResolution::Low].bvQuantFactor = DefaultQuantFactor;
+			Params.resolutionParams[(uint8)ENavigationDataResolution::High].bvQuantFactor = DefaultQuantFactor;
+		}
 	}
 	else
 	{
 		Params.walkableHeight = NavMeshOwner->AgentHeight;
 		Params.walkableRadius = NavMeshOwner->AgentRadius;
 		Params.walkableClimb = NavMeshOwner->AgentMaxStepHeight;
-		Params.bvQuantFactor = 1.f / NavMeshOwner->CellSize;
+		const float DefaultQuantFactor =  1.f / NavMeshOwner->GetCellSize(ENavigationDataResolution::Default);
+		for(uint8 Index = 0; Index < (uint8)ENavigationDataResolution::MAX; Index++)
+		{
+			Params.resolutionParams[Index].bvQuantFactor = DefaultQuantFactor;	
+		}
 	}
 
 	if (Ar.IsLoading())
@@ -577,7 +597,8 @@ void FPImplRecastNavMesh::Serialize( FArchive& Ar, int32 NavMeshVersion )
 		// at this point we can tell whether navmesh being loaded is in line
 		// ARecastNavMesh's params. If not, just skip it.
 		// assumes tiles are rectangular
-		const FVector::FReal ActorsTileSize = FVector::FReal(int32(NavMeshOwner->TileSizeUU / NavMeshOwner->CellSize) * NavMeshOwner->CellSize);
+		const float DefaultCellSize = NavMeshOwner->GetCellSize(ENavigationDataResolution::Default);
+		const FVector::FReal ActorsTileSize = FVector::FReal(int32(NavMeshOwner->TileSizeUU / DefaultCellSize) * DefaultCellSize);
 
 		if (ActorsTileSize != Params.tileWidth)
 		{
@@ -783,6 +804,12 @@ void FPImplRecastNavMesh::SerializeRecastMeshTile(FArchive& Ar, int32 NavMeshVer
 		Ar << Header->layer << Header->polyCount << Header->vertCount;
 		Ar << Header->maxLinkCount << Header->detailMeshCount << Header->detailVertCount << Header->detailTriCount;
 		Ar << Header->bvNodeCount << Header->offMeshConCount<< Header->offMeshBase;
+		
+		if (NavMeshVersion >= NAVMESHVER_TILE_RESOLUTIONS)
+		{
+			Ar << Header->resolution;
+		}
+		
 		Ar << Header->bmin[0] << Header->bmin[1] << Header->bmin[2];
 		Ar << Header->bmax[0] << Header->bmax[1] << Header->bmax[2];
 #if WITH_NAVMESH_CLUSTER_LINKS
diff --git a/Engine/Source/Runtime/NavigationSystem/Private/NavMesh/RecastNavMesh.cpp b/Engine/Source/Runtime/NavigationSystem/Private/NavMesh/RecastNavMesh.cpp
--- a/Engine/Source/Runtime/NavigationSystem/Private/NavMesh/RecastNavMesh.cpp
+++ b/Engine/Source/Runtime/NavigationSystem/Private/NavMesh/RecastNavMesh.cpp
@@ -424,7 +424,9 @@ FRecastNavMeshGenerationProperties::FRecastNavMeshGenerationProperties(const ARe
 {
 	TilePoolSize = RecastNavMesh.TilePoolSize;
 	TileSizeUU = RecastNavMesh.TileSizeUU;
-	CellSize = RecastNavMesh.CellSize;
+	
+	// @todo: Handle navmesh resolution if FRecastNavMeshGenerationProperties is used.
+	CellSize = RecastNavMesh.GetCellSize(ENavigationDataResolution::Default);
 	CellHeight = RecastNavMesh.CellHeight;
 	AgentRadius = RecastNavMesh.AgentRadius;
 	AgentHeight = RecastNavMesh.AgentHeight;
@@ -629,8 +631,22 @@ void ARecastNavMesh::PostLoad()
 		}
 	}
 
-	UE_CLOG(TileSizeUU < CellSize, LogNavigation, Error, TEXT("%s: TileSizeUU (%f) being less than CellSize (%f) is an invalid case and will cause navmesh generation issues.")
-		, *GetName(), TileSizeUU, CellSize);
+	PRAGMA_DISABLE_DEPRECATION_WARNINGS
+	// If needed, initialize from deprecated value.
+	if (NavMeshVersion < NAVMESHVER_TILE_RESOLUTIONS)
+	{
+		for (int i = 0; i < (uint8)ENavigationDataResolution::MAX; ++i)
+		{
+			SetCellSize((ENavigationDataResolution)i, CellSize);
+		}
+	}
+	PRAGMA_ENABLE_DEPRECATION_WARNINGS
+	
+	for (uint8 Index = 0; Index < (uint8)ENavigationDataResolution::MAX; Index++)
+	{
+		UE_CLOG(TileSizeUU < GetCellSize((ENavigationDataResolution)Index), LogNavigation, Error, TEXT("%s: TileSizeUU (%f) being less than CellSize (%f) is an invalid case and will cause navmesh generation issues.")
+			, *GetName(), TileSizeUU, GetCellSize((ENavigationDataResolution)Index));
+	}
 	
 	if (!UWorld::IsPartitionedWorld(GetWorld()))
 	{
@@ -711,12 +727,17 @@ void ARecastNavMesh::PostInitProperties()
 			TileSizeUU = DefOb->TileSizeUU;
 		}
 
-		if (CellSize != DefOb->CellSize)
+		for (int i = 0; i < (uint8)ENavigationDataResolution::MAX; ++i)
 		{
-			UE_LOG(LogNavigation, Warning, TEXT("%s param: CellSize(%f) differs from config settings, forcing value %f so it can be used with voxel cache!"),
-				*GetNameSafe(this), CellSize, DefOb->CellSize);
+			const float CurrentCellSize = NavMeshResolutionParams[i].CellSize;
+			const float DefaultObjectCellSize = DefOb->NavMeshResolutionParams[i].CellSize;
+			if (CurrentCellSize != DefaultObjectCellSize)
+			{
+				UE_LOG(LogNavigation, Warning, TEXT("%s param: CellSize(%f) differs from config settings, forcing value %f so it can be used with voxel cache!"),
+					*GetNameSafe(this), CurrentCellSize, DefaultObjectCellSize);
 
-			CellSize = DefOb->CellSize;
+				NavMeshResolutionParams[i].CellSize = DefaultObjectCellSize;
+			}
 		}
 
 		if (CellHeight != DefOb->CellHeight)
@@ -943,7 +964,7 @@ void ARecastNavMesh::OnRegistered()
 	check(RecastNavMeshImpl);
 
 	// This check can fail when the NavMeshVersion indicates the map needs the nav mesh rebuilt
-	ensure(RecastNavMeshImpl->GetRecastMesh() == nullptr || RecastNavMeshImpl->GetRecastMesh()->getBVQuantFactor() != 0);
+	ensure(RecastNavMeshImpl->GetRecastMesh() == nullptr || RecastNavMeshImpl->GetRecastMesh()->getBVQuantFactor((uint8)ENavigationDataResolution::Default) != 0);
 }
 
 void ARecastNavMesh::SerializeRecastNavMesh(FArchive& Ar, FPImplRecastNavMesh*& NavMesh, int32 InNavMeshVersion)
@@ -1330,6 +1351,13 @@ bool ARecastNavMesh::IsResizable() const
 	return !bFixedTilePoolSize;
 }
 
+float ARecastNavMesh::GetTileSizeUU() const
+{
+	const float DefaultCellSize = GetCellSize(ENavigationDataResolution::Default);
+	const float RcTileSize = FMath::TruncToFloat(TileSizeUU / DefaultCellSize);
+	return RcTileSize * DefaultCellSize;
+}
+
 void ARecastNavMesh::GetEdgesForPathCorridor(const TArray<NavNodeRef>* PathCorridor, TArray<FNavigationPortalEdge>* PathCorridorEdges) const
 {
 	check(PathCorridor != NULL && PathCorridorEdges != NULL);
@@ -2084,12 +2112,11 @@ int32 ARecastNavMesh::ReplaceAreaInTileBounds(const FBox& Bounds, TSubclassOf<UN
 		ensure(NewAreaID != OldAreaID);
 
 		// workaround for privacy issue in the recast API
-		dtNavMesh* DetourNavMesh = RecastNavMeshImpl->GetRecastMesh();
+		const dtNavMesh* DetourNavMesh = RecastNavMeshImpl->GetRecastMesh();
 		dtNavMesh const* const ConstDetourNavMesh = RecastNavMeshImpl->GetRecastMesh();
 
 		const FVector RcNavMeshOrigin = Unreal2RecastPoint(NavMeshOriginOffset);
-		const float RcTileSize = FMath::TruncToFloat(TileSizeUU / CellSize);
-		const float TileSizeInWorldUnits = RcTileSize * CellSize;
+		const float TileSizeInWorldUnits = GetTileSizeUU();
 		const FRcTileBox TileBox(Bounds, RcNavMeshOrigin, TileSizeInWorldUnits);
 
 		for (int32 TileY = TileBox.YMin; TileY <= TileBox.YMax; ++TileY)
@@ -2111,7 +2138,7 @@ int32 ARecastNavMesh::ReplaceAreaInTileBounds(const FBox& Bounds, TSubclassOf<UN
 					if (TileRef)
 					{
 						const int32 TileIndex = (int32)ConstDetourNavMesh->decodePolyIdTile(TileRef);
-						const dtMeshTile* Tile = ((const dtNavMesh*)DetourNavMesh)->getTile(TileIndex);
+						const dtMeshTile* Tile = DetourNavMesh->getTile(TileIndex);
 						//const int32 MaxPolys = Tile && Tile->header ? Tile->header->offMeshBase : 0;
 						const int32 MaxPolys = Tile && Tile->header
 							? (ReplaceLinks ? Tile->header->polyCount : Tile->header->offMeshBase)
@@ -3018,6 +3045,66 @@ void ARecastNavMesh::UpdateNavVersion()
 
 #if WITH_EDITOR
 
+void ARecastNavMesh::PostEditChangeChainProperty(FPropertyChangedChainEvent& PropertyChangedChainEvent)
+{
+	static const FName NAME_Generation = FName(TEXT("Generation"));
+	static const FName NAME_NavMeshResolutionParams = FName(TEXT("NavMeshResolutionParams"));
+
+	Super::PostEditChangeChainProperty(PropertyChangedChainEvent);
+	
+	if (PropertyChangedChainEvent.Property != NULL)
+	{
+		const FName CategoryName = FObjectEditorUtils::GetCategoryFName(PropertyChangedChainEvent.Property);
+		if (CategoryName == NAME_Generation)
+		{
+			const FName PropName = PropertyChangedChainEvent.Property->GetFName();
+			
+			if (PropName == GET_MEMBER_NAME_CHECKED(FNavMeshResolutionParam, CellSize))
+			{
+				const int32 ChangedIndex = PropertyChangedChainEvent.GetArrayIndex(NAME_NavMeshResolutionParams.ToString());
+				if (ChangedIndex != INDEX_NONE)
+				{
+					float& RefCellSize = NavMeshResolutionParams[ChangedIndex].CellSize;
+					RefCellSize = UE::NavMesh::Private::GetClampedCellSize(RefCellSize);
+				 
+					TileSizeUU = UE::NavMesh::Private::GetClampedTileSizeUU(TileSizeUU, RefCellSize, AgentRadius);
+
+					// Adjust tile size to be a multiple of RefCellSize
+					const float RefCellCount = FMath::TruncToFloat( TileSizeUU / RefCellSize);
+					TileSizeUU = RefCellCount * RefCellSize;
+
+					// Adjust the other cell size (find count of cells and set the size of the cell)
+					for (uint8 Index = 0; Index < (uint8)ENavigationDataResolution::MAX; Index++)
+					{
+						if (Index != ChangedIndex)
+						{
+							float& ResolutionCellSize = NavMeshResolutionParams[Index].CellSize;
+							const float ResolutionCellCount = FMath::TruncToFloat(TileSizeUU / ResolutionCellSize);
+							ResolutionCellSize = UE::NavMesh::Private::GetClampedCellSize(TileSizeUU / ResolutionCellCount);
+						}
+					}
+
+					PRAGMA_DISABLE_DEPRECATION_WARNINGS
+					// Update the deprecated CellSize to fit the default resolution CellSize
+					CellSize = NavMeshResolutionParams[(uint8)ENavigationDataResolution::Default].CellSize;
+					PRAGMA_ENABLE_DEPRECATION_WARNINGS
+
+					// update config
+					FillConfig(NavDataConfig);
+
+					UNavigationSystemV1* NavSys = FNavigationSystem::GetCurrent<UNavigationSystemV1>(GetWorld());
+					if (!HasAnyFlags(RF_ClassDefaultObject)
+						&& NavSys && NavSys->GetIsAutoUpdateEnabled()
+						&& PropName != GET_MEMBER_NAME_CHECKED(ARecastNavMesh, MaxSimultaneousTileGenerationJobsCount))
+					{
+						RebuildAll();
+					}
+				}
+			}
+		}
+	}
+}
+
 void ARecastNavMesh::PostEditChangeProperty(FPropertyChangedEvent& PropertyChangedEvent)
 {
 	static const FName NAME_Generation = FName(TEXT("Generation"));
@@ -3034,6 +3121,7 @@ void ARecastNavMesh::PostEditChangeProperty(FPropertyChangedEvent& PropertyChang
 		if (CategoryName == NAME_Generation)
 		{
 			const FName PropName = PropertyChangedEvent.Property->GetFName();
+			const FName MemberName = PropertyChangedEvent.MemberProperty->GetFName();
 			
 			if (PropName == GET_MEMBER_NAME_CHECKED(ARecastNavMesh, AgentRadius))
 			{
@@ -3047,22 +3135,19 @@ void ARecastNavMesh::PostEditChangeProperty(FPropertyChangedEvent& PropertyChang
 			}
 			else if (PropName == GET_MEMBER_NAME_CHECKED(ARecastNavMesh, TileSizeUU))
 			{
-				CellSize = UE::NavMesh::Private::GetClampedCellSize(CellSize);
-				TileSizeUU = UE::NavMesh::Private::GetClampedTileSizeUU(TileSizeUU, CellSize, AgentRadius);
-						
-				// Match cell size to tile size.
-				CellSize = TileSizeUU / FMath::TruncToFloat(TileSizeUU / CellSize);
+				SetCellSize(ENavigationDataResolution::Default, UE::NavMesh::Private::GetClampedCellSize(GetCellSize(ENavigationDataResolution::Default)));
+				TileSizeUU = UE::NavMesh::Private::GetClampedTileSizeUU(TileSizeUU, GetCellSize(ENavigationDataResolution::Default), AgentRadius);
 
-				// update config
-				FillConfig(NavDataConfig);
-			}
-			else if (PropName == GET_MEMBER_NAME_CHECKED(ARecastNavMesh, CellSize))
-			{
-				CellSize = UE::NavMesh::Private::GetClampedCellSize(CellSize);
-				TileSizeUU = UE::NavMesh::Private::GetClampedTileSizeUU(TileSizeUU, CellSize, AgentRadius);
-				
-				// Match tile size to cell size.
-				TileSizeUU = CellSize * FMath::TruncToFloat(TileSizeUU / CellSize);
+				// Match cell sizes to tile size.
+				for (uint8 Index = 0; Index < (uint8)ENavigationDataResolution::MAX; Index++)
+				{
+					SetCellSize((ENavigationDataResolution)Index, TileSizeUU / FMath::TruncToFloat(TileSizeUU / GetCellSize((ENavigationDataResolution)Index)));	
+				}
+
+				PRAGMA_DISABLE_DEPRECATION_WARNINGS
+				// Set deprecated CellSize
+				CellSize = GetCellSize(ENavigationDataResolution::Default);
+				PRAGMA_ENABLE_DEPRECATION_WARNINGS
 
 				// update config
 				FillConfig(NavDataConfig);
@@ -3177,7 +3262,9 @@ void ARecastNavMesh::UpdateGenerationProperties(const FRecastNavMeshGenerationPr
 {
 	TilePoolSize = GenerationProps.TilePoolSize;
 	TileSizeUU = GenerationProps.TileSizeUU;
+	PRAGMA_DISABLE_DEPRECATION_WARNINGS
 	CellSize = GenerationProps.CellSize;
+	PRAGMA_ENABLE_DEPRECATION_WARNINGS
 	CellHeight = GenerationProps.CellHeight;
 	AgentRadius = GenerationProps.AgentRadius;
 	AgentHeight = GenerationProps.AgentHeight;
@@ -3246,7 +3333,7 @@ bool ARecastNavMesh::HasCompleteDataInRadius(const FVector& TestLocation, FVecto
 
 	const dtNavMesh* NavMesh = RecastNavMeshImpl->DetourNavMesh;
 	const dtNavMeshParams* NavParams = RecastNavMeshImpl->DetourNavMesh->getParams();
-	const float NavTileSize = CellSize * FMath::TruncToFloat(TileSizeUU / CellSize);
+	const float NavTileSize = GetTileSizeUU();
 	const FVector RcNavOrigin(NavParams->orig[0], NavParams->orig[1], NavParams->orig[2]);
 
 	const FBox RcBounds = Unreal2RecastBox(FBox::BuildAABB(TestLocation, FVector(TestRadius, TestRadius, 0)));
@@ -3305,7 +3392,7 @@ void ARecastNavMesh::UpdateActiveTiles(const TArray<FNavigationInvokerRaw>& Invo
 	check(NavParams && MyGenerator);
 	const FRecastBuildConfig& Config = MyGenerator->GetConfig();
 	const FVector NavmeshOrigin = Recast2UnrealPoint(NavParams->orig);
-	const FVector::FReal TileDim = Config.tileSize * Config.cs;
+	const FVector::FReal TileDim = Config.GetTileSizeUU();
 
 	TArray<FIntPoint>& ActiveTiles = GetActiveTiles();
 	TArray<FIntPoint> OldActiveSet = ActiveTiles;
@@ -3362,11 +3449,12 @@ void ARecastNavMesh::UpdateActiveTiles(const TArray<FNavigationInvokerRaw>& Invo
 		}
 	}
 
+	// Find tiles to update
 	TArray<FIntPoint> TilesToUpdate;
 	TilesToUpdate.Reserve(ActiveTiles.Num());
 	for (int32 Index = TilesInMinDistance.Num() - 1; Index >= 0; --Index)
 	{
-		// check if it's a new tile
+		// Check if it's a new tile (not in the active set)
 		const FIntPoint& Tile = TilesInMinDistance[Index];
 		if (OldActiveSet.Find(Tile) == INDEX_NONE)
 		{
@@ -3447,8 +3535,7 @@ void ARecastNavMesh::DirtyTilesInBounds(const FBox& Bounds)
 		// Add tiles within the overlapping bounds
 		TArray<FIntPoint> Points;
 		const FVector RcNavMeshOrigin = Unreal2RecastPoint(NavMeshOriginOffset);
-		const float RcTileSize = FMath::TruncToFloat(TileSizeUU / CellSize);
-		const float TileSizeInWorldUnits = RcTileSize * CellSize;
+		const float TileSizeInWorldUnits = GetTileSizeUU();
 		const FRcTileBox TileBox(OverlappingBounds, RcNavMeshOrigin, TileSizeInWorldUnits);
 
 		UE_LOG(LogNavigation, VeryVerbose, TEXT("RebuildTilesFromBounds %i tiles: (%i,%i) to (%i,%i)"), (TileBox.XMax-TileBox.XMin)*(TileBox.YMax-TileBox.YMin), TileBox.XMin, TileBox.YMin, TileBox.XMax, TileBox.YMax);
diff --git a/Engine/Source/Runtime/NavigationSystem/Private/NavMesh/RecastNavMeshDataChunk.cpp b/Engine/Source/Runtime/NavigationSystem/Private/NavMesh/RecastNavMeshDataChunk.cpp
--- a/Engine/Source/Runtime/NavigationSystem/Private/NavMesh/RecastNavMeshDataChunk.cpp
+++ b/Engine/Source/Runtime/NavigationSystem/Private/NavMesh/RecastNavMeshDataChunk.cpp
@@ -518,7 +518,8 @@ void URecastNavMeshDataChunk::MoveTiles(FPImplRecastNavMesh& NavMeshImpl, const
 
 				const int OffsetWithRotX = Offset.X + DeltaX;
 				const int OffsetWithRotY = Offset.Y + DeltaY;
-				const bool bSuccess = dtTransformTileData(TileData.TileRawData->RawData, TileData.TileDataSize, OffsetWithRotX, OffsetWithRotY, TileWidth, TileHeight, RotationDeg, NavMesh->getBVQuantFactor());
+
+				const bool bSuccess = dtTransformTileData(TileData.TileRawData->RawData, TileData.TileDataSize, OffsetWithRotX, OffsetWithRotY, TileWidth, TileHeight, RotationDeg, NavMesh->getBVQuantFactor(Header->resolution));
 				UE_CLOG(bSuccess, LogNavigation, Verbose, TEXT("   Moved tile from (%i,%i) to (%i,%i)."), TileData.OriginalX, TileData.OriginalY, (TileData.OriginalX + OffsetWithRotX), (TileData.OriginalY + OffsetWithRotY));
 			}
 		}
diff --git a/Engine/Source/Runtime/NavigationSystem/Private/NavMesh/RecastNavMeshGenerator.cpp b/Engine/Source/Runtime/NavigationSystem/Private/NavMesh/RecastNavMeshGenerator.cpp
--- a/Engine/Source/Runtime/NavigationSystem/Private/NavMesh/RecastNavMeshGenerator.cpp
+++ b/Engine/Source/Runtime/NavigationSystem/Private/NavMesh/RecastNavMeshGenerator.cpp
@@ -1632,7 +1632,9 @@ FRecastTileGenerator::FRecastTileGenerator(FRecastNavMeshGenerator& ParentGenera
 	TileX = Location.X;
 	TileY = Location.Y;
 
+	// Copy tile config from parent generator
 	TileConfig = ParentGenerator.GetConfig();
+	
 	TileDebugSettings = ParentGenerator.GetTileDebugSettings();
 	Version = ParentGenerator.GetVersion();
 	AdditionalCachedData = ParentGenerator.GetAdditionalCachedData();
@@ -1665,9 +1667,11 @@ FRecastTileGenerator::~FRecastTileGenerator()
 
 void FRecastTileGenerator::Setup(const FRecastNavMeshGenerator& ParentGenerator, const TArray<FBox>& DirtyAreas)
 {
+	QUICK_SCOPE_CYCLE_COUNTER(STAT_FRecastTileGenerator_Setup);
+	
 	const FVector RcNavMeshOrigin = ParentGenerator.GetRcNavMeshOrigin();
 	const FBox NavTotalBounds = ParentGenerator.GetTotalBounds();
-	const FVector::FReal TileSizeUU = (TileConfig.tileSize * TileConfig.cs);
+	const FVector::FReal TileSizeUU = TileConfig.GetTileSizeUU();
 
 	NavDataConfig = ParentGenerator.GetOwner()->GetConfig();
 
@@ -1977,17 +1981,43 @@ void FRecastTileGenerator::GatherGeometry(const FRecastNavMeshGenerator& ParentG
 	const bool bUseVirtualGeometryFilteringAndDirtying = OwnerNav->bUseVirtualGeometryFilteringAndDirtying;
 	const FNavDataConfig& OwnerNavDataConfig = OwnerNav->GetConfig();
 
-	NavigationOctree->FindElementsWithBoundsTest(ParentGenerator.GrowBoundingBox(TileBB, /*bIncludeAgentHeight*/ false),
-		[&OwnerNavDataConfig, &ParentGenerator, this, NavSys, bGeometryChanged, bUseVirtualGeometryFilteringAndDirtying](const FNavigationOctreeElement& Element)
+	TArray<TSharedRef<FNavigationRelevantData, ESPMode::ThreadSafe>> RelevantDataArray;
+
+	const FBox NewBounds = ParentGenerator.GrowBoundingBox(TileBB, /*bIncludeAgentHeight*/ false);
+	ENavigationDataResolution HighestResolution = ENavigationDataResolution::Low;
+	bool bNewResolutionFound = false;
+	
+	NavigationOctree->FindElementsWithBoundsTest(NewBounds,
+		[&HighestResolution, &bNewResolutionFound, &RelevantDataArray, &OwnerNavDataConfig, &ParentGenerator, this, NavSys, bGeometryChanged, bUseVirtualGeometryFilteringAndDirtying](const FNavigationOctreeElement& Element)
 	{
 		const bool bShouldUse = bUseVirtualGeometryFilteringAndDirtying ?
 			ParentGenerator.ShouldGenerateGeometryForOctreeElement(Element, OwnerNavDataConfig) :
 			Element.ShouldUseGeometry(OwnerNavDataConfig);
 		if (bShouldUse)
 		{
-			GatherNavigationDataGeometry(Element.Data, *NavSys, OwnerNavDataConfig, bGeometryChanged);
+			RelevantDataArray.Add(Element.Data);
+
+			// Keep highest resolution that is not the default.
+			const ENavigationDataResolution Resolution = Element.Data->Modifiers.GetNavMeshResolution();
+			if (Resolution != ENavigationDataResolution::Invalid)
+			{
+				HighestResolution = FMath::Max(HighestResolution, Resolution);
+				bNewResolutionFound = true;
+			}
 		}
 	});
+
+	check(HighestResolution != ENavigationDataResolution::Invalid);
+	if (bNewResolutionFound && ParentGenerator.GetOwner()->NavMeshResolutionParams[(uint8)HighestResolution].IsValid())
+	{
+		// Update the TileConfig
+		ParentGenerator.SetupTileConfig(HighestResolution, TileConfig);
+	}
+	
+	for(TSharedRef<FNavigationRelevantData, ESPMode::ThreadSafe> ElementData : RelevantDataArray)
+	{
+		GatherNavigationDataGeometry(ElementData, *NavSys, OwnerNavDataConfig, bGeometryChanged);
+	}
 }
 
 void FRecastTileGenerator::GatherNavigationDataGeometry(const TSharedRef<FNavigationRelevantData, ESPMode::ThreadSafe>& ElementData, UNavigationSystemV1& NavSys, const FNavDataConfig& OwnerNavDataConfig, const bool bGeometryChanged)
@@ -3782,6 +3812,7 @@ bool FRecastTileGenerator::GenerateNavigationDataLayer(FNavMeshBuildContext& Bui
 		rcVcopy(Params.bmax, GenerationContext.Layer->header->bmax);
 		Params.cs = TileConfig.cs;
 		Params.ch = TileConfig.ch;
+		Params.tileResolutionLevel = (unsigned char)TileConfig.TileResolution;
 		Params.buildBvTree = TileConfig.bGenerateBVTree;
 
 #if WITH_NAVMESH_CLUSTER_LINKS
@@ -4514,10 +4545,35 @@ FRecastNavMeshGenerator::~FRecastNavMeshGenerator()
 	DEC_DWORD_STAT_BY( STAT_NavigationMemory, sizeof(*this) );
 }
 
+void FRecastNavMeshGenerator::SetupTileConfig(const ENavigationDataResolution TileResolution, FRecastBuildConfig& OutConfig) const
+{
+	check(GetOwner());
+	ensure(GetConfig().cs == GetOwner()->GetCellSize(ENavigationDataResolution::Default));
+	const float CellSize = GetOwner()->GetCellSize(TileResolution);
+	
+	// Update all settings that depends directly on indirectly of the CellSize
+	OutConfig.TileResolution = TileResolution;
+	OutConfig.cs = CellSize;
+	OutConfig.walkableRadius = FMath::CeilToInt(DestNavMesh->AgentRadius / CellSize);
+
+	OutConfig.borderSize = OutConfig.walkableRadius + 3; // +1 for voxelization rounding, +1 for ledge neighbor access, +1 for occasional errors
+	OutConfig.maxEdgeLen = (int32)(1200.0f / CellSize);
+
+	OutConfig.minRegionArea = (int32)rcSqr(DestNavMesh->MinRegionArea / CellSize);
+	OutConfig.mergeRegionArea = (int32)rcSqr(DestNavMesh->MergeRegionSize / CellSize);
+
+	OutConfig.tileSize = FMath::Max(FMath::TruncToInt(DestNavMesh->TileSizeUU / CellSize), 1);
+	UE_CLOG(OutConfig.tileSize == 1, LogNavigation, Error, TEXT("RecastNavMesh TileSize of 1 is highly discouraged. This occurence indicates an issue with RecastNavMesh\'s generation properties (specifically TileSizeUU: %f, CellSize: %f). Please ensure their correctness.")
+		, DestNavMesh->TileSizeUU, CellSize);
+
+	OutConfig.regionChunkSize = FMath::Max(1, OutConfig.tileSize / FMath::Max(1, DestNavMesh->LayerChunkSplits));
+	OutConfig.TileCacheChunkSize = FMath::Max(1, OutConfig.tileSize / FMath::Max(1, DestNavMesh->RegionChunkSplits));
+}
+
 void FRecastNavMeshGenerator::ConfigureBuildProperties(FRecastBuildConfig& OutConfig)
 {
 	// @TODO those variables should be tweakable per navmesh actor
-	const float CellSize = DestNavMesh->CellSize;
+	const float CellSize = DestNavMesh->GetCellSize(ENavigationDataResolution::Default);
 	const float CellHeight = DestNavMesh->CellHeight;
 	const float AgentHeight = DestNavMesh->AgentHeight;
 	const float AgentMaxSlope = DestNavMesh->AgentMaxSlope;
@@ -4616,7 +4672,7 @@ void FRecastNavMeshGenerator::Init()
 		}
 	}
 
-	ensure(DestNavMesh->GetRecastMesh() == nullptr || DestNavMesh->GetRecastMesh()->getBVQuantFactor() != 0);
+	ensure(DestNavMesh->GetRecastMesh() == nullptr || DestNavMesh->GetRecastMesh()->getBVQuantFactor((unsigned char)ENavigationDataResolution::Default) != 0);
 
 	UpdateNavigationBounds();
 
@@ -4654,7 +4710,7 @@ void FRecastNavMeshGenerator::Init()
 			}
 			else
 			{
-				const FVector::FReal TileDim = Config.tileSize * Config.cs;
+				const FVector::FReal TileDim = Config.GetTileSizeUU();
 				if (SavedNavParams->tileHeight == TileDim && SavedNavParams->tileWidth == TileDim)
 				{
 					const FVector Orig = Recast2UnrealPoint(SavedNavParams->orig);
@@ -4764,8 +4820,8 @@ bool FRecastNavMeshGenerator::ConstructTiledNavMesh()
 		FVector NMOrigin = RcNavMeshOrigin;
 		rcVcopy(TiledMeshParameters.orig, &NMOrigin.X);
 
-		TiledMeshParameters.tileWidth = Config.tileSize * Config.cs;
-		TiledMeshParameters.tileHeight = Config.tileSize * Config.cs;
+		TiledMeshParameters.tileWidth = Config.GetTileSizeUU();
+		TiledMeshParameters.tileHeight = Config.GetTileSizeUU();
 
 		CalcNavMeshProperties(TiledMeshParameters.maxTiles, TiledMeshParameters.maxPolys);
 		Config.MaxPolysPerTile = TiledMeshParameters.maxPolys;
@@ -4773,7 +4829,9 @@ bool FRecastNavMeshGenerator::ConstructTiledNavMesh()
 		TiledMeshParameters.walkableClimb = Config.AgentMaxClimb;
 		TiledMeshParameters.walkableHeight = Config.AgentHeight;
 		TiledMeshParameters.walkableRadius = Config.AgentRadius;
-		TiledMeshParameters.bvQuantFactor = 1.f / Config.cs;
+		TiledMeshParameters.resolutionParams[(uint8)ENavigationDataResolution::Low].bvQuantFactor = 1.f / DestNavMesh->GetCellSize(ENavigationDataResolution::Low);
+		TiledMeshParameters.resolutionParams[(uint8)ENavigationDataResolution::Default].bvQuantFactor = 1.f / DestNavMesh->GetCellSize(ENavigationDataResolution::Default);
+		TiledMeshParameters.resolutionParams[(uint8)ENavigationDataResolution::High].bvQuantFactor = 1.f / DestNavMesh->GetCellSize(ENavigationDataResolution::High);
 
 		if (TiledMeshParameters.maxTiles == 0)
 		{
@@ -4837,7 +4895,7 @@ void FRecastNavMeshGenerator::CalcNavMeshProperties(int32& MaxTiles, int32& MaxP
 	int32 MaxRequestedTiles = 0;
 	if (DestNavMesh->IsResizable())
 	{
-		MaxRequestedTiles = UE::NavMesh::Private::CalculateMaxTilesCount(InclusionBounds, Config.tileSize * Config.cs, AvgLayersPerTile, DestNavMesh->NavMeshVersion);
+		MaxRequestedTiles = UE::NavMesh::Private::CalculateMaxTilesCount(InclusionBounds, Config.GetTileSizeUU(), AvgLayersPerTile, DestNavMesh->NavMeshVersion);
 	}
 	else
 	{
@@ -5007,7 +5065,7 @@ void FRecastNavMeshGenerator::OnNavigationBoundsChanged()
 	if (!IsGameStaticNavMesh(DestNavMesh) && DestNavMesh->IsResizable() && DetourMesh)
 	{
 		// Check whether Navmesh size needs to be changed
-		const int32 MaxRequestedTiles = UE::NavMesh::Private::CalculateMaxTilesCount(InclusionBounds, Config.tileSize * Config.cs, AvgLayersPerTile, DestNavMesh->NavMeshVersion);
+		const int32 MaxRequestedTiles = UE::NavMesh::Private::CalculateMaxTilesCount(InclusionBounds, Config.GetTileSizeUU(), AvgLayersPerTile, DestNavMesh->NavMeshVersion);
 		if (DetourMesh->getMaxTiles() != MaxRequestedTiles)
 		{
 			UE_LOG(LogNavigation, Log, TEXT("%s> Navigation bounds changed, rebuilding navmesh"), *DestNavMesh->GetName());
@@ -5654,7 +5712,7 @@ bool FRecastNavMeshGenerator::HasDirtyTiles(const FBox& AreaBounds) const
 	}
 
 	bool bRetDirty = false;
-	const FVector::FReal TileSizeInWorldUnits = Config.tileSize * Config.cs;
+	const FVector::FReal TileSizeInWorldUnits = Config.GetTileSizeUU();
 	const FRcTileBox TileBox(AreaBounds, RcNavMeshOrigin, TileSizeInWorldUnits);
 		
 	for (int32 Index = 0; bRetDirty == false && Index < PendingDirtyTiles.Num(); ++Index)
@@ -5671,7 +5729,7 @@ bool FRecastNavMeshGenerator::HasDirtyTiles(const FBox& AreaBounds) const
 
 int32 FRecastNavMeshGenerator::GetDirtyTilesCount(const FBox& AreaBounds) const
 {
-	const FVector::FReal TileSizeInWorldUnits = Config.tileSize * Config.cs;
+	const FVector::FReal TileSizeInWorldUnits = Config.GetTileSizeUU();
 	const FRcTileBox TileBox(AreaBounds, RcNavMeshOrigin, TileSizeInWorldUnits);
 
 	int32 DirtyPendingCount = 0;
@@ -5712,7 +5770,7 @@ void FRecastNavMeshGenerator::MarkDirtyTiles(const TArray<FNavigationDirtyArea>&
 	QUICK_SCOPE_CYCLE_COUNTER(STAT_RecastNavMeshGenerator_MarkDirtyTiles);
 	
 	check(bInitialized);
-	const FVector::FReal TileSizeInWorldUnits = Config.tileSize * Config.cs;
+	const FVector::FReal TileSizeInWorldUnits = Config.GetTileSizeUU();
 	check(TileSizeInWorldUnits > 0);
 
 	const bool bGameStaticNavMesh = IsGameStaticNavMesh(DestNavMesh);
@@ -5935,7 +5993,7 @@ void FRecastNavMeshGenerator::SortPendingBuildTiles()
 
 	if (SeedLocations.Num() > 0)
 	{
-		const FVector::FReal TileSizeInWorldUnits = Config.tileSize * Config.cs;
+		const FVector::FReal TileSizeInWorldUnits = Config.GetTileSizeUU();
 		
 		// Calculate shortest distances between tiles and players
 		for (FPendingTileElement& Element : PendingDirtyTiles)
diff --git a/Engine/Source/Runtime/NavigationSystem/Private/NavModifierVolume.cpp b/Engine/Source/Runtime/NavigationSystem/Private/NavModifierVolume.cpp
--- a/Engine/Source/Runtime/NavigationSystem/Private/NavModifierVolume.cpp
+++ b/Engine/Source/Runtime/NavigationSystem/Private/NavModifierVolume.cpp
@@ -19,6 +19,7 @@
 ANavModifierVolume::ANavModifierVolume(const FObjectInitializer& ObjectInitializer)
 	: Super(ObjectInitializer)
 	, AreaClass(UNavArea_Null::StaticClass())
+	, NavMeshResolution(ENavigationDataResolution::Invalid)
 {
 	if (GetBrushComponent())
 	{
@@ -34,15 +35,23 @@ void ANavModifierVolume::GetNavigationData(FNavigationRelevantData& Data) const
 		Data.Modifiers.CreateAreaModifiers(GetBrushComponent(), AreaClass);
 	}
 
-	if (bMaskFillCollisionUnderneathForNavmesh)
+	if (GetBrushComponent()->Brush != nullptr)
 	{
-		if (GetBrushComponent()->Brush != nullptr)
+		if (bMaskFillCollisionUnderneathForNavmesh)
 		{
 			const FBox& Box = GetBrushComponent()->Brush->Bounds.GetBox();
-			FAreaNavModifier AreaMod(Box, GetBrushComponent()->GetComponentTransform(), AreaClass);
+			const FAreaNavModifier AreaMod(Box, GetBrushComponent()->GetComponentTransform(), AreaClass);
 			Data.Modifiers.SetMaskFillCollisionUnderneathForNavmesh(true);
 			Data.Modifiers.Add(AreaMod);
 		}
+
+		if (NavMeshResolution != ENavigationDataResolution::Invalid)
+		{
+			const FBox& Box = GetBrushComponent()->Brush->Bounds.GetBox();
+			const FAreaNavModifier AreaMod(Box, GetBrushComponent()->GetComponentTransform(), AreaClass);
+			Data.Modifiers.SetNavMeshResolution(NavMeshResolution);
+			Data.Modifiers.Add(AreaMod);
+		}
 	}
 }
 
diff --git a/Engine/Source/Runtime/NavigationSystem/Public/NavMesh/RecastNavMesh.h b/Engine/Source/Runtime/NavigationSystem/Public/NavMesh/RecastNavMesh.h
--- a/Engine/Source/Runtime/NavigationSystem/Public/NavMesh/RecastNavMesh.h
+++ b/Engine/Source/Runtime/NavigationSystem/Public/NavMesh/RecastNavMesh.h
@@ -591,6 +591,18 @@ namespace FNavMeshConfig
 	};
 }
 
+USTRUCT()
+struct NAVIGATIONSYSTEM_API FNavMeshResolutionParam
+{
+	GENERATED_BODY()
+
+	bool IsValid() const { return CellSize > 0.f; }
+	
+	/** horizontal size of voxelization cell */
+	UPROPERTY(EditAnywhere, Category = Generation, config, meta = (ClampMin = "1.0", ClampMax = "1024.0"))
+	float CellSize = 25.f;
+};
+
 UCLASS(config=Engine, defaultconfig, hidecategories=(Input,Rendering,Tags,Transformation,Actor,Layers,Replication), notplaceable)
 class NAVIGATIONSYSTEM_API ARecastNavMesh : public ANavigationData
 {
@@ -701,13 +713,20 @@ class NAVIGATIONSYSTEM_API ARecastNavMesh : public ANavigationData
 	float TileSizeUU;
 
 	/** horizontal size of voxelization cell */
-	UPROPERTY(EditAnywhere, Category = Generation, config, meta = (ClampMin = "1.0", ClampMax = "1024.0"))
+	UE_DEPRECATED(5.2, "Set the CellSizes for the required navmesh resolutions in NavMeshResolutionParams.")
+	UPROPERTY(config)
 	float CellSize;
 
 	/** vertical size of voxelization cell */
 	UPROPERTY(EditAnywhere, Category = Generation, config, meta = (ClampMin = "1.0", ClampMax = "1024.0"))
 	float CellHeight;
 
+	/** Resolution params 
+	 * If using multiple resolutions, it's recommended to chose the highest resolution first and 
+	 * set it according to the highest desired precision and then the other resolutions. */
+	UPROPERTY(EditAnywhere, Category = Generation, config)
+	FNavMeshResolutionParam NavMeshResolutionParams[(uint8)ENavigationDataResolution::MAX];
+
 	/** Radius of smallest agent to traverse this navmesh */
 	UPROPERTY(EditAnywhere, Category = Generation, config, meta = (ClampMin = "0.0"))
 	float AgentRadius;
@@ -954,6 +973,7 @@ public:
 	virtual void BeginDestroy() override;
 
 #if WITH_EDITOR
+	virtual void PostEditChangeChainProperty( struct FPropertyChangedChainEvent& PropertyChangedChainEvent) override;
 	virtual void PostEditChangeProperty( struct FPropertyChangedEvent& PropertyChangedEvent) override;
 #endif // WITH_EDITOR
 	//~ End UObject Interface
@@ -1029,6 +1049,15 @@ protected:
 	virtual void OnRegistered() override;
 
 public:
+	/** Get the CellSize for the given resolution. */
+	float GetCellSize(const ENavigationDataResolution Resolution) const { return NavMeshResolutionParams[(uint8)Resolution].CellSize; }
+
+	/** Set the CellSize for the given resolution. */
+	void SetCellSize(const ENavigationDataResolution Resolution, const float Size) { NavMeshResolutionParams[(uint8)Resolution].CellSize = Size; }
+
+	/** Returns the tile size in world units. */
+	float GetTileSizeUU() const;
+	
 	/** Whether NavMesh should adjust its tile pool size when NavBounds are changed */
 	bool IsResizable() const;
 
diff --git a/Engine/Source/Runtime/NavigationSystem/Public/NavMesh/RecastNavMeshGenerator.h b/Engine/Source/Runtime/NavigationSystem/Public/NavMesh/RecastNavMeshGenerator.h
--- a/Engine/Source/Runtime/NavigationSystem/Public/NavMesh/RecastNavMeshGenerator.h
+++ b/Engine/Source/Runtime/NavigationSystem/Public/NavMesh/RecastNavMeshGenerator.h
@@ -81,6 +81,8 @@ struct FRecastBuildConfig : public rcConfig
 	/** Agent index for filtering links */
 	int32 AgentIndex;
 
+	ENavigationDataResolution TileResolution;
+
 	FRecastBuildConfig()
 	{
 		Reset();
@@ -101,7 +103,10 @@ struct FRecastBuildConfig : public rcConfig
 		MaxPolysPerTile = -1;
 		AgentIndex = 0;
 	}
+
+	rcReal GetTileSizeUU() const { return tileSize * cs; }
 };
+
 PRAGMA_ENABLE_DEPRECATION_WARNINGS
 
 struct FRecastVoxelCache
@@ -746,6 +751,9 @@ public:
 	
 	void SetNextTimeSliceRegenActive(bool bRegenState) { SyncTimeSlicedData.bNextTimeSliceRegenActive = bRegenState; }
 
+	/** Update the config according to the resolution */
+	virtual void SetupTileConfig(const ENavigationDataResolution TileResolution, FRecastBuildConfig& OutConfig) const;
+
 protected:
 	// Performs initial setup of member variables so that generator is ready to
 	// do its thing from this point on. Called just after construction by ARecastNavMesh
diff --git a/Engine/Source/Runtime/NavigationSystem/Public/NavMesh/RecastVersion.h b/Engine/Source/Runtime/NavigationSystem/Public/NavMesh/RecastVersion.h
--- a/Engine/Source/Runtime/NavigationSystem/Public/NavMesh/RecastVersion.h
+++ b/Engine/Source/Runtime/NavigationSystem/Public/NavMesh/RecastVersion.h
@@ -21,6 +21,7 @@
 #define NAVMESHVER_LWCOORDS_OPTIMIZATION		20
 #define NAVMESHVER_OPTIM_FIX_SERIALIZE_PARAMS	21 // Fix, serialize params that used to be in the tile and are now in the navmesh.
 #define NAVMESHVER_MAXTILES_COUNT_SKIP_INCLUSION 22
+#define NAVMESHVER_TILE_RESOLUTIONS				23 // Addition of a tile resolution index to the tile header.
 
-#define NAVMESHVER_LATEST				NAVMESHVER_MAXTILES_COUNT_SKIP_INCLUSION
+#define NAVMESHVER_LATEST				NAVMESHVER_TILE_RESOLUTIONS
 #define NAVMESHVER_MIN_COMPATIBLE		NAVMESHVER_LWCOORDS_OPTIMIZATION
diff --git a/Engine/Source/Runtime/NavigationSystem/Public/NavModifierVolume.h b/Engine/Source/Runtime/NavigationSystem/Public/NavModifierVolume.h
--- a/Engine/Source/Runtime/NavigationSystem/Public/NavModifierVolume.h
+++ b/Engine/Source/Runtime/NavigationSystem/Public/NavModifierVolume.h
@@ -6,6 +6,7 @@
 #include "UObject/ObjectMacros.h"
 #include "Templates/SubclassOf.h"
 #include "AI/Navigation/NavRelevantInterface.h"
+#include "AI/Navigation/NavigationTypes.h"
 #include "NavAreas/NavArea.h"
 #include "GameFramework/Volume.h"
 #include "NavModifierVolume.generated.h"
@@ -28,6 +29,11 @@ protected:
 	UPROPERTY(EditAnywhere, Category = Default, AdvancedDisplay)
 	bool bMaskFillCollisionUnderneathForNavmesh;
 
+	/** Experimental: When not set to None, the navmesh tiles touched by the navigation modifier volume will be built
+	 * using the highest resolution found. */
+	UPROPERTY(EditAnywhere, Category = Default, AdvancedDisplay)
+	ENavigationDataResolution NavMeshResolution;
+
 public:
 	ANavModifierVolume(const FObjectInitializer& ObjectInitializer);
 
diff --git a/Engine/Source/Runtime/Navmesh/Private/DebugUtils/DetourDebugDraw.cpp b/Engine/Source/Runtime/Navmesh/Private/DebugUtils/DetourDebugDraw.cpp
--- a/Engine/Source/Runtime/Navmesh/Private/DebugUtils/DetourDebugDraw.cpp
+++ b/Engine/Source/Runtime/Navmesh/Private/DebugUtils/DetourDebugDraw.cpp
@@ -388,13 +388,13 @@ static void drawMeshTileBVTree(duDebugDraw* dd, const dtMeshTile* tile, const dt
 void duDebugDrawNavMeshBVTree(duDebugDraw* dd, const dtNavMesh& mesh)
 {
 	if (!dd) return;
-
-	const dtReal bvQuantFactor = mesh.getBVQuantFactor();
 	
 	for (int i = 0; i < mesh.getMaxTiles(); ++i)
 	{
 		const dtMeshTile* tile = mesh.getTile(i);
 		if (!tile->header) continue;
+
+		const dtReal bvQuantFactor = mesh.getBVQuantFactor(tile->header->resolution);
 		drawMeshTileBVTree(dd, tile, bvQuantFactor);
 	}
 }
diff --git a/Engine/Source/Runtime/Navmesh/Private/Detour/DetourNavMesh.cpp b/Engine/Source/Runtime/Navmesh/Private/Detour/DetourNavMesh.cpp
--- a/Engine/Source/Runtime/Navmesh/Private/Detour/DetourNavMesh.cpp
+++ b/Engine/Source/Runtime/Navmesh/Private/Detour/DetourNavMesh.cpp
@@ -780,10 +780,6 @@ Notes:
 dtNavMesh::dtNavMesh() :
 	m_tileWidth(0),
 	m_tileHeight(0),
-	m_walkableHeight(0),
-	m_walkableRadius(0),
-	m_walkableClimb(0),
-	m_bvQuantFactor(0),
 	m_maxTiles(0),
 	m_tileLutSize(0),
 	m_tileLutMask(0),
@@ -825,14 +821,18 @@ dtNavMesh::~dtNavMesh()
 		
 dtStatus dtNavMesh::init(const dtNavMeshParams* params)
 {
-	m_walkableHeight = params->walkableHeight;
-	m_walkableRadius = params->walkableRadius; 
-	m_walkableClimb = params->walkableClimb;	
-	m_bvQuantFactor = params->bvQuantFactor;
-	check(m_bvQuantFactor != 0);
-	
 	memcpy(&m_params, params, sizeof(dtNavMeshParams));
 	dtVcopy(m_orig, params->orig);
+
+	// @UE BEGIN
+#if DO_CHECK	
+	for (uint8 i = 0; i < DT_RESOLUTION_COUNT; i++)
+	{
+		check(m_params.resolutionParams[i].bvQuantFactor != 0);
+	}
+#endif // DO_CHECK	
+	// @UE END
+	
 	m_tileWidth = params->tileWidth;
 	m_tileHeight = params->tileHeight;
 	
@@ -944,7 +944,7 @@ int dtNavMesh::findConnectingPolys(const dtReal* va, const dtReal* vb,
 			calcSlabEndPoints(vc, vd, bmin, bmax, side);
 
 			unsigned char overlapMode = 0;
-			if (!overlapSlabs(amin, amax, bmin, bmax, 0.01f, m_walkableClimb, &overlapMode)) continue;
+			if (!overlapSlabs(amin, amax, bmin, bmax, 0.01f, getWalkableClimb(), &overlapMode)) continue;
 
 			// if overlapping with only one side, verify height difference using detailed mesh
 			if (overlapMode == SLABOVERLAP_Max || overlapMode == SLABOVERLAP_Min)
@@ -957,7 +957,7 @@ int dtNavMesh::findConnectingPolys(const dtReal* va, const dtReal* vb,
 				const dtReal aH = getHeightFromDMesh(fromTile, fromPolyIdx, apt);
 				const dtReal bH = getHeightFromDMesh(tile, i, bpt);
 				const dtReal heightDiff = dtAbs(aH - bH);
-				if (heightDiff > m_walkableClimb)
+				if (heightDiff > getWalkableClimb())
 					continue;
 			}
 
@@ -1621,13 +1621,14 @@ int dtNavMesh::queryPolygonsInTile(const dtMeshTile* tile, const dtReal* qmin, c
 		dtReal maxz = dtClamp(qmax[2], tbmin[2], tbmax[2]) - tbmin[2];
 
 		// Quantize
-		UE_CLOG(m_bvQuantFactor == 0.f, LogDetour, Warning, TEXT("dtNavMesh::queryPolygonsInTile bounding volume quantization factor is zero! The query might not return the right result"));
-		bmin[0] = (unsigned short)(m_bvQuantFactor * minx) & 0xfffe;
-		bmin[1] = (unsigned short)(m_bvQuantFactor * miny) & 0xfffe;
-		bmin[2] = (unsigned short)(m_bvQuantFactor * minz) & 0xfffe;
-		bmax[0] = (unsigned short)(m_bvQuantFactor * maxx + 1) | 1;
-		bmax[1] = (unsigned short)(m_bvQuantFactor * maxy + 1) | 1;
-		bmax[2] = (unsigned short)(m_bvQuantFactor * maxz + 1) | 1;
+		const dtReal bvQuantFactor = m_params.resolutionParams[tile->header->resolution].bvQuantFactor; 	//@UE
+		UE_CLOG(bvQuantFactor == 0.f, LogDetour, Warning, TEXT("dtNavMesh::queryPolygonsInTile bounding volume quantization factor is zero! The query might not return the right result"));
+		bmin[0] = (unsigned short)(bvQuantFactor * minx) & 0xfffe;
+		bmin[1] = (unsigned short)(bvQuantFactor * miny) & 0xfffe;
+		bmin[2] = (unsigned short)(bvQuantFactor * minz) & 0xfffe;
+		bmax[0] = (unsigned short)(bvQuantFactor * maxx + 1) | 1;
+		bmax[1] = (unsigned short)(bvQuantFactor * maxy + 1) | 1;
+		bmax[2] = (unsigned short)(bvQuantFactor * maxz + 1) | 1;
 		
 		// Traverse tree
 		dtPolyRef base = getPolyRefBase(tile);
diff --git a/Engine/Source/Runtime/Navmesh/Private/Detour/DetourNavMeshBuilder.cpp b/Engine/Source/Runtime/Navmesh/Private/Detour/DetourNavMeshBuilder.cpp
--- a/Engine/Source/Runtime/Navmesh/Private/Detour/DetourNavMeshBuilder.cpp
+++ b/Engine/Source/Runtime/Navmesh/Private/Detour/DetourNavMeshBuilder.cpp
@@ -546,6 +546,7 @@ bool dtCreateNavMeshData(dtNavMeshCreateParams* params, unsigned char** outData,
 	bool bAllValuesSet = true;
 
 	header->version = DT_NAVMESH_VERSION;
+	header->resolution = params->tileResolutionLevel;	//@UE
 	header->x = params->tileX;
 	header->y = params->tileY;
 	setHeaderShort(params->tileLayer, header->layer, TEXT("tile layer"), bAllValuesSet);
diff --git a/Engine/Source/Runtime/Navmesh/Private/Detour/DetourNavMeshQuery.cpp b/Engine/Source/Runtime/Navmesh/Private/Detour/DetourNavMeshQuery.cpp
--- a/Engine/Source/Runtime/Navmesh/Private/Detour/DetourNavMeshQuery.cpp
+++ b/Engine/Source/Runtime/Navmesh/Private/Detour/DetourNavMeshQuery.cpp
@@ -1336,7 +1336,7 @@ int dtNavMeshQuery::queryPolygonsInTile(const dtMeshTile* tile, const dtReal* qm
 		const dtBVNode* end = &tile->bvTree[tile->header->bvNodeCount];
 		const dtReal* tbmin = tile->header->bmin;
 		const dtReal* tbmax = tile->header->bmax;
-		const dtReal qfac = m_nav->getBVQuantFactor();
+		const dtReal qfac = m_nav->getBVQuantFactor(tile->header->resolution);	//@UE
 
 		// Calculate quantized box
 		unsigned short bmin[3], bmax[3];
diff --git a/Engine/Source/Runtime/Navmesh/Public/Detour/DetourNavMesh.h b/Engine/Source/Runtime/Navmesh/Public/Detour/DetourNavMesh.h
--- a/Engine/Source/Runtime/Navmesh/Public/Detour/DetourNavMesh.h
+++ b/Engine/Source/Runtime/Navmesh/Public/Detour/DetourNavMesh.h
@@ -123,6 +123,8 @@ static const int DT_MAX_AREAS = 64;
 static const int DT_MIN_SALT_BITS = 5;
 static const int DT_SALT_BASE = 1;
 
+static const int DT_RESOLUTION_COUNT = 3; 
+
 #if WITH_NAVMESH_SEGMENT_LINKS
 /// Max segment parts for segment-to-segment off mesh connection
 static const int DT_MAX_OFFMESH_SEGMENT_PARTS = 4;
@@ -399,6 +401,8 @@ struct dtMeshHeader
 #if WITH_NAVMESH_CLUSTER_LINKS
 	unsigned short clusterCount;			///< Number of clusters
 #endif // WITH_NAVMESH_CLUSTER_LINKS
+
+	unsigned char resolution;			///< The resolution index used for the tile.
 	//@UE END
 
 	// These should be at the bottom, as they are less often used than the rest of the data. The rest will fit in one cache line.
@@ -458,18 +462,27 @@ struct dtMeshTile
 	//@UE END
 };
 
+//@UE BEGIN
+/// Configuration parameters depending on navmesh resolution.
+struct dtNavMeshResParams
+{
+	dtReal bvQuantFactor;			///< The bounding volume quantization factor.
+};
+//@UE END
+
 /// Configuration parameters used to define multi-tile navigation meshes.
 /// The values are used to allocate space during the initialization of a navigation mesh.
 /// @see dtNavMesh::init()
 /// @ingroup detour
 struct dtNavMeshParams
 {
-	//@UE BEGIN Memory optimization
+	//@UE BEGIN
 	dtReal walkableHeight;			///< The height of the agents using the tile.
 	dtReal walkableRadius;			///< The radius of the agents using the tile.
 	dtReal walkableClimb;			///< The maximum climb height of the agents using the tile.	
-	dtReal bvQuantFactor;			///< The bounding volume quantization factor.
-	//@UE END Memory optimization 	
+
+	dtNavMeshResParams resolutionParams[DT_RESOLUTION_COUNT]; ///< Parameters depending on resolutions.
+	//@UE END
 	
 	dtReal orig[3];					///< The world space origin of the navigation mesh's tile space. [(x, y, z)]
 	dtReal tileWidth;				///< The width of each tile. (Along the x-axis.)
@@ -854,10 +867,10 @@ public:
 	/// @}
 
 	//@UE Begin LWCoords
-	dtReal getWalkableHeight() const { return m_walkableHeight; }
-	dtReal getWalkableRadius() const { return m_walkableRadius; }
-	dtReal getWalkableClimb() const { return m_walkableClimb; }
-	dtReal getBVQuantFactor() const { return m_bvQuantFactor; }
+	dtReal getWalkableHeight() const { return m_params.walkableHeight; }
+	dtReal getWalkableRadius() const { return m_params.walkableRadius; }
+	dtReal getWalkableClimb() const { return m_params.walkableClimb; }
+	dtReal getBVQuantFactor(const unsigned char resolution) const { return m_params.resolutionParams[resolution].bvQuantFactor; }
 	//@UE END LWCoords
 
 private:
@@ -934,17 +947,10 @@ private:
 	void calcTileLoc(const dtReal* pos, dtReal* tx, dtReal* ty) const;
 
 public:
-	dtNavMeshParams m_params;			///< Current initialization params. TODO: do not store this info twice.
+	dtNavMeshParams m_params;			///< Current initialization params.
 	dtReal m_orig[3];					///< Origin of the tile (0,0)
 	dtReal m_tileWidth, m_tileHeight;	///< Dimensions of each tile.
 
-	//@UE BEGIN Memory optimization
-	dtReal m_walkableHeight;			///< The height of the agents using the tile.
-	dtReal m_walkableRadius;			///< The radius of the agents using the tile.
-	dtReal m_walkableClimb;				///< The maximum climb height of the agents using the tile.
-	dtReal m_bvQuantFactor;				///< The bounding volume quantization factor. 
-	//@UE END Memory optimization
-
 	int m_maxTiles;						///< Max number of tiles.
 	int m_tileLutSize;					///< Tile hash lookup size (must be pot).
 	int m_tileLutMask;					///< Tile hash lookup mask.
diff --git a/Engine/Source/Runtime/Navmesh/Public/Detour/DetourNavMeshBuilder.h b/Engine/Source/Runtime/Navmesh/Public/Detour/DetourNavMeshBuilder.h
--- a/Engine/Source/Runtime/Navmesh/Public/Detour/DetourNavMeshBuilder.h
+++ b/Engine/Source/Runtime/Navmesh/Public/Detour/DetourNavMeshBuilder.h
@@ -143,6 +143,8 @@ struct dtNavMeshCreateParams
 	dtReal cs;				///< The xz-plane cell size of the polygon mesh. [Limit: > 0] [Unit: wu]
 	dtReal ch;				///< The y-axis cell height of the polygon mesh. [Limit: > 0] [Unit: wu]
 
+	unsigned char tileResolutionLevel;	///< Tile resolution index 		//@UE
+
 	/// True if a bounding volume tree should be built for the tile.
 	/// @note The BVTree is not normally needed for layered navigation meshes.
 	bool buildBvTree;
