Fix missing default tile resolution initialisation

diff --git a/Engine/Source/Runtime/NavigationSystem/Public/NavMesh/RecastNavMeshGenerator.h b/Engine/Source/Runtime/NavigationSystem/Public/NavMesh/RecastNavMeshGenerator.h
--- a/Engine/Source/Runtime/NavigationSystem/Public/NavMesh/RecastNavMeshGenerator.h
+++ b/Engine/Source/Runtime/NavigationSystem/Public/NavMesh/RecastNavMeshGenerator.h
@@ -45,7 +45,7 @@ class UNavigationSystemV1;
 PRAGMA_DISABLE_DEPRECATION_WARNINGS
 struct FRecastBuildConfig : public rcConfig
 {
-	/** controls whether voxel filterring will be applied (via FRecastTileGenerator::ApplyVoxelFilter) */
+	/** controls whether voxel filtering will be applied (via FRecastTileGenerator::ApplyVoxelFilter) */
 	uint32 bPerformVoxelFiltering:1;
 	/** generate detailed mesh (additional tessellation to match heights of geometry) */
 	uint32 bGenerateDetailedMesh:1;
@@ -80,7 +80,7 @@ struct FRecastBuildConfig : public rcConfig
 	float AgentRadius;
 	/** Agent index for filtering links */
 	int32 AgentIndex;
-
+	/** Resolution level */ 
 	ENavigationDataResolution TileResolution;
 
 	FRecastBuildConfig()
@@ -102,6 +102,7 @@ struct FRecastBuildConfig : public rcConfig
 		PolyMaxHeight = 10;
 		MaxPolysPerTile = -1;
 		AgentIndex = 0;
+		TileResolution = ENavigationDataResolution::Default;
 	}
 
 	rcReal GetTileSizeUU() const { return tileSize * cs; }
