Add option to allow pathfind query with an end location that has no valid polygon around

diff --git a/Engine/Source/Runtime/NavigationSystem/Private/NavMesh/PImplRecastNavMesh.cpp b/Engine/Source/Runtime/NavigationSystem/Private/NavMesh/PImplRecastNavMesh.cpp
--- a/Engine/Source/Runtime/NavigationSystem/Private/NavMesh/PImplRecastNavMesh.cpp
+++ b/Engine/Source/Runtime/NavigationSystem/Private/NavMesh/PImplRecastNavMesh.cpp
@@ -1114,13 +1114,23 @@ void FPImplRecastNavMesh::Raycast(const FVector& StartLoc, const FVector& EndLoc
 	}
 }
 
+// DEPRECATED
 ENavigationQueryResult::Type FPImplRecastNavMesh::FindPath(const FVector& StartLoc, const FVector& EndLoc, FNavMeshPath& Path, const FNavigationQueryFilter& Filter, const UObject* Owner) const
 {
-	return FindPath(StartLoc, EndLoc, TNumericLimits<FVector::FReal>::Max(), Path, Filter, Owner);
+	constexpr FVector::FReal CostLimit = TNumericLimits<FVector::FReal>::Max();
+	constexpr bool bRequireNavigableEndLocation = true; 
+	return FindPath(StartLoc, EndLoc, CostLimit, bRequireNavigableEndLocation, Path, Filter, Owner);
 }
 
-// @TODONAV
+// DEPRECATED
 ENavigationQueryResult::Type FPImplRecastNavMesh::FindPath(const FVector& StartLoc, const FVector& EndLoc, const FVector::FReal CostLimit, FNavMeshPath& Path, const FNavigationQueryFilter& InQueryFilter, const UObject* Owner) const
+{
+	constexpr bool bRequireNavigableEndLocation = true; 
+	return FindPath(StartLoc, EndLoc, CostLimit, bRequireNavigableEndLocation, Path, InQueryFilter, Owner);
+}
+
+// @TODONAV
+ENavigationQueryResult::Type FPImplRecastNavMesh::FindPath(const FVector& StartLoc, const FVector& EndLoc, const FVector::FReal CostLimit, const bool bRequireNavigableEndLocation, FNavMeshPath& Path, const FNavigationQueryFilter& InQueryFilter, const UObject* Owner) const
 {
 	// temporarily disabling this check due to it causing too much "crashes"
 	// @todo but it needs to be back at some point since it realy checks for a buggy setup
@@ -1147,6 +1157,7 @@ ENavigationQueryResult::Type FPImplRecastNavMesh::FindPath(const FVector& StartL
 
 	FRecastSpeciaLinkFilter LinkFilter(FNavigationSystem::GetCurrent<UNavigationSystemV1>(NavMeshOwner->GetWorld()), Owner);
 	INITIALIZE_NAVQUERY(NavQuery, InQueryFilter.GetMaxSearchNodes(), LinkFilter);
+	NavQuery.setRequireNavigableEndLocation(bRequireNavigableEndLocation);
 
 	FVector RecastStartPos, RecastEndPos;
 	NavNodeRef StartPolyID, EndPolyID;
@@ -1217,7 +1228,14 @@ ENavigationQueryResult::Type FPImplRecastNavMesh::PostProcessPathInternal(dtStat
 	return DTStatusToNavQueryResult(FindPathStatus);
 }
 
+// DEPRECATED
 ENavigationQueryResult::Type FPImplRecastNavMesh::TestPath(const FVector& StartLoc, const FVector& EndLoc, const FNavigationQueryFilter& InQueryFilter, const UObject* Owner, int32* NumVisitedNodes) const
+{
+	constexpr bool bRequireNavigableEndLocation = true;
+	return TestPath(StartLoc, EndLoc, bRequireNavigableEndLocation, InQueryFilter, Owner, NumVisitedNodes);
+}
+
+ENavigationQueryResult::Type FPImplRecastNavMesh::TestPath(const FVector& StartLoc, const FVector& EndLoc, const bool bRequireNavigableEndLocation, const FNavigationQueryFilter& InQueryFilter, const UObject* Owner, int32* NumVisitedNodes) const
 {
 	const dtQueryFilter* QueryFilter = ((const FRecastQueryFilter*)(InQueryFilter.GetImplementation()))->GetAsDetourQueryFilter();
 	if (QueryFilter == NULL)
@@ -1228,6 +1246,7 @@ ENavigationQueryResult::Type FPImplRecastNavMesh::TestPath(const FVector& StartL
 
 	FRecastSpeciaLinkFilter LinkFilter(FNavigationSystem::GetCurrent<UNavigationSystemV1>(NavMeshOwner->GetWorld()), Owner);
 	INITIALIZE_NAVQUERY(NavQuery, InQueryFilter.GetMaxSearchNodes(), LinkFilter);
+	NavQuery.setRequireNavigableEndLocation(bRequireNavigableEndLocation);
 
 	FVector RecastStartPos, RecastEndPos;
 	NavNodeRef StartPolyID, EndPolyID;
@@ -1260,6 +1279,7 @@ ENavigationQueryResult::Type FPImplRecastNavMesh::TestClusterPath(const FVector&
 
 	check(NavMeshOwner->DefaultMaxHierarchicalSearchNodes >= 0. && NavMeshOwner->DefaultMaxHierarchicalSearchNodes <= (float)TNumericLimits<int32>::Max());
 	INITIALIZE_NAVQUERY_SIMPLE(ClusterQuery, static_cast<int32>(NavMeshOwner->DefaultMaxHierarchicalSearchNodes));
+	ClusterQuery.setRequireNavigableEndLocation(true);
 
 	const bool bCanSearch = InitPathfinding(StartLoc, EndLoc, ClusterQuery, ClusterFilter, RecastStartPos, StartPolyID, RecastEndPos, EndPolyID);
 	if (!bCanSearch)
@@ -1292,7 +1312,7 @@ bool FPImplRecastNavMesh::InitPathfinding(const FVector& UnrealStart, const FVec
 	Query.findNearestPoly(&RecastStartToProject.X, Extent, Filter, &StartPoly, &RecastStart.X);
 	if (StartPoly == INVALID_NAVNODEREF)
 	{
-		UE_VLOG(NavMeshOwner, LogNavigation, Warning, TEXT("FPImplRecastNavMesh::InitPathfinding start point not on navmesh"));
+		UE_VLOG(NavMeshOwner, LogNavigation, Warning, TEXT("FPImplRecastNavMesh::InitPathfinding start point not on navmesh (%s)"), *UnrealStart.ToString());
 		UE_VLOG_SEGMENT(NavMeshOwner, LogNavigation, Warning, UnrealStart, UnrealEnd, FColor::Red, TEXT("Failed path"));
 		UE_VLOG_LOCATION(NavMeshOwner, LogNavigation, Warning, UnrealStart, 15, FColor::Red, TEXT("Start failed"));
 		UE_VLOG_BOX(NavMeshOwner, LogNavigation, Warning, FBox(UnrealStart - NavExtent, UnrealStart + NavExtent), FColor::Red, TEXT_EMPTY);
@@ -1304,12 +1324,18 @@ bool FPImplRecastNavMesh::InitPathfinding(const FVector& UnrealStart, const FVec
 	Query.findNearestPoly(&RecastEndToProject.X, Extent, Filter, &EndPoly, &RecastEnd.X);
 	if (EndPoly == INVALID_NAVNODEREF)
 	{
-		UE_VLOG(NavMeshOwner, LogNavigation, Warning, TEXT("FPImplRecastNavMesh::InitPathfinding end point not on navmesh"));
-		UE_VLOG_SEGMENT(NavMeshOwner, LogNavigation, Warning, UnrealEnd, UnrealEnd, FColor::Red, TEXT("Failed path"));
-		UE_VLOG_LOCATION(NavMeshOwner, LogNavigation, Warning, UnrealEnd, 15, FColor::Red, TEXT("End failed"));
-		UE_VLOG_BOX(NavMeshOwner, LogNavigation, Warning, FBox(UnrealEnd - NavExtent, UnrealEnd + NavExtent), FColor::Red, TEXT_EMPTY);
+		if (Query.isRequiringNavigableEndLocation())
+		{
+			UE_VLOG(NavMeshOwner, LogNavigation, Warning, TEXT("FPImplRecastNavMesh::InitPathfinding end point not on navmesh (%s)"), *UnrealEnd.ToString());
+			UE_VLOG_SEGMENT(NavMeshOwner, LogNavigation, Warning, UnrealEnd, UnrealEnd, FColor::Red, TEXT("Failed path"));
+			UE_VLOG_LOCATION(NavMeshOwner, LogNavigation, Warning, UnrealEnd, 15, FColor::Red, TEXT("End failed"));
+			UE_VLOG_BOX(NavMeshOwner, LogNavigation, Warning, FBox(UnrealEnd - NavExtent, UnrealEnd + NavExtent), FColor::Red, TEXT_EMPTY);
 
-		return false;
+			return false;
+		}
+
+		// we will use RecastEndToProject as the estimated end location since we didn't find a poly. It will be used to compute the heuristic mainly
+		dtVcopy(&RecastEnd.X, &RecastEndToProject.X);
 	}
 
 	return true;
@@ -1626,7 +1652,14 @@ static void StorePathfindingDebugStep(const dtNavMeshQuery& NavQuery, const dtNa
 	}
 }
 
+// DEPRECATED
 int32 FPImplRecastNavMesh::DebugPathfinding(const FVector& StartLoc, const FVector& EndLoc, const FVector::FReal CostLimit, const FNavigationQueryFilter& Filter, const UObject* Owner, TArray<FRecastDebugPathfindingData>& Steps)
+{
+	constexpr bool bRequireNavigableEndLocation = true;
+	return DebugPathfinding(StartLoc, EndLoc, CostLimit, bRequireNavigableEndLocation, Filter, Owner, Steps);
+}
+
+int32 FPImplRecastNavMesh::DebugPathfinding(const FVector& StartLoc, const FVector& EndLoc, const FVector::FReal CostLimit, const bool bRequireNavigableEndLocation, const FNavigationQueryFilter& Filter, const UObject* Owner, TArray<FRecastDebugPathfindingData>& Steps)
 {
 	int32 NumSteps = 0;
 
@@ -1639,6 +1672,7 @@ int32 FPImplRecastNavMesh::DebugPathfinding(const FVector& StartLoc, const FVect
 
 	FRecastSpeciaLinkFilter LinkFilter(FNavigationSystem::GetCurrent<UNavigationSystemV1>(NavMeshOwner->GetWorld()), Owner);
 	INITIALIZE_NAVQUERY(NavQuery, Filter.GetMaxSearchNodes(), LinkFilter);
+	NavQuery.setRequireNavigableEndLocation(bRequireNavigableEndLocation);
 
 	FVector RecastStartPos, RecastEndPos;
 	NavNodeRef StartPolyID, EndPolyID;
@@ -1648,7 +1682,7 @@ int32 FPImplRecastNavMesh::DebugPathfinding(const FVector& StartLoc, const FVect
 		return NumSteps;
 	}
 
-	dtStatus status = NavQuery.initSlicedFindPath(StartPolyID, EndPolyID, &RecastStartPos.X, &RecastEndPos.X, CostLimit, QueryFilter);
+	dtStatus status = NavQuery.initSlicedFindPath(StartPolyID, EndPolyID, &RecastStartPos.X, &RecastEndPos.X, CostLimit, bRequireNavigableEndLocation, QueryFilter);
 	while (dtStatusInProgress(status))
 	{
 		StorePathfindingDebugStep(NavQuery, DetourNavMesh, Steps);
diff --git a/Engine/Source/Runtime/NavigationSystem/Private/NavMesh/RecastNavMesh.cpp b/Engine/Source/Runtime/NavigationSystem/Private/NavMesh/RecastNavMesh.cpp
--- a/Engine/Source/Runtime/NavigationSystem/Private/NavMesh/RecastNavMesh.cpp
+++ b/Engine/Source/Runtime/NavigationSystem/Private/NavMesh/RecastNavMesh.cpp
@@ -1891,9 +1891,11 @@ ENavigationQueryResult::Type ARecastNavMesh::CalcPathLengthAndCost(const FVector
 			TSharedRef<FNavMeshPath> Path = MakeShareable(new FNavMeshPath());
 			Path->SetWantsStringPulling(false);
 			Path->SetWantsPathCorridor(true);
-
-			const FVector::FReal CostLimit = TNumericLimits<FVector::FReal>::Max();
-			Result = RecastNavMeshImpl->FindPath(PathStart, PathEnd, CostLimit, Path.Get(), GetRightFilterRef(QueryFilter), QueryOwner);
+			
+			// LWC_TODO_AI: CostLimit should be FVector::FReal. Not until after 5.0!
+			constexpr FVector::FReal CostLimit = TNumericLimits<FVector::FReal>::Max();
+			constexpr bool bRequireNavigableEndLocation = true;
+			Result = RecastNavMeshImpl->FindPath(PathStart, PathEnd, CostLimit, bRequireNavigableEndLocation, Path.Get(), GetRightFilterRef(QueryFilter), QueryOwner);
 
 			if (Result == ENavigationQueryResult::Success || (Result == ENavigationQueryResult::Fail && Path->IsPartial()))
 			{
@@ -2813,7 +2815,7 @@ FPathFindingResult ARecastNavMesh::FindPath(const FNavAgentProperties& AgentProp
 		}
 		else
 		{
-			Result.Result = RecastNavMesh->RecastNavMeshImpl->FindPath(Query.StartLocation, AdjustedEndLocation, Query.CostLimit, *NavMeshPath, *NavFilter, Query.Owner.Get());
+			Result.Result = RecastNavMesh->RecastNavMeshImpl->FindPath(Query.StartLocation, AdjustedEndLocation, Query.CostLimit, Query.bRequireNavigableEndLocation, *NavMeshPath, *NavFilter, Query.Owner.Get());
 
 			const bool bPartialPath = Result.IsPartial();
 			if (bPartialPath)
@@ -2848,7 +2850,7 @@ bool ARecastNavMesh::TestPath(const FNavAgentProperties& AgentProperties, const
 		const FVector AdjustedEndLocation = NavFilter->GetAdjustedEndLocation(Query.EndLocation);
 		if ((Query.StartLocation - AdjustedEndLocation).IsNearlyZero() == false)
 		{
-			ENavigationQueryResult::Type Result = RecastNavMesh->RecastNavMeshImpl->TestPath(Query.StartLocation, AdjustedEndLocation, *NavFilter, Query.Owner.Get(), NumVisitedNodes);
+			ENavigationQueryResult::Type Result = RecastNavMesh->RecastNavMeshImpl->TestPath(Query.StartLocation, AdjustedEndLocation, Query.bRequireNavigableEndLocation, *NavFilter, Query.Owner.Get(), NumVisitedNodes);
 			bPathExists = (Result == ENavigationQueryResult::Success);
 		}
 	}
@@ -2900,7 +2902,7 @@ bool ARecastNavMesh::TestHierarchicalPath(const FNavAgentProperties& AgentProper
 
 			if (bUseFallbackSearch)
 			{
-				ENavigationQueryResult::Type Result = RecastNavMesh->RecastNavMeshImpl->TestPath(Query.StartLocation, AdjustedEndLocation, *NavFilter, Query.Owner.Get(), NumVisitedNodes);
+				ENavigationQueryResult::Type Result = RecastNavMesh->RecastNavMeshImpl->TestPath(Query.StartLocation, AdjustedEndLocation, Query.bRequireNavigableEndLocation, *NavFilter, Query.Owner.Get(), NumVisitedNodes);
 				bPathExists = (Result == ENavigationQueryResult::Success);
 			}
 		}
@@ -3032,7 +3034,7 @@ int32 ARecastNavMesh::DebugPathfinding(const FPathFindingQuery& Query, TArray<FR
 
 	if ((Query.StartLocation - Query.EndLocation).IsNearlyZero() == false)
 	{
-		NumSteps = RecastNavMesh->RecastNavMeshImpl->DebugPathfinding(Query.StartLocation, Query.EndLocation, Query.CostLimit, *(Query.QueryFilter.Get()), Query.Owner.Get(), Steps);
+		NumSteps = RecastNavMesh->RecastNavMeshImpl->DebugPathfinding(Query.StartLocation, Query.EndLocation, Query.CostLimit, Query.bRequireNavigableEndLocation, *(Query.QueryFilter.Get()), Query.Owner.Get(), Steps);
 	}
 
 	return NumSteps;
diff --git a/Engine/Source/Runtime/NavigationSystem/Private/NavigationData.cpp b/Engine/Source/Runtime/NavigationSystem/Private/NavigationData.cpp
--- a/Engine/Source/Runtime/NavigationSystem/Private/NavigationData.cpp
+++ b/Engine/Source/Runtime/NavigationSystem/Private/NavigationData.cpp
@@ -27,8 +27,8 @@ static TAutoConsoleVariable<int32> CVarDestroyNavDataInCleanUpAndMarkPendingKill
 //----------------------------------------------------------------------//
 // FPathFindingQuery
 //----------------------------------------------------------------------//
-FPathFindingQuery::FPathFindingQuery(const UObject* InOwner, const ANavigationData& InNavData, const FVector& Start, const FVector& End, FSharedConstNavQueryFilter SourceQueryFilter, FNavPathSharedPtr InPathInstanceToFill, const FVector::FReal CostLimit) :
-	FPathFindingQueryData(InOwner, Start, End, SourceQueryFilter, 0 /*InNavDataFlags*/, true /*bInAllowPartialPaths*/, CostLimit),
+FPathFindingQuery::FPathFindingQuery(const UObject* InOwner, const ANavigationData& InNavData, const FVector& Start, const FVector& End, FSharedConstNavQueryFilter SourceQueryFilter, FNavPathSharedPtr InPathInstanceToFill, const FVector::FReal CostLimit, const bool bInRequireNavigableEndLocation) :
+	FPathFindingQueryData(InOwner, Start, End, SourceQueryFilter, 0 /*InNavDataFlags*/, true /*bInAllowPartialPaths*/, CostLimit, bInRequireNavigableEndLocation),
 	NavData(&InNavData), PathInstanceToFill(InPathInstanceToFill), NavAgentProperties(InNavData.GetConfig())
 {
 	if (!QueryFilter.IsValid() && NavData.IsValid())
@@ -48,7 +48,7 @@ FPathFindingQuery::FPathFindingQuery(const INavAgentInterface& InNavAgent, const
 }
 
 FPathFindingQuery::FPathFindingQuery(const FPathFindingQuery& Source) :
-	FPathFindingQueryData(Source.Owner.Get(), Source.StartLocation, Source.EndLocation, Source.QueryFilter, Source.NavDataFlags, Source.bAllowPartialPaths, Source.CostLimit),
+	FPathFindingQueryData(Source.Owner.Get(), Source.StartLocation, Source.EndLocation, Source.QueryFilter, Source.NavDataFlags, Source.bAllowPartialPaths, Source.CostLimit, Source.bRequireNavigableEndLocation),
 	NavData(Source.NavData), PathInstanceToFill(Source.PathInstanceToFill), NavAgentProperties(Source.NavAgentProperties)
 {
 	if (!QueryFilter.IsValid() && NavData.IsValid())
diff --git a/Engine/Source/Runtime/NavigationSystem/Private/NavigationTestingActor.cpp b/Engine/Source/Runtime/NavigationSystem/Private/NavigationTestingActor.cpp
--- a/Engine/Source/Runtime/NavigationSystem/Private/NavigationTestingActor.cpp
+++ b/Engine/Source/Runtime/NavigationSystem/Private/NavigationTestingActor.cpp
@@ -57,6 +57,7 @@ ANavigationTestingActor::ANavigationTestingActor(const FObjectInitializer& Objec
 	OffsetFromCornersDistance = 0.f;
 
 	QueryingExtent = FVector(DEFAULT_NAV_QUERY_EXTENT_HORIZONTAL, DEFAULT_NAV_QUERY_EXTENT_HORIZONTAL, DEFAULT_NAV_QUERY_EXTENT_VERTICAL);
+	bRequireNavigableEndLocation = true;
 
 	CapsuleComponent = CreateDefaultSubobject<UCapsuleComponent>(TEXT("CollisionCylinder"));
 	CapsuleComponent->InitCapsuleSize(NavAgentProps.AgentRadius, NavAgentProps.AgentHeight / 2);
@@ -495,7 +496,9 @@ FPathFindingQuery ANavigationTestingActor::BuildPathFindingQuery(const ANavigati
 	check(Goal);
 	if (MyNavData)
 	{
-		return FPathFindingQuery(this, *MyNavData, GetNavAgentLocation(), Goal->GetNavAgentLocation(), UNavigationQueryFilter::GetQueryFilter(*MyNavData, this, FilterClass));
+		constexpr float DefaultCostLimit = FLT_MAX;
+		const FNavPathSharedPtr NoSharedPath = nullptr;
+		return FPathFindingQuery(this, *MyNavData, GetNavAgentLocation(), Goal->GetNavAgentLocation(), UNavigationQueryFilter::GetQueryFilter(*MyNavData, this, FilterClass), NoSharedPath, DefaultCostLimit, bRequireNavigableEndLocation);
 	}
 	
 	return FPathFindingQuery();
diff --git a/Engine/Source/Runtime/NavigationSystem/Public/NavMesh/PImplRecastNavMesh.h b/Engine/Source/Runtime/NavigationSystem/Public/NavMesh/PImplRecastNavMesh.h
--- a/Engine/Source/Runtime/NavigationSystem/Public/NavMesh/PImplRecastNavMesh.h
+++ b/Engine/Source/Runtime/NavigationSystem/Public/NavMesh/PImplRecastNavMesh.h
@@ -82,15 +82,23 @@ public:
 
 	// @TODONAV
 	/** Generates path from the given query. Synchronous. */
-	UE_DEPRECATED(4.25, "Use the version with the added CostLimit parameter (TNumericLimits<FVector::FReal>::Max() can be used as default).")
+	UE_DEPRECATED(4.25, "Please use FindPath with the added CostLimit and bRequireNavigableEndLocation parameters (TNumericLimits<FVector::FReal>::Max() and true can be used as default respectively).")
 	ENavigationQueryResult::Type FindPath(const FVector& StartLoc, const FVector& EndLoc, FNavMeshPath& Path, const FNavigationQueryFilter& Filter, const UObject* Owner) const;
-	
+
 	/** Generates path from the given query. Synchronous. */
+	UE_DEPRECATED(5.2, "Please use FindPath with the added bRequireNavigableEndLocation parameter (true can be used as default).")
 	ENavigationQueryResult::Type FindPath(const FVector& StartLoc, const FVector& EndLoc, const FVector::FReal CostLimit, FNavMeshPath& Path, const FNavigationQueryFilter& Filter, const UObject* Owner) const;
 
+	/** Generates path from the given query. Synchronous. */
+	ENavigationQueryResult::Type FindPath(const FVector& StartLoc, const FVector& EndLoc, const FVector::FReal CostLimit, const bool bRequireNavigableEndLocation, FNavMeshPath& Path, const FNavigationQueryFilter& Filter, const UObject* Owner) const;
+
 	/** Check if path exists */
+	UE_DEPRECATED(5.2, "Please use TestPath with the added bRequireNavigableEndLocation parameter (true can be used as default).")
 	ENavigationQueryResult::Type TestPath(const FVector& StartLoc, const FVector& EndLoc, const FNavigationQueryFilter& Filter, const UObject* Owner, int32* NumVisitedNodes = 0) const;
 
+	/** Check if path exists */
+	ENavigationQueryResult::Type TestPath(const FVector& StartLoc, const FVector& EndLoc, const bool bRequireNavigableEndLocation, const FNavigationQueryFilter& Filter, const UObject* Owner, int32* NumVisitedNodes = 0) const;
+
 	template< typename TRecastAStar, typename TRecastAStartGraph, typename TRecastGraphAStarFilter, typename TRecastAStarResult >
 	ENavigationQueryResult::Type FindPathCustomAStar(TRecastAStartGraph& RecastGraphWrapper, TRecastAStar& AStarAlgo, const FVector& StartLoc, const FVector& EndLoc, const FVector::FReal CostLimit, FNavMeshPath& Path, const FNavigationQueryFilter& Filter, const UObject* Owner) const;
 
@@ -99,8 +107,12 @@ public:
 		ARecastNavMesh::FRaycastResult& RaycastResult, NavNodeRef StartNode = INVALID_NAVNODEREF) const;
 
 	/** Generates path from given query and collect data for every step of A* algorithm */
+	UE_DEPRECATED(5.2, "Please use DebugPathfinding with the added bRequireNavigableEndLocation parameter (true can be used as default).")
 	int32 DebugPathfinding(const FVector& StartLoc, const FVector& EndLoc, const FVector::FReal CostLimit, const FNavigationQueryFilter& Filter, const UObject* Owner, TArray<FRecastDebugPathfindingData>& Steps);
 
+	/** Generates path from given query and collect data for every step of A* algorithm */
+	int32 DebugPathfinding(const FVector& StartLoc, const FVector& EndLoc, const FVector::FReal CostLimit, const bool bRequireNavigableEndLocation, const FNavigationQueryFilter& Filter, const UObject* Owner, TArray<FRecastDebugPathfindingData>& Steps);
+
 	/** Returns a random location on the navmesh. */
 	FNavLocation GetRandomPoint(const FNavigationQueryFilter& Filter, const UObject* Owner) const;
 
diff --git a/Engine/Source/Runtime/NavigationSystem/Public/NavigationSystemTypes.h b/Engine/Source/Runtime/NavigationSystem/Public/NavigationSystemTypes.h
--- a/Engine/Source/Runtime/NavigationSystem/Public/NavigationSystemTypes.h
+++ b/Engine/Source/Runtime/NavigationSystem/Public/NavigationSystemTypes.h
@@ -36,10 +36,13 @@ struct NAVIGATIONSYSTEM_API FPathFindingQueryData
 	/** if set, allow partial paths as a result */
 	uint32 bAllowPartialPaths : 1;
 
-	FPathFindingQueryData() : StartLocation(FNavigationSystem::InvalidLocation), EndLocation(FNavigationSystem::InvalidLocation), CostLimit(FLT_MAX), NavDataFlags(0), bAllowPartialPaths(true) {}
+	/** if set, require the end location to be linked to the navigation data */
+	uint32 bRequireNavigableEndLocation : 1;
 
-	FPathFindingQueryData(const UObject* InOwner, const FVector& InStartLocation, const FVector& InEndLocation, FSharedConstNavQueryFilter InQueryFilter = nullptr, int32 InNavDataFlags = 0, bool bInAllowPartialPaths = true, const FVector::FReal InCostLimit = TNumericLimits<FVector::FReal>::Max()) :
-		Owner(InOwner), StartLocation(InStartLocation), EndLocation(InEndLocation), QueryFilter(InQueryFilter), CostLimit(InCostLimit), NavDataFlags(InNavDataFlags), bAllowPartialPaths(bInAllowPartialPaths) {}
+	FPathFindingQueryData() : StartLocation(FNavigationSystem::InvalidLocation), EndLocation(FNavigationSystem::InvalidLocation), CostLimit(TNumericLimits<FVector::FReal>::Max()), NavDataFlags(0), bAllowPartialPaths(true), bRequireNavigableEndLocation(true) {}
+
+	FPathFindingQueryData(const UObject* InOwner, const FVector& InStartLocation, const FVector& InEndLocation, FSharedConstNavQueryFilter InQueryFilter = nullptr, int32 InNavDataFlags = 0, bool bInAllowPartialPaths = true, const  FVector::FReal InCostLimit = TNumericLimits<FVector::FReal>::Max(), const bool bInRequireNavigableEndLocation = true) :
+		Owner(InOwner), StartLocation(InStartLocation), EndLocation(InEndLocation), QueryFilter(InQueryFilter), CostLimit(InCostLimit), NavDataFlags(InNavDataFlags), bAllowPartialPaths(bInAllowPartialPaths), bRequireNavigableEndLocation(bInRequireNavigableEndLocation) {}
 };
 
 struct NAVIGATIONSYSTEM_API FPathFindingQuery : public FPathFindingQueryData
@@ -50,13 +53,14 @@ struct NAVIGATIONSYSTEM_API FPathFindingQuery : public FPathFindingQueryData
 
 	FPathFindingQuery() : FPathFindingQueryData() {}
 	FPathFindingQuery(const FPathFindingQuery& Source);
-	FPathFindingQuery(const UObject* InOwner, const ANavigationData& InNavData, const FVector& Start, const FVector& End, FSharedConstNavQueryFilter SourceQueryFilter = NULL, FNavPathSharedPtr InPathInstanceToFill = NULL, const FVector::FReal CostLimit = TNumericLimits<FVector::FReal>::Max());
-	FPathFindingQuery(const INavAgentInterface& InNavAgent, const ANavigationData& InNavData, const FVector& Start, const FVector& End, FSharedConstNavQueryFilter SourceQueryFilter = NULL, FNavPathSharedPtr InPathInstanceToFill = NULL, const FVector::FReal CostLimit = TNumericLimits<FVector::FReal>::Max());
+	FPathFindingQuery(const UObject* InOwner, const ANavigationData& InNavData, const FVector& Start, const FVector& End, FSharedConstNavQueryFilter SourceQueryFilter = NULL, FNavPathSharedPtr InPathInstanceToFill = NULL, const FVector::FReal CostLimit = TNumericLimits<FVector::FReal>::Max(), const bool bInRequireNavigableEndLocation = true);
+	FPathFindingQuery(const INavAgentInterface& InNavAgent, const ANavigationData& InNavData, const FVector& Start, const FVector& End, FSharedConstNavQueryFilter SourceQueryFilter = NULL, FNavPathSharedPtr InPathInstanceToFill = NULL, const FVector::FReal CostLimit = TNumericLimits<FVector::FReal>::Max(), const bool bInRequireNavigableEndLocation = true);
 
 	explicit FPathFindingQuery(FNavPathSharedRef PathToRecalculate, const ANavigationData* NavDataOverride = NULL);
 
 	FPathFindingQuery& SetPathInstanceToUpdate(FNavPathSharedPtr InPathInstanceToFill) { PathInstanceToFill = InPathInstanceToFill; return *this; }
-	FPathFindingQuery& SetAllowPartialPaths(bool bAllow) { bAllowPartialPaths = bAllow; return *this; }
+	FPathFindingQuery& SetAllowPartialPaths(const bool bAllow) { bAllowPartialPaths = bAllow; return *this; }
+	FPathFindingQuery& SetRequireNavigableEndLocation(const bool bRequire) { bRequireNavigableEndLocation = bRequire; return *this; }
 	FPathFindingQuery& SetNavAgentProperties(const FNavAgentProperties& InNavAgentProperties) { NavAgentProperties = InNavAgentProperties; return *this; }
 
 	/** utility function to compute a cost limit using an Euclidean heuristic, an heuristic scale and a cost limit factor
diff --git a/Engine/Source/Runtime/NavigationSystem/Public/NavigationTestingActor.h b/Engine/Source/Runtime/NavigationSystem/Public/NavigationTestingActor.h
--- a/Engine/Source/Runtime/NavigationSystem/Public/NavigationTestingActor.h
+++ b/Engine/Source/Runtime/NavigationSystem/Public/NavigationTestingActor.h
@@ -107,6 +107,10 @@ public:
 	UPROPERTY(EditAnywhere, Category=Pathfinding)
 	uint32 bGatherDetailedInfo : 1;
 
+	/** if set, require the end location to be close to the navigation data. The tolerance is controlled by QueryingExtent */
+	UPROPERTY(EditAnywhere, Category=Pathfinding)
+	uint32 bRequireNavigableEndLocation : 1;
+
 	UPROPERTY(EditAnywhere, Category=Query)
 	uint32 bDrawDistanceToWall : 1;
 
diff --git a/Engine/Source/Runtime/Navmesh/Private/Detour/DetourNavMeshQuery.cpp b/Engine/Source/Runtime/Navmesh/Private/Detour/DetourNavMeshQuery.cpp
--- a/Engine/Source/Runtime/Navmesh/Private/Detour/DetourNavMeshQuery.cpp
+++ b/Engine/Source/Runtime/Navmesh/Private/Detour/DetourNavMeshQuery.cpp
@@ -1492,11 +1492,12 @@ dtStatus dtNavMeshQuery::findPath(dtPolyRef startRef, dtPolyRef endRef,
 	
 	m_queryNodes = 0;
 
-	if (!startRef || !endRef)
+	if (!startRef || (m_query.requireNavigableEndLocation && !endRef))		//@UE
 		return DT_FAILURE | DT_INVALID_PARAM;
 	
 	// Validate input
-	if (!m_nav->isValidPolyRef(startRef) || !m_nav->isValidPolyRef(endRef))
+	// endRef could be 0 if requireNavigableEndLocation is false, but we don't want it to reference a polygon that doesn't exist
+	if (!m_nav->isValidPolyRef(startRef) || (!m_nav->isValidPolyRef(endRef) && (endRef || m_query.requireNavigableEndLocation)))	//@UE
 		return DT_FAILURE | DT_INVALID_PARAM;
 	
 	if (startRef == endRef)
@@ -1920,7 +1921,7 @@ dtStatus dtNavMeshQuery::testClusterPath(dtPolyRef startRef, dtPolyRef endRef) c
 /// path query.
 ///
 dtStatus dtNavMeshQuery::initSlicedFindPath(dtPolyRef startRef, dtPolyRef endRef,
-											const dtReal* startPos, const dtReal* endPos, const dtReal costLimit, //@UE
+											const dtReal* startPos, const dtReal* endPos, const dtReal costLimit, const bool requireNavigableEndLocation, //@UE 
 											const dtQueryFilter* filter)
 {
 	dtAssert(m_nav);
@@ -1935,13 +1936,15 @@ dtStatus dtNavMeshQuery::initSlicedFindPath(dtPolyRef startRef, dtPolyRef endRef
 	dtVcopy(m_query.startPos, startPos);
 	dtVcopy(m_query.endPos, endPos);
 	m_query.costLimit = costLimit; //@UE
+	m_query.requireNavigableEndLocation = requireNavigableEndLocation; //@UE
 	m_query.filter = filter;
 	
-	if (!startRef || !endRef)
+	if (!startRef || (m_query.requireNavigableEndLocation && !endRef))	//@UE
 		return DT_FAILURE | DT_INVALID_PARAM;
 	
 	// Validate input
-	if (!m_nav->isValidPolyRef(startRef) || !m_nav->isValidPolyRef(endRef))
+	// endRef could be 0 if requireNavigableEndLocation is false, but we don't want it to reference a polygon that doesn't exist
+	if (!m_nav->isValidPolyRef(startRef) || (!m_nav->isValidPolyRef(endRef) && (endRef || m_query.requireNavigableEndLocation)))	//@UE
 		return DT_FAILURE | DT_INVALID_PARAM;
 
 	if (startRef == endRef)
@@ -1979,7 +1982,8 @@ dtStatus dtNavMeshQuery::updateSlicedFindPath(const int maxIter, int* doneIters)
 		return m_query.status;
 
 	// Make sure the request is still valid.
-	if (!m_nav->isValidPolyRef(m_query.startRef) || !m_nav->isValidPolyRef(m_query.endRef))
+	// endRef could be 0 if requireNavigableEndLocation is false, but we don't want it to reference a polygon that doesn't exist
+	if (!m_nav->isValidPolyRef(m_query.startRef) || (!m_nav->isValidPolyRef(m_query.endRef) && (m_query.endRef || m_query.requireNavigableEndLocation)))	//@UE
 	{
 		m_query.status = DT_FAILURE;
 		return DT_FAILURE;
diff --git a/Engine/Source/Runtime/Navmesh/Private/DetourCrowd/DetourCrowd.cpp b/Engine/Source/Runtime/Navmesh/Private/DetourCrowd/DetourCrowd.cpp
--- a/Engine/Source/Runtime/Navmesh/Private/DetourCrowd/DetourCrowd.cpp
+++ b/Engine/Source/Runtime/Navmesh/Private/DetourCrowd/DetourCrowd.cpp
@@ -854,8 +854,10 @@ void dtCrowd::updateMoveRequest(const dtReal /*dt*/)
 			// Quick seach towards the goal.
 			static const int MAX_ITER = 20;
 			m_navquery->updateLinkFilter(ag->params.linkFilter.Get());
+
 			const dtReal costLimit = DT_REAL_MAX; //@UE
-			m_navquery->initSlicedFindPath(path[0], ag->targetRef, ag->npos, ag->targetPos, costLimit, &m_filters[ag->params.filter]); //@UE
+			const bool requireNavigableEndLocation = true; //@UE
+			m_navquery->initSlicedFindPath(path[0], ag->targetRef, ag->npos, ag->targetPos, costLimit, requireNavigableEndLocation, &m_filters[ag->params.filter]); //@UE
 			m_navquery->updateSlicedFindPath(MAX_ITER, 0);
 			dtStatus status = 0;
 			if (ag->targetReplan) // && npath > 10)
@@ -922,8 +924,9 @@ void dtCrowd::updateMoveRequest(const dtReal /*dt*/)
 	{
 		dtCrowdAgent* ag = queue[i];
 		const dtReal costLimit = DT_REAL_MAX; //@UE
+		const bool requireNavigableEndLocation = true; //@UE
 		ag->targetPathqRef = m_pathq.request(ag->corridor.getLastPoly(), ag->targetRef,
-			ag->corridor.getTarget(), ag->targetPos, costLimit, &m_filters[ag->params.filter], ag->params.linkFilter); //@UE
+			ag->corridor.getTarget(), ag->targetPos, costLimit, requireNavigableEndLocation, &m_filters[ag->params.filter], ag->params.linkFilter); //@UE
 		if (ag->targetPathqRef != DT_PATHQ_INVALID)
 			ag->targetState = DT_CROWDAGENT_TARGET_WAITING_FOR_PATH;
 	}
diff --git a/Engine/Source/Runtime/Navmesh/Private/DetourCrowd/DetourPathCorridor.cpp b/Engine/Source/Runtime/Navmesh/Private/DetourCrowd/DetourPathCorridor.cpp
--- a/Engine/Source/Runtime/Navmesh/Private/DetourCrowd/DetourPathCorridor.cpp
+++ b/Engine/Source/Runtime/Navmesh/Private/DetourCrowd/DetourPathCorridor.cpp
@@ -525,7 +525,8 @@ bool dtPathCorridor::optimizePathTopology(dtNavMeshQuery* navquery, const dtQuer
 	dtPolyRef res[MAX_RES];
 	int nres = 0;
 	const dtReal costLimit = DT_REAL_MAX; //@UE
-	navquery->initSlicedFindPath(m_path[0], m_path[m_npath-1], m_pos, m_target, costLimit, filter); //@UE
+	const bool requireNavigableEndLocation = true; //@UE
+	navquery->initSlicedFindPath(m_path[0], m_path[m_npath-1], m_pos, m_target, costLimit, requireNavigableEndLocation, filter); //@UE
 	navquery->updateSlicedFindPath(MAX_ITER, 0);
 	dtStatus status = navquery->finalizeSlicedFindPathPartial(m_path, m_npath, res, &nres, MAX_RES);
 	
diff --git a/Engine/Source/Runtime/Navmesh/Private/DetourCrowd/DetourPathQueue.cpp b/Engine/Source/Runtime/Navmesh/Private/DetourCrowd/DetourPathQueue.cpp
--- a/Engine/Source/Runtime/Navmesh/Private/DetourCrowd/DetourPathQueue.cpp
+++ b/Engine/Source/Runtime/Navmesh/Private/DetourCrowd/DetourPathQueue.cpp
@@ -111,7 +111,7 @@ void dtPathQueue::update(const int maxIters)
 		// Handle query start.
 		if (q.status == 0)
 		{
-			q.status = m_navquery->initSlicedFindPath(q.startRef, q.endRef, q.startPos, q.endPos, q.costLimit, q.filter); //@UE
+			q.status = m_navquery->initSlicedFindPath(q.startRef, q.endRef, q.startPos, q.endPos, q.costLimit, q.requireNavigableEndLocation, q.filter); //@UE
 		}		
 		// Handle query in progress.
 		if (dtStatusInProgress(q.status))
@@ -133,7 +133,7 @@ void dtPathQueue::update(const int maxIters)
 }
 
 dtPathQueueRef dtPathQueue::request(dtPolyRef startRef, dtPolyRef endRef,
-									const dtReal* startPos, const dtReal* endPos, const dtReal costLimit, //@UE
+									const dtReal* startPos, const dtReal* endPos, const dtReal costLimit, const bool requireNavigableEndLocation, //@UE
 									const dtQueryFilter* filter, TSharedPtr<dtQuerySpecialLinkFilter> linkFilter)
 {
 	// Find empty slot
@@ -160,6 +160,7 @@ dtPathQueueRef dtPathQueue::request(dtPolyRef startRef, dtPolyRef endRef,
 	dtVcopy(q.endPos, endPos);
 	q.endRef = endRef;
 	q.costLimit = costLimit; //@UE
+	q.requireNavigableEndLocation = requireNavigableEndLocation; //@UE
 	
 	q.status = 0;
 	q.npath = 0;
diff --git a/Engine/Source/Runtime/Navmesh/Public/Detour/DetourNavMeshQuery.h b/Engine/Source/Runtime/Navmesh/Public/Detour/DetourNavMeshQuery.h
--- a/Engine/Source/Runtime/Navmesh/Public/Detour/DetourNavMeshQuery.h
+++ b/Engine/Source/Runtime/Navmesh/Public/Detour/DetourNavMeshQuery.h
@@ -392,15 +392,16 @@ public:
 	///@{ 
 
 	/// Initializes a sliced path query.
-	///  @param[in]		startRef	The refrence id of the start polygon.
-	///  @param[in]		endRef		The reference id of the end polygon.
-	///  @param[in]		startPos	A position within the start polygon. [(x, y, z)]
-	///  @param[in]		endPos		A position within the end polygon. [(x, y, z)]
-	///  @param[in]		costLimit	Cost limit of nodes allowed to be added to the open list	//@UE
-	///  @param[in]		filter		The polygon filter to apply to the query.
+	///  @param[in]		startRef					The refrence id of the start polygon.
+	///  @param[in]		endRef						The reference id of the end polygon.
+	///  @param[in]		startPos					A position within the start polygon. [(x, y, z)]
+	///  @param[in]		endPos						A position within the end polygon. [(x, y, z)]
+	///  @param[in]		costLimit					Cost limit of nodes allowed to be added to the open list	//@UE
+	///  @param[in]		requireNavigableEndLocation	Define if the end location is required to be a valid navmesh polygon //@UE
+	///  @param[in]		filter						The polygon filter to apply to the query.
 	/// @returns The status flags for the query.
 	dtStatus initSlicedFindPath(dtPolyRef startRef, dtPolyRef endRef,
-								const dtReal* startPos, const dtReal* endPos, const dtReal costLimit, //@UE
+								const dtReal* startPos, const dtReal* endPos, const dtReal costLimit, const bool requireNavigableEndLocation, //@UE
 								const dtQueryFilter* filter);
 
 	/// Updates an in-progress sliced path query.
@@ -812,6 +813,11 @@ public:
 			&& m_linkFilter->isLinkAllowed(tile->offMeshCons[linkIdx].userId) == false);
 	}
 
+	//@UE BEGIN
+	inline void setRequireNavigableEndLocation(const bool value) { m_query.requireNavigableEndLocation = value; }
+	inline bool isRequiringNavigableEndLocation() const { return m_query.requireNavigableEndLocation; }
+	//@UE END
+
 private:
 	const dtNavMesh* m_nav;							///< Pointer to navmesh data.
 	dtQuerySpecialLinkFilter* m_linkFilter;			///< Pointer to optional special link filter
@@ -825,6 +831,7 @@ private:
 		dtReal startPos[3], endPos[3];
 		dtReal costLimit; 					//@UE ///< Cost limit of nodes allowed to be added to the open list
 		const dtQueryFilter* filter;
+		unsigned char requireNavigableEndLocation : 1;	//@UE
 	};
 	dtQueryData m_query;				///< Sliced query state.
 
diff --git a/Engine/Source/Runtime/Navmesh/Public/DetourCrowd/DetourPathQueue.h b/Engine/Source/Runtime/Navmesh/Public/DetourCrowd/DetourPathQueue.h
--- a/Engine/Source/Runtime/Navmesh/Public/DetourCrowd/DetourPathQueue.h
+++ b/Engine/Source/Runtime/Navmesh/Public/DetourCrowd/DetourPathQueue.h
@@ -46,6 +46,7 @@ class dtPathQueue
 		dtReal startPos[3], endPos[3];
 		dtPolyRef startRef, endRef;
 		dtReal costLimit;
+		unsigned char requireNavigableEndLocation : 1;	// @UE
 		/// Result.
 		dtPolyRef* path;
 		const dtQueryFilter* filter;
@@ -74,7 +75,7 @@ public:
 	void update(const int maxIters);
 	
 	dtPathQueueRef request(dtPolyRef startRef, dtPolyRef endRef,
-						   const dtReal* startPos, const dtReal* endPos, const dtReal costLimit,
+						   const dtReal* startPos, const dtReal* endPos, const dtReal costLimit, const bool requireNavigableEndLocation, //@UE
 						   const dtQueryFilter* filter,
 						   TSharedPtr<dtQuerySpecialLinkFilter> linkFilter);
 	
