build fix : missing param from merge

diff --git a/Engine/Source/Runtime/NavigationSystem/Private/NavigationData.cpp b/Engine/Source/Runtime/NavigationSystem/Private/NavigationData.cpp
--- a/Engine/Source/Runtime/NavigationSystem/Private/NavigationData.cpp
+++ b/Engine/Source/Runtime/NavigationSystem/Private/NavigationData.cpp
@@ -37,8 +37,8 @@ FPathFindingQuery::FPathFindingQuery(const UObject* InOwner, const ANavigationDa
 	}
 }
 
-FPathFindingQuery::FPathFindingQuery(const INavAgentInterface& InNavAgent, const ANavigationData& InNavData, const FVector& Start, const FVector& End, FSharedConstNavQueryFilter SourceQueryFilter, FNavPathSharedPtr InPathInstanceToFill, const FVector::FReal CostLimit) :
-	FPathFindingQueryData(Cast<UObject>(&InNavAgent), Start, End, SourceQueryFilter, 0 /*InNavDataFlags*/, true /*bInAllowPartialPaths*/, CostLimit),
+FPathFindingQuery::FPathFindingQuery(const INavAgentInterface& InNavAgent, const ANavigationData& InNavData, const FVector& Start, const FVector& End, FSharedConstNavQueryFilter SourceQueryFilter, FNavPathSharedPtr InPathInstanceToFill, const FVector::FReal CostLimit, const bool bInRequireNavigableEndLocation) :
+	FPathFindingQueryData(Cast<UObject>(&InNavAgent), Start, End, SourceQueryFilter, 0 /*InNavDataFlags*/, true /*bInAllowPartialPaths*/, CostLimit, bInRequireNavigableEndLocation),
 	NavData(&InNavData), PathInstanceToFill(InPathInstanceToFill), NavAgentProperties(InNavAgent.GetNavAgentPropertiesRef())
 {
 	if (!QueryFilter.IsValid() && NavData.IsValid())
