fix missing include in RecastQueryFilter.h

diff --git a/Engine/Source/Runtime/NavigationSystem/Public/NavMesh/RecastQueryFilter.h b/Engine/Source/Runtime/NavigationSystem/Public/NavMesh/RecastQueryFilter.h
--- a/Engine/Source/Runtime/NavigationSystem/Public/NavMesh/RecastQueryFilter.h
+++ b/Engine/Source/Runtime/NavigationSystem/Public/NavMesh/RecastQueryFilter.h
@@ -7,6 +7,7 @@
 #include "NavFilters/NavigationQueryFilter.h"
 #endif //UE_ENABLE_INCLUDE_ORDER_DEPRECATED_IN_5_2
 #include "AI/Navigation/NavigationTypes.h"
+#include "AI/Navigation/NavQueryFilter.h"
 
 #if WITH_RECAST
 #include "Detour/DetourNavMesh.h"
