Fix crash in Lyra (needed when tileWidth of loading navmesh mismatch)

diff --git a/Engine/Source/Runtime/NavigationSystem/Private/NavMesh/PImplRecastNavMesh.cpp b/Engine/Source/Runtime/NavigationSystem/Private/NavMesh/PImplRecastNavMesh.cpp
--- a/Engine/Source/Runtime/NavigationSystem/Private/NavMesh/PImplRecastNavMesh.cpp
+++ b/Engine/Source/Runtime/NavigationSystem/Private/NavMesh/PImplRecastNavMesh.cpp
@@ -1237,6 +1237,12 @@ ENavigationQueryResult::Type FPImplRecastNavMesh::TestPath(const FVector& StartL
 
 ENavigationQueryResult::Type FPImplRecastNavMesh::TestPath(const FVector& StartLoc, const FVector& EndLoc, const bool bRequireNavigableEndLocation, const FNavigationQueryFilter& InQueryFilter, const UObject* Owner, int32* NumVisitedNodes) const
 {
+	// Same check as in FPImplRecastNavMesh::FindPath (ex: this can occur when tileWidth of loading navmesh mismatch).
+	if (DetourNavMesh == NULL || NavMeshOwner == NULL)
+	{
+		return ENavigationQueryResult::Error;
+	}
+	
 	const dtQueryFilter* QueryFilter = ((const FRecastQueryFilter*)(InQueryFilter.GetImplementation()))->GetAsDetourQueryFilter();
 	if (QueryFilter == NULL)
 	{
