[Navmesh] On previous navmesh version, read deprecated CellSize property during serialization to avoid invalidating the navmesh

diff --git a/Engine/Source/Runtime/NavigationSystem/Private/NavMesh/PImplRecastNavMesh.cpp b/Engine/Source/Runtime/NavigationSystem/Private/NavMesh/PImplRecastNavMesh.cpp
--- a/Engine/Source/Runtime/NavigationSystem/Private/NavMesh/PImplRecastNavMesh.cpp
+++ b/Engine/Source/Runtime/NavigationSystem/Private/NavMesh/PImplRecastNavMesh.cpp
@@ -597,7 +597,17 @@ void FPImplRecastNavMesh::Serialize( FArchive& Ar, int32 NavMeshVersion )
 		// at this point we can tell whether navmesh being loaded is in line
 		// ARecastNavMesh's params. If not, just skip it.
 		// assumes tiles are rectangular
-		const float DefaultCellSize = NavMeshOwner->GetCellSize(ENavigationDataResolution::Default);
+		
+		float DefaultCellSize = NavMeshOwner->GetCellSize(ENavigationDataResolution::Default);
+		PRAGMA_DISABLE_DEPRECATION_WARNINGS
+		if (NavMeshVersion < NAVMESHVER_TILE_RESOLUTIONS)
+		{
+			// For backward compatibility, read original CellSize value.
+			// In ARecastNavMesh::PostLoad(), cell sizes for the different resolutions are set to the old CellSize value but it occurs later (PostLoad).
+			// This why we explicitly need to read CellSize for older versions here.
+			DefaultCellSize = NavMeshOwner->CellSize;
+		}
+		PRAGMA_ENABLE_DEPRECATION_WARNINGS
 		const FVector::FReal ActorsTileSize = FVector::FReal(int32(NavMeshOwner->TileSizeUU / DefaultCellSize) * DefaultCellSize);
 
 		if (ActorsTileSize != Params.tileWidth)
