Extract method that forward information to NavData and NavGenerator when we activate NavInvoker usage

diff --git a/Engine/Source/Runtime/NavigationSystem/Private/NavigationSystem.cpp b/Engine/Source/Runtime/NavigationSystem/Private/NavigationSystem.cpp
--- a/Engine/Source/Runtime/NavigationSystem/Private/NavigationSystem.cpp
+++ b/Engine/Source/Runtime/NavigationSystem/Private/NavigationSystem.cpp
@@ -865,27 +865,16 @@ void UNavigationSystemV1::PostEditChangeChainProperty(FPropertyChangedChainEvent
 
 void UNavigationSystemV1::PostEditChangeProperty(FPropertyChangedEvent& PropertyChangedEvent)
 {
-	static const FName NAME_EnableActiveTiles = GET_MEMBER_NAME_CHECKED(UNavigationSystemV1, bGenerateNavigationOnlyAroundNavigationInvokers);
+	static const FName NAME_GenerateNavigationOnlyAroundNavigationInvokers = GET_MEMBER_NAME_CHECKED(UNavigationSystemV1, bGenerateNavigationOnlyAroundNavigationInvokers);
 
 	Super::PostEditChangeProperty(PropertyChangedEvent);
 
 	if (PropertyChangedEvent.Property)
 	{
 		FName PropName = PropertyChangedEvent.Property->GetFName();
-		if (PropName == NAME_EnableActiveTiles)
+		if (PropName == NAME_GenerateNavigationOnlyAroundNavigationInvokers)
 		{
-			if (DefaultOctreeController.NavOctree.IsValid())
-			{
-				DefaultOctreeController.NavOctree->SetDataGatheringMode(DataGatheringMode);
-			}
-
-			for (auto NavData : NavDataSet)
-			{
-				if (NavData)
-				{
-					NavData->RestrictBuildingToActiveTiles(bGenerateNavigationOnlyAroundNavigationInvokers);
-				}
-			}
+			OnGenerateNavigationOnlyAroundNavigationInvokersChanged();
 		}
 		else if (PropName == GET_MEMBER_NAME_CHECKED(FNavDataConfig, AgentRadius))
 		{
@@ -4322,6 +4311,22 @@ bool UNavigationSystemV1::IsAllowedToRebuild() const
 	return World && (!World->IsGameWorld() || GetRuntimeGenerationType() == ERuntimeGenerationType::Dynamic);
 }
 
+void UNavigationSystemV1::OnGenerateNavigationOnlyAroundNavigationInvokersChanged()
+{
+	if (DefaultOctreeController.NavOctree.IsValid())
+	{
+		DefaultOctreeController.NavOctree->SetDataGatheringMode(DataGatheringMode);
+	}
+
+	for (auto NavData : NavDataSet)
+	{
+		if (NavData)
+		{
+			NavData->RestrictBuildingToActiveTiles(bGenerateNavigationOnlyAroundNavigationInvokers);
+		}
+	}
+}
+
 //----------------------------------------------------------------------//
 // Blueprint functions
 //----------------------------------------------------------------------//
diff --git a/Engine/Source/Runtime/NavigationSystem/Public/NavigationSystem.h b/Engine/Source/Runtime/NavigationSystem/Public/NavigationSystem.h
--- a/Engine/Source/Runtime/NavigationSystem/Public/NavigationSystem.h
+++ b/Engine/Source/Runtime/NavigationSystem/Public/NavigationSystem.h
@@ -1284,6 +1284,9 @@ protected:
 	 *  Depends on runtime generation settings of each navigation data, always true in the editor
 	 */
 	bool IsAllowedToRebuild() const;
+
+	/** Handle forwarding the information where needed when the setting is changed*/
+	void OnGenerateNavigationOnlyAroundNavigationInvokersChanged();
 	
 	//----------------------------------------------------------------------//
 	// new stuff
