Fix invalid rendering in nav tester when the projected location is not on the navmesh

diff --git a/Engine/Source/Runtime/NavigationSystem/Private/NavMesh/NavTestRenderingComponent.cpp b/Engine/Source/Runtime/NavigationSystem/Private/NavMesh/NavTestRenderingComponent.cpp
--- a/Engine/Source/Runtime/NavigationSystem/Private/NavMesh/NavTestRenderingComponent.cpp
+++ b/Engine/Source/Runtime/NavigationSystem/Private/NavMesh/NavTestRenderingComponent.cpp
@@ -67,18 +67,21 @@ void FNavTestSceneProxy::GetDynamicMeshElements(const TArray<const FSceneView*>&
 
 				//@todo - the rendering thread should never read from UObjects directly!  These are race conditions, the properties should be mirrored on the proxy
 				const FVector ActorLocation = NavTestActor->GetActorLocation();
-				const FVector ProjectedLocation = NavTestActor->ProjectedLocation + (FVector)NavMeshDrawOffset;
-				const FColor ProjectedColor = NavTestActor->bProjectedLocationValid ? FColor(0, 255, 0, 120) : FColor(255, 0, 0, 120);
+				const FVector ProjectedLocation = NavTestActor->bProjectedLocationValid ? (NavTestActor->ProjectedLocation + (FVector)NavMeshDrawOffset) : (ActorLocation - FVector(0, 0, NavTestActor->QueryingExtent.Z));
+				const FColor ProjectedColor = (NavTestActor->bProjectedLocationValid ? FColor::Green : FColor::Red).WithAlpha(120);
 				const FColor ClosestWallColor = FColorList::Orange;
 				const FVector BoxExtent(20, 20, 20);
 
 				const FMaterialRenderProxy* const ColoredMeshInstance = &Collector.AllocateOneFrameResource<FColoredMaterialRenderProxy>(GEngine->DebugMeshMaterial->GetRenderProxy(), ProjectedColor);
 				//DrawBox(PDI, FTransform(ProjectedLocation).ToMatrixNoScale(),BoxExtent, ColoredMeshInstance, SDPG_World);
-				GetSphereMesh(ProjectedLocation, BoxExtent, 10, 7, ColoredMeshInstance, SDPG_World, false, ViewIndex, Collector);
+				if (NavTestActor->bProjectedLocationValid)
+				{
+					GetSphereMesh(ProjectedLocation, BoxExtent, 10, 7, ColoredMeshInstance, SDPG_World, false, ViewIndex, Collector);
+				}
 
 				//DrawWireBox(PDI, FBox(ProjectedLocation-BoxExtent, ProjectedLocation+BoxExtent), ProjectedColor, false);
 				DrawWireBox(PDI, FBox(ActorLocation - BoxExtent, ActorLocation + BoxExtent), FColor::White, false);
-				const FVector LineEnd = ProjectedLocation - (ProjectedLocation - ActorLocation).GetSafeNormal()*BoxExtent.X;
+				const FVector LineEnd = NavTestActor->bProjectedLocationValid ? ProjectedLocation - (ProjectedLocation - ActorLocation).GetSafeNormal()*BoxExtent.X : ProjectedLocation;
 				PDI->DrawLine(LineEnd, ActorLocation, ProjectedColor, SDPG_World, 2.5);
 				DrawArrowHead(PDI, LineEnd, ActorLocation, 20.f, ProjectedColor, SDPG_World, 2.5f);
 
