Change Navmesh debug display to only see the tiles around the debug actor even when active tile generation is enabled

diff --git a/Engine/Source/Runtime/NavigationSystem/Private/NavMesh/PImplRecastNavMesh.cpp b/Engine/Source/Runtime/NavigationSystem/Private/NavMesh/PImplRecastNavMesh.cpp
--- a/Engine/Source/Runtime/NavigationSystem/Private/NavMesh/PImplRecastNavMesh.cpp
+++ b/Engine/Source/Runtime/NavigationSystem/Private/NavMesh/PImplRecastNavMesh.cpp
@@ -2721,45 +2721,34 @@ bool FPImplRecastNavMesh::GetDebugGeometryForTile(FRecastDebugGeometry& OutGeome
 	check(NavMeshOwner);
 
 	const dtNavMesh* const ConstNavMesh = DetourNavMesh;
-		
-	// presize our tarrays for efficiency
-	const int32 NumTiles = TileIndex == INDEX_NONE ? ConstNavMesh->getMaxTiles() : TileIndex + 1;
-	const int32 StartingTile = TileIndex == INDEX_NONE ? 0 : TileIndex;
 
+	// presize our tarrays for efficiency
 	int32 NumVertsToReserve = 0;
 	int32 NumIndicesToReserve = 0;
 
-	uint16 ForbiddenFlags = OutGeometry.bMarkForbiddenPolys 
+	const uint16 ForbiddenFlags = OutGeometry.bMarkForbiddenPolys 
 		? GetFilterForbiddenFlags((const FRecastQueryFilter*)NavMeshOwner->GetDefaultQueryFilterImpl()) 
 		: 0;
 
 	const FRecastNavMeshGenerator* Generator = static_cast<const FRecastNavMeshGenerator*>(NavMeshOwner->GetGenerator());
+	const bool bIsGenerationRestrictedToActiveTiles = Generator && Generator->IsBuildingRestrictedToActiveTiles() && NavMeshOwner->GetActiveTiles().Num() > 0;
 
-	if (Generator && Generator->IsBuildingRestrictedToActiveTiles()
-		// if not active tiles try drawing all tiles
-		&& NavMeshOwner->GetActiveTiles().Num() > 0)
+	auto ComputeSizeToReserve = [](dtMeshTile const* const Tile, int32& OutNumVertsToReserve, int32& OutNumIndicesToReserve)
 	{
-		const TArray<FIntPoint>& ActiveTiles = NavMeshOwner->GetActiveTiles();
-		for (const FIntPoint& TileLocation : ActiveTiles)
+		if (Tile != nullptr && Tile->header != nullptr)
 		{
-			const int32 LayersCount = ConstNavMesh->getTileCountAt(TileLocation.X, TileLocation.Y);
+			OutNumVertsToReserve += Tile->header->vertCount + Tile->header->detailVertCount;
 
-			for (int32 Layer = 0; Layer < LayersCount; ++Layer)
+			for (int32 PolyIdx = 0; PolyIdx < Tile->header->polyCount; ++PolyIdx)
 			{
-				dtMeshTile const* const Tile = ConstNavMesh->getTileAt(TileLocation.X, TileLocation.Y, Layer);
-				if (Tile != nullptr && Tile->header != nullptr)
-				{
-					NumVertsToReserve += Tile->header->vertCount + Tile->header->detailVertCount;
-
-					for (int32 PolyIdx = 0; PolyIdx < Tile->header->polyCount; ++PolyIdx)
-					{
-						dtPolyDetail const* const DetailPoly = &Tile->detailMeshes[PolyIdx];
-						NumIndicesToReserve += (DetailPoly->triCount * 3);
-					}
-				}
+				dtPolyDetail const* const DetailPoly = &Tile->detailMeshes[PolyIdx];
+				OutNumIndicesToReserve += (DetailPoly->triCount * 3);
 			}
 		}
+	};
 
+	auto ReserveGeometryArrays = [](FRecastDebugGeometry& OutGeometry, const int32 NumVertsToReserve, const int32 NumIndicesToReserve)
+	{
 		OutGeometry.MeshVerts.Reserve(OutGeometry.MeshVerts.Num() + NumVertsToReserve);
 		OutGeometry.AreaIndices[0].Reserve(OutGeometry.AreaIndices[0].Num() + NumIndicesToReserve);
 		OutGeometry.BuiltMeshIndices.Reserve(OutGeometry.BuiltMeshIndices.Num() + NumIndicesToReserve);
@@ -2767,6 +2756,40 @@ bool FPImplRecastNavMesh::GetDebugGeometryForTile(FRecastDebugGeometry& OutGeome
 		{
 			OutGeometry.TileBuildTimesIndices[Index].Reserve(OutGeometry.TileBuildTimesIndices[Index].Num() + NumIndicesToReserve);	
 		}
+	};
+
+	if (TileIndex != INDEX_NONE)
+	{
+		dtMeshTile const* const Tile = ConstNavMesh->getTile(TileIndex);
+		if (Tile != nullptr && Tile->header != nullptr)
+		{
+			const FIntPoint TileCoord = FIntPoint(Tile->header->x, Tile->header->y);
+			if (!bIsGenerationRestrictedToActiveTiles || Generator->IsInActiveSet(TileCoord))
+			{
+				ComputeSizeToReserve(Tile, NumVertsToReserve, NumIndicesToReserve);
+
+				ReserveGeometryArrays(OutGeometry, NumVertsToReserve, NumIndicesToReserve);
+
+				const uint32 VertBase = OutGeometry.MeshVerts.Num();
+				GetTilesDebugGeometry(Generator, *Tile, VertBase, OutGeometry, TileIndex, ForbiddenFlags);
+			}
+		}
+	}
+	else if (bIsGenerationRestrictedToActiveTiles)
+	{
+		const TArray<FIntPoint>& ActiveTiles = NavMeshOwner->GetActiveTiles();
+		for (const FIntPoint& TileLocation : ActiveTiles)
+		{
+			const int32 LayersCount = ConstNavMesh->getTileCountAt(TileLocation.X, TileLocation.Y);
+
+			for (int32 Layer = 0; Layer < LayersCount; ++Layer)
+			{
+				dtMeshTile const* const Tile = ConstNavMesh->getTileAt(TileLocation.X, TileLocation.Y, Layer);
+				ComputeSizeToReserve(Tile, NumVertsToReserve, NumIndicesToReserve);
+			}
+		}
+
+		ReserveGeometryArrays(OutGeometry, NumVertsToReserve, NumIndicesToReserve);
 
 		uint32 VertBase = OutGeometry.MeshVerts.Num();
 		for (const FIntPoint& TileLocation : ActiveTiles)
@@ -2787,48 +2810,26 @@ bool FPImplRecastNavMesh::GetDebugGeometryForTile(FRecastDebugGeometry& OutGeome
 	}
 	else
 	{
-		for (int32 TileIdx = StartingTile; TileIdx < NumTiles; ++TileIdx)
+		const int32 NumTiles = ConstNavMesh->getMaxTiles();
+		for (int32 TileIdx = 0; TileIdx < NumTiles; ++TileIdx)
 		{
 			dtMeshTile const* const Tile = ConstNavMesh->getTile(TileIdx);
-			dtMeshHeader const* const Header = Tile->header;
-
-			if (Header != NULL)
-			{
-				NumVertsToReserve += Header->vertCount + Header->detailVertCount;
-
-				for (int32 PolyIdx = 0; PolyIdx < Header->polyCount; ++PolyIdx)
-				{
-					dtPolyDetail const* const DetailPoly = &Tile->detailMeshes[PolyIdx];
-					NumIndicesToReserve += (DetailPoly->triCount * 3);
-				}
-			}
+			ComputeSizeToReserve(Tile, NumVertsToReserve, NumIndicesToReserve);
 		}
 
-		OutGeometry.MeshVerts.Reserve(OutGeometry.MeshVerts.Num() + NumVertsToReserve);
-		OutGeometry.AreaIndices[0].Reserve(OutGeometry.AreaIndices[0].Num() + NumIndicesToReserve);
-		OutGeometry.BuiltMeshIndices.Reserve(OutGeometry.BuiltMeshIndices.Num() + NumIndicesToReserve);
-		for (int32 Index = 0; Index < FRecastDebugGeometry::BuildTimeBucketsCount; Index++)
-		{
-			OutGeometry.TileBuildTimesIndices[Index].Reserve(OutGeometry.TileBuildTimesIndices[Index].Num() + NumIndicesToReserve);	
-		}
+		ReserveGeometryArrays(OutGeometry, NumVertsToReserve, NumIndicesToReserve);
 
 		uint32 VertBase = OutGeometry.MeshVerts.Num();
-		for (int32 TileIdx = StartingTile; TileIdx < NumTiles; ++TileIdx)
+		for (int32 TileIdx = 0; TileIdx < NumTiles; ++TileIdx)
 		{
 			dtMeshTile const* const Tile = ConstNavMesh->getTile(TileIdx);
-
-			if (Tile == nullptr || Tile->header == nullptr)
+			if (Tile != nullptr && Tile->header != nullptr)
 			{
-				continue;
+				VertBase += GetTilesDebugGeometry(Generator, *Tile, VertBase, OutGeometry, TileIdx, ForbiddenFlags);
 			}
-
-			VertBase += GetTilesDebugGeometry(Generator, *Tile, VertBase, OutGeometry, TileIdx, ForbiddenFlags);
 		}
 
-		if (TileIndex == INDEX_NONE)
-		{
-			bDone = true;
-		}
+		bDone = true;
 	}
 
 	return bDone;
