Fix navigation bounds not being updated on load when "Update Navigation Automatically" is disabled.
This fixes world partitioned navmesh data chunks not being generated for areas covered by NavMeshBoundsVolume that are placed inside data layers.

diff --git a/Engine/Source/Runtime/NavigationSystem/Private/NavMesh/RecastNavMeshGenerator.cpp b/Engine/Source/Runtime/NavigationSystem/Private/NavMesh/RecastNavMeshGenerator.cpp
--- a/Engine/Source/Runtime/NavigationSystem/Private/NavMesh/RecastNavMeshGenerator.cpp
+++ b/Engine/Source/Runtime/NavigationSystem/Private/NavMesh/RecastNavMeshGenerator.cpp
@@ -5064,30 +5064,34 @@ void FRecastNavMeshGenerator::OnNavigationBoundsChanged()
 	dtNavMesh* DetourMesh = DestNavMesh->GetRecastNavMeshImpl() ? DestNavMesh->GetRecastNavMeshImpl()->GetRecastMesh() : nullptr;
 	if (!IsGameStaticNavMesh(DestNavMesh) && DestNavMesh->IsResizable() && DetourMesh)
 	{
-		// Check whether Navmesh size needs to be changed
-		const int32 MaxRequestedTiles = UE::NavMesh::Private::CalculateMaxTilesCount(InclusionBounds, Config.GetTileSizeUU(), AvgLayersPerTile, DestNavMesh->NavMeshVersion);
-		if (DetourMesh->getMaxTiles() != MaxRequestedTiles)
+		const UNavigationSystemV1* NavSys = FNavigationSystem::GetCurrent<UNavigationSystemV1>(GetWorld());
+		if (NavSys && !NavSys->IsNavigationBuildingLocked())
 		{
-			UE_LOG(LogNavigation, Log, TEXT("%s> Navigation bounds changed, rebuilding navmesh"), *DestNavMesh->GetName());
-			// Destroy current NavMesh
-			DestNavMesh->GetRecastNavMeshImpl()->SetRecastMesh(nullptr);
-
-			// if there are any valid bounds recreate detour navmesh instance
-			// and mark all bounds as dirty
-			if (InclusionBounds.Num() > 0)
+			// Check whether Navmesh size needs to be changed
+			const int32 MaxRequestedTiles = UE::NavMesh::Private::CalculateMaxTilesCount(InclusionBounds, Config.GetTileSizeUU(), AvgLayersPerTile, DestNavMesh->NavMeshVersion);
+			if (DetourMesh->getMaxTiles() != MaxRequestedTiles)
 			{
-				TArray<FNavigationDirtyArea> AsDirtyAreas;
-				AsDirtyAreas.Reserve(InclusionBounds.Num());
-				for (const FBox& BBox : InclusionBounds)
+				UE_LOG(LogNavigation, Log, TEXT("%s> Navigation bounds changed, rebuilding navmesh"), *DestNavMesh->GetName());
+				// Destroy current NavMesh
+				DestNavMesh->GetRecastNavMeshImpl()->SetRecastMesh(nullptr);
+
+				// if there are any valid bounds recreate detour navmesh instance
+				// and mark all bounds as dirty
+				if (InclusionBounds.Num() > 0)
 				{
-					AsDirtyAreas.Add(FNavigationDirtyArea(BBox, ENavigationDirtyFlag::NavigationBounds));
-				}
+					TArray<FNavigationDirtyArea> AsDirtyAreas;
+					AsDirtyAreas.Reserve(InclusionBounds.Num());
+					for (const FBox& BBox : InclusionBounds)
+					{
+						AsDirtyAreas.Add(FNavigationDirtyArea(BBox, ENavigationDirtyFlag::NavigationBounds));
+					}
 				
-				RebuildDirtyAreas(AsDirtyAreas);
+					RebuildDirtyAreas(AsDirtyAreas);
+				}
 			}
-		}
 
-		UE::NavMesh::Private::CheckTileIndicesInValidRange(InclusionBounds, *DestNavMesh);
+			UE::NavMesh::Private::CheckTileIndicesInValidRange(InclusionBounds, *DestNavMesh);
+		}
 	}
 }
 
diff --git a/Engine/Source/Runtime/NavigationSystem/Private/NavigationSystem.cpp b/Engine/Source/Runtime/NavigationSystem/Private/NavigationSystem.cpp
--- a/Engine/Source/Runtime/NavigationSystem/Private/NavigationSystem.cpp
+++ b/Engine/Source/Runtime/NavigationSystem/Private/NavigationSystem.cpp
@@ -3441,19 +3441,19 @@ void UNavigationSystemV1::PerformNavigationBoundsUpdate(const TArray<FNavigation
 		}
 	}
 
-	if (!IsNavigationBuildingLocked())
+	if (UpdatedAreas.Num())
 	{
-		if (UpdatedAreas.Num())
+		for (ANavigationData* NavData : NavDataSet)
 		{
-			for (ANavigationData* NavData : NavDataSet)
+			if (NavData)
 			{
-				if (NavData)
-				{
-					NavData->OnNavigationBoundsChanged();	
-				}
+				NavData->OnNavigationBoundsChanged();	
 			}
 		}
+	}
 
+	if (!IsNavigationBuildingLocked())
+	{
 		// Propagate to generators areas that needs to be updated
 		AddDirtyAreas(UpdatedAreas, ENavigationDirtyFlag::All | ENavigationDirtyFlag::NavigationBounds, "Navigation bounds update");
 	}
