Add NavAgent selector in NavInvoker system : it allows to define areas where the navmesh should be generated for specific agents, not necessarily for all of them

diff --git a/Engine/Source/Runtime/NavigationSystem/Private/NavMesh/RecastNavMesh.cpp b/Engine/Source/Runtime/NavigationSystem/Private/NavMesh/RecastNavMesh.cpp
--- a/Engine/Source/Runtime/NavigationSystem/Private/NavMesh/RecastNavMesh.cpp
+++ b/Engine/Source/Runtime/NavigationSystem/Private/NavMesh/RecastNavMesh.cpp
@@ -3406,6 +3406,11 @@ void ARecastNavMesh::UpdateActiveTiles(const TArray<FNavigationInvokerRaw>& Invo
 
 	for (const FNavigationInvokerRaw& Invoker : InvokerLocations)
 	{
+		if (!Invoker.SupportedAgents.Contains(Config.AgentIndex))
+		{
+			continue;
+		}
+
 		const FVector InvokerRelativeLocation = (NavmeshOrigin - Invoker.Location);
 		const  FVector::FReal TileCenterDistanceToRemoveSq = FMath::Square(TileDim * UE_SQRT_2 / 2 + Invoker.RadiusMax);
 		const  FVector::FReal TileCenterDistanceToAddSq = FMath::Square(TileDim * UE_SQRT_2 / 2 + Invoker.RadiusMin);
diff --git a/Engine/Source/Runtime/NavigationSystem/Private/NavigationInvokerComponent.cpp b/Engine/Source/Runtime/NavigationSystem/Private/NavigationInvokerComponent.cpp
--- a/Engine/Source/Runtime/NavigationSystem/Private/NavigationInvokerComponent.cpp
+++ b/Engine/Source/Runtime/NavigationSystem/Private/NavigationInvokerComponent.cpp
@@ -11,6 +11,7 @@ UNavigationInvokerComponent::UNavigationInvokerComponent(const FObjectInitialize
 	, TileRemovalRadius(5000)
 {
 	bAutoActivate = true;
+	SupportedAgents.MarkInitialized();
 }
 
 void UNavigationInvokerComponent::Activate(bool bReset)
@@ -20,7 +21,7 @@ void UNavigationInvokerComponent::Activate(bool bReset)
 	AActor* Owner = GetOwner();
 	if (Owner)
 	{
-		UNavigationSystemV1::RegisterNavigationInvoker(*Owner, TileGenerationRadius, TileRemovalRadius);
+		UNavigationSystemV1::RegisterNavigationInvoker(*Owner, TileGenerationRadius, TileRemovalRadius, SupportedAgents);
 	}
 }
 
@@ -35,6 +36,12 @@ void UNavigationInvokerComponent::Deactivate()
 	}
 }
 
+void UNavigationInvokerComponent::PostInitProperties()
+{
+	Super::PostInitProperties();
+	SupportedAgents.MarkInitialized();
+}
+
 void UNavigationInvokerComponent::RegisterWithNavigationSystem(UNavigationSystemV1& NavSys)
 {
 	if (IsActive())
@@ -42,7 +49,7 @@ void UNavigationInvokerComponent::RegisterWithNavigationSystem(UNavigationSystem
 		AActor* Owner = GetOwner();
 		if (Owner)
 		{
-			NavSys.RegisterInvoker(*Owner, TileGenerationRadius, TileRemovalRadius);
+			NavSys.RegisterInvoker(*Owner, TileGenerationRadius, TileRemovalRadius, SupportedAgents);
 		}
 	}
 }
@@ -51,5 +58,4 @@ void UNavigationInvokerComponent::SetGenerationRadii(const float GenerationRadiu
 {
 	TileGenerationRadius = GenerationRadius;
 	TileRemovalRadius = RemovalRadius;
-}
-
+}
\ No newline at end of file
diff --git a/Engine/Source/Runtime/NavigationSystem/Private/NavigationSystem.cpp b/Engine/Source/Runtime/NavigationSystem/Private/NavigationSystem.cpp
--- a/Engine/Source/Runtime/NavigationSystem/Private/NavigationSystem.cpp
+++ b/Engine/Source/Runtime/NavigationSystem/Private/NavigationSystem.cpp
@@ -221,15 +221,16 @@ FNavigationInvoker::FNavigationInvoker()
 , GenerationRadius(0)
 , RemovalRadius(0)
 {
-
+	SupportedAgents.MarkInitialized();
 }
 
-FNavigationInvoker::FNavigationInvoker(AActor& InActor, float InGenerationRadius, float InRemovalRadius)
+FNavigationInvoker::FNavigationInvoker(AActor& InActor, float InGenerationRadius, float InRemovalRadius, const FNavAgentSelector& InSupportedAgents)
 : Actor(&InActor)
 , GenerationRadius(InGenerationRadius)
 , RemovalRadius(InRemovalRadius)
+, SupportedAgents(InSupportedAgents)
 {
-
+	SupportedAgents.MarkInitialized();
 }
 
 //----------------------------------------------------------------------//
@@ -4683,12 +4684,12 @@ void UNavigationSystemV1::ResetMaxSimultaneousTileGenerationJobsCount()
 // Active tiles
 //----------------------------------------------------------------------//
 
-void UNavigationSystemV1::RegisterNavigationInvoker(AActor& Invoker, float TileGenerationRadius, float TileRemovalRadius)
+void UNavigationSystemV1::RegisterNavigationInvoker(AActor& Invoker, float TileGenerationRadius, float TileRemovalRadius, const FNavAgentSelector& Agents)
 {
 	UNavigationSystemV1* NavSys = FNavigationSystem::GetCurrent<UNavigationSystemV1>(Invoker.GetWorld());
 	if (NavSys)
 	{
-		NavSys->RegisterInvoker(Invoker, TileGenerationRadius, TileRemovalRadius);
+		NavSys->RegisterInvoker(Invoker, TileGenerationRadius, TileRemovalRadius, Agents);
 	}
 }
 
@@ -4710,7 +4711,13 @@ void UNavigationSystemV1::SetGeometryGatheringMode(ENavDataGatheringModeConfig N
 	}
 }
 
+// Deprecated
 void UNavigationSystemV1::RegisterInvoker(AActor& Invoker, float TileGenerationRadius, float TileRemovalRadius)
+{
+	RegisterInvoker(Invoker, TileGenerationRadius, TileRemovalRadius, FNavAgentSelector());
+}
+
+void UNavigationSystemV1::RegisterInvoker(AActor& Invoker, float TileGenerationRadius, float TileRemovalRadius, const FNavAgentSelector& Agents)
 {
 	UE_CVLOG(bGenerateNavigationOnlyAroundNavigationInvokers == false, this, LogNavigation, Warning
 		, TEXT("Trying to register %s as enforcer, but NavigationSystem is not set up for enforcer-centric generation. See GenerateNavigationOnlyAroundNavigationInvokers in NavigationSystem's properties")
@@ -4723,6 +4730,8 @@ void UNavigationSystemV1::RegisterInvoker(AActor& Invoker, float TileGenerationR
 	Data.Actor = &Invoker;
 	Data.GenerationRadius = TileGenerationRadius;
 	Data.RemovalRadius = TileRemovalRadius;
+	Data.SupportedAgents = Agents;
+	Data.SupportedAgents.MarkInitialized();
 
 	UE_VLOG_CYLINDER(this, LogNavigation, Log, Invoker.GetActorLocation(), Invoker.GetActorLocation() + FVector(0, 0, 20), TileGenerationRadius, FColorList::LimeGreen
 		, TEXT("%s %.0f %.0f"), *Invoker.GetName(), TileGenerationRadius, TileRemovalRadius);
@@ -4761,7 +4770,7 @@ void UNavigationSystemV1::UpdateInvokers()
 #endif //WITH_EDITOR
 					)
 				{
-					InvokerLocations.Add(FNavigationInvokerRaw(Actor->GetActorLocation(), ItemIterator->Value.GenerationRadius, ItemIterator->Value.RemovalRadius));
+					InvokerLocations.Add(FNavigationInvokerRaw(Actor->GetActorLocation(), ItemIterator->Value.GenerationRadius, ItemIterator->Value.RemovalRadius, ItemIterator->Value.SupportedAgents));
 				}
 				else
 				{
@@ -4811,7 +4820,8 @@ void UNavigationSystemV1::RegisterNavigationInvoker(AActor* Invoker, float TileG
 {
 	if (Invoker != nullptr)
 	{
-		RegisterInvoker(*Invoker, TileGenerationRadius, TileRemovalRadius);
+		// The FNavAgentSelector class is not yet exposed in BP so we use the default value to specify that we want to generate the navmesh for all agents
+		RegisterInvoker(*Invoker, TileGenerationRadius, TileRemovalRadius, FNavAgentSelector());
 	}
 }
 
diff --git a/Engine/Source/Runtime/NavigationSystem/Public/NavigationInvokerComponent.h b/Engine/Source/Runtime/NavigationSystem/Public/NavigationInvokerComponent.h
--- a/Engine/Source/Runtime/NavigationSystem/Public/NavigationInvokerComponent.h
+++ b/Engine/Source/Runtime/NavigationSystem/Public/NavigationInvokerComponent.h
@@ -22,6 +22,10 @@ protected:
 	UPROPERTY(EditAnywhere, Category = Navigation, meta = (ClampMin = "0.1", UIMin = "0.1"))
 	float TileRemovalRadius;
 
+	/** restrict navigation generation to specific agents */
+	UPROPERTY(EditAnywhere, Category = Navigation)
+	FNavAgentSelector SupportedAgents;
+
 public:
 	UNavigationInvokerComponent(const FObjectInitializer& ObjectInitializer = FObjectInitializer::Get());
 
@@ -36,4 +40,5 @@ public:
 
 	virtual void Activate(bool bReset = false) override;
 	virtual void Deactivate() override;
+	virtual void PostInitProperties() override;
 };
diff --git a/Engine/Source/Runtime/NavigationSystem/Public/NavigationSystem.h b/Engine/Source/Runtime/NavigationSystem/Public/NavigationSystem.h
--- a/Engine/Source/Runtime/NavigationSystem/Public/NavigationSystem.h
+++ b/Engine/Source/Runtime/NavigationSystem/Public/NavigationSystem.h
@@ -746,10 +746,12 @@ public:
 	//----------------------------------------------------------------------//
 	// Active tiles
 	//----------------------------------------------------------------------//
+	UE_DEPRECATED(5.2, "This function is deprecated. Please use the new RegisterInvoker method with the Agents parameter (FNavAgentSelector() can be used as default value to keep the same behavior)")
 	virtual void RegisterInvoker(AActor& Invoker, float TileGenerationRadius, float TileRemovalRadius);
+	virtual void RegisterInvoker(AActor& Invoker, float TileGenerationRadius, float TileRemovalRadius, const FNavAgentSelector& Agents);
 	virtual void UnregisterInvoker(AActor& Invoker);
 
-	static void RegisterNavigationInvoker(AActor& Invoker, float TileGenerationRadius, float TileRemovalRadius);
+	static void RegisterNavigationInvoker(AActor& Invoker, float TileGenerationRadius, float TileRemovalRadius, const FNavAgentSelector& Agents = FNavAgentSelector());
 	static void UnregisterNavigationInvoker(AActor& Invoker);
 
 	const TArray<FNavigationInvokerRaw>& GetInvokerLocations() const { return InvokerLocations; }
diff --git a/Engine/Source/Runtime/NavigationSystem/Public/NavigationSystemTypes.h b/Engine/Source/Runtime/NavigationSystem/Public/NavigationSystemTypes.h
--- a/Engine/Source/Runtime/NavigationSystem/Public/NavigationSystemTypes.h
+++ b/Engine/Source/Runtime/NavigationSystem/Public/NavigationSystemTypes.h
@@ -97,9 +97,10 @@ struct FNavigationInvokerRaw
 	FVector Location;
 	float RadiusMin;
 	float RadiusMax;
+	FNavAgentSelector SupportedAgents;
 
-	FNavigationInvokerRaw(const FVector& InLocation, float Min, float Max)
-		: Location(InLocation), RadiusMin(Min), RadiusMax(Max)
+	FNavigationInvokerRaw(const FVector& InLocation, float Min, float Max, const FNavAgentSelector& InSupportedAgents)
+		: Location(InLocation), RadiusMin(Min), RadiusMax(Max), SupportedAgents(InSupportedAgents)
 	{}
 };
 
@@ -112,9 +113,11 @@ struct FNavigationInvoker
 	/** tiles over RemovalRadius will get removed.
 	*	@Note needs to be >= GenerationRadius or will get clampped */
 	float RemovalRadius;
+	/** restrict navigation generation to specific agents */
+	FNavAgentSelector SupportedAgents;
 
 	FNavigationInvoker();
-	FNavigationInvoker(AActor& InActor, float InGenerationRadius, float InRemovalRadius);
+	FNavigationInvoker(AActor& InActor, float InGenerationRadius, float InRemovalRadius, const FNavAgentSelector& InSupportedAgents);
 };
 
 namespace NavigationHelper
