Change navigation geometry export warning to be based on triangle count and avoid repeating the warning.

diff --git a/Engine/Config/BaseEngine.ini b/Engine/Config/BaseEngine.ini
--- a/Engine/Config/BaseEngine.ini
+++ b/Engine/Config/BaseEngine.ini
@@ -1110,6 +1110,7 @@ DefaultWorldPartitionHLODLayer=/Engine/Maps/Templates/HLODs/HLODLayer_Instanced.
 +PropertyRedirects=(OldName="KismetMathLibrary.Conv_DoubleToVector.InFloat", NewName="Conv_DoubleToVector.InDouble")
 
 +PropertyRedirects=(OldName="KismetSystemLibrary.MakeTopLevelAssetPath.FullPathOrPackageName", NewName="KismetSystemLibrary.MakeTopLevelAssetPath.PackageName")
++PropertyRedirects=(OldName="/Script/NavigationSystem.NavigationSystemV1.GeometryExportVertexCountWarningThreshold",NewName="/Script/NavigationSystem.NavigationSystemV1.GeometryExportTriangleCountWarningThreshold")
 
 [CoreUObject.Metadata]
 ; Note: UnrealHeaderTool should be rerun after making changes to MetadataRedirects to catch any keys specified in class headers
diff --git a/Engine/Source/Runtime/NavigationSystem/Private/NavMesh/RecastNavMeshGenerator.cpp b/Engine/Source/Runtime/NavigationSystem/Private/NavMesh/RecastNavMeshGenerator.cpp
--- a/Engine/Source/Runtime/NavigationSystem/Private/NavMesh/RecastNavMeshGenerator.cpp
+++ b/Engine/Source/Runtime/NavigationSystem/Private/NavMesh/RecastNavMeshGenerator.cpp
@@ -1027,11 +1027,21 @@ FORCEINLINE_DEBUGGABLE void ValidateGeometryExport(const FRecastGeometryExport&
 	{
 		if (const UNavigationSystemV1* NavSys = FNavigationSystem::GetCurrent<UNavigationSystemV1>(Owner->GetWorld()))
 		{
-			UE_CLOG(NavSys->GeometryExportVertexCountWarningThreshold > 0 && GeomExport.VertexBuffer.Num() > NavSys->GeometryExportVertexCountWarningThreshold,
-				LogNavigation,
-				Warning, 
-				TEXT("Exporting geometry with too many verticies (%i). This might cause performance and memory issues. Simplify collision or change GeometryExportVertexCountWarningThreshold. See '%s'."), 
-				GeomExport.VertexBuffer.Num(), *GetFullNameSafe(Owner));
+			constexpr int32 CoordinatePerTriangle = 9;
+			if (NavSys->GeometryExportTriangleCountWarningThreshold > 0 && (GeomExport.VertexBuffer.Num() / CoordinatePerTriangle) > NavSys->GeometryExportTriangleCountWarningThreshold)
+			{
+				static uint32 LastNameHash = 0;
+				const FString FullName = GetFullNameSafe(Owner);
+				const uint32 NameHash = GetTypeHash(FullName);
+				if (NameHash != LastNameHash)
+				{
+					UE_LOG(LogNavigation,
+						Warning, 
+						TEXT("Exporting collision geometry with too many triangles (%i). This might cause performance and memory issues. Add a simple collision or change GeometryExportVertexCountWarningThreshold. See '%s'."), 
+						GeomExport.VertexBuffer.Num() / CoordinatePerTriangle, *FullName);
+				}
+				LastNameHash = NameHash;
+			}
 		}
 	}
 }
diff --git a/Engine/Source/Runtime/NavigationSystem/Public/NavigationSystem.h b/Engine/Source/Runtime/NavigationSystem/Public/NavigationSystem.h
--- a/Engine/Source/Runtime/NavigationSystem/Public/NavigationSystem.h
+++ b/Engine/Source/Runtime/NavigationSystem/Public/NavigationSystem.h
@@ -317,10 +317,17 @@ public:
 	UPROPERTY(config, EditAnywhere, Category=NavigationSystem)
 	uint32 bSkipAgentHeightCheckWhenPickingNavData:1;
 
+#if WITH_EDITOR
 	/** Warnings are logged if exporting the navigation collision for an object exceed this vertex count.
 	 * Use -1 to disable. */
-	UPROPERTY(config, EditAnywhere, AdvancedDisplay, Category = NavigationSystem)
+	UE_DEPRECATED(5.2, "This property is deprecated. Please use GeometryExportTriangleCountWarningThreshold instead.")
 	int32 GeometryExportVertexCountWarningThreshold = 1000000;
+#endif // WITH_EDITOR
+
+	/** Warnings are logged if exporting the navigation collision for an object exceed this triangle count.
+	 * Use -1 to disable. */
+	UPROPERTY(config, EditAnywhere, AdvancedDisplay, Category = NavigationSystem)
+	int32 GeometryExportTriangleCountWarningThreshold = 200000;
 	
 protected:
 	/** If set to true navigation will be generated only around registered "navigation enforcers"
