Changing "Navigation data generation finished" log to verbose to avoid having too many logs in the editor when moving static mesh in automatic navigation update mode.

diff --git a/Engine/Source/Runtime/NavigationSystem/Private/NavigationSystem.cpp b/Engine/Source/Runtime/NavigationSystem/Private/NavigationSystem.cpp
--- a/Engine/Source/Runtime/NavigationSystem/Private/NavigationSystem.cpp
+++ b/Engine/Source/Runtime/NavigationSystem/Private/NavigationSystem.cpp
@@ -3950,7 +3950,7 @@ void UNavigationSystemV1::OnNavigationGenerationFinished(ANavigationData& NavDat
 #if WITH_EDITOR
 	if (GetWorld()->IsGameWorld() == false)
 	{
-		UE_LOG(LogNavigationDataBuild, Log, TEXT("Navigation data generation finished for %s (%s)."), *NavData.GetActorLabel(), *NavData.GetFullName());
+		UE_LOG(LogNavigationDataBuild, Verbose, TEXT("Navigation data generation finished for %s (%s)."), *NavData.GetActorLabel(), *NavData.GetFullName());
 	}
 
 	// Reset bIsBuildingOnLoad
