Ensure that no reload is needed for NavArea Registration to take effect

diff --git a/Engine/Source/Runtime/Engine/Classes/AI/NavigationSystemBase.h b/Engine/Source/Runtime/Engine/Classes/AI/NavigationSystemBase.h
--- a/Engine/Source/Runtime/Engine/Classes/AI/NavigationSystemBase.h
+++ b/Engine/Source/Runtime/Engine/Classes/AI/NavigationSystemBase.h
@@ -213,6 +213,7 @@ namespace FNavigationSystem
 	DECLARE_DELEGATE_RetVal(TSubclassOf<AActor>, FNavDataClassFetchSignature);
 	DECLARE_DELEGATE_TwoParams(FWorldBoolBasedSignature, UWorld& /*World*/, const bool /*bShow*/);
 	DECLARE_MULTICAST_DELEGATE_OneParam(FOnNavigationInitSignature, const UNavigationSystemBase&);
+	DECLARE_MULTICAST_DELEGATE_TwoParams(FOnNavAreaGenericEvent, const UWorld&, const UClass*);
 }
 
 
@@ -266,6 +267,8 @@ public:
 
 	static FNavigationSystem::FOnNavigationInitSignature& OnNavigationInitStartStaticDelegate();
 	static FNavigationSystem::FOnNavigationInitSignature& OnNavigationInitDoneStaticDelegate();
+	static FNavigationSystem::FOnNavAreaGenericEvent& OnNavAreaRegisteredDelegate();
+	static FNavigationSystem::FOnNavAreaGenericEvent& OnNavAreaUnregisteredDelegate();
 
 protected:
 	/**	Sets the Transform the Navigation System will use when converting from FromCoordType
diff --git a/Engine/Source/Runtime/Engine/Private/AI/NavigationSystemBase.cpp b/Engine/Source/Runtime/Engine/Private/AI/NavigationSystemBase.cpp
--- a/Engine/Source/Runtime/Engine/Private/AI/NavigationSystemBase.cpp
+++ b/Engine/Source/Runtime/Engine/Private/AI/NavigationSystemBase.cpp
@@ -151,6 +151,8 @@ namespace FNavigationSystem
 		FWorldBasedSignature Build;
 		FOnNavigationInitSignature OnNavigationInitStart;
 		FOnNavigationInitSignature OnNavigationInitDone;
+		FOnNavAreaGenericEvent OnNavAreaRegistered;
+		FOnNavAreaGenericEvent OnNavAreaUnregistered;
 #if WITH_EDITOR
 		FWorldBasedSignature OnPIEStart;
 		FWorldBasedSignature OnPIEEnd;
@@ -427,6 +429,8 @@ FNavigationSystem::FWorldBoolBasedSignature& UNavigationSystemBase::VerifyNaviga
 FNavigationSystem::FWorldBasedSignature& UNavigationSystemBase::BuildDelegate() { return FNavigationSystem::Delegates.Build; }
 FNavigationSystem::FOnNavigationInitSignature& UNavigationSystemBase::OnNavigationInitStartStaticDelegate() { return FNavigationSystem::Delegates.OnNavigationInitStart; }
 FNavigationSystem::FOnNavigationInitSignature& UNavigationSystemBase::OnNavigationInitDoneStaticDelegate() { return FNavigationSystem::Delegates.OnNavigationInitDone; }
+FNavigationSystem::FOnNavAreaGenericEvent& UNavigationSystemBase::OnNavAreaRegisteredDelegate() { return FNavigationSystem::Delegates.OnNavAreaRegistered; }
+FNavigationSystem::FOnNavAreaGenericEvent& UNavigationSystemBase::OnNavAreaUnregisteredDelegate() { return FNavigationSystem::Delegates.OnNavAreaUnregistered; }
 #if WITH_EDITOR
 FNavigationSystem::FWorldBasedSignature& UNavigationSystemBase::OnPIEStartDelegate() { return FNavigationSystem::Delegates.OnPIEStart; }
 FNavigationSystem::FWorldBasedSignature& UNavigationSystemBase::OnPIEEndDelegate() { return FNavigationSystem::Delegates.OnPIEEnd; }
diff --git a/Engine/Source/Runtime/NavigationSystem/Private/NavAreas/NavArea.cpp b/Engine/Source/Runtime/NavigationSystem/Private/NavAreas/NavArea.cpp
--- a/Engine/Source/Runtime/NavigationSystem/Private/NavAreas/NavArea.cpp
+++ b/Engine/Source/Runtime/NavigationSystem/Private/NavAreas/NavArea.cpp
@@ -6,7 +6,6 @@
 
 #include UE_INLINE_GENERATED_CPP_BY_NAME(NavArea)
 
-
 UNavArea::UNavArea(const FObjectInitializer& ObjectInitializer) 
 	: Super(ObjectInitializer)
 {
@@ -42,14 +41,6 @@ void UNavArea::PostInitProperties()
 
 void UNavArea::RegisterArea()
 {
-	if (HasAnyFlags(RF_ClassDefaultObject) && 
-		!HasAnyFlags(RF_NeedInitialization)  // Don't register BP Area that has still not finished loaded their properties, it was also try again to register later via UNavArea::PostLoad()
-		&& !IsReloadActive()
-		)
-	{
-		UNavigationSystemV1::RequestAreaRegistering(GetClass());
-	}
-
 	if (!SupportedAgents.IsInitialized())
 	{
 		SupportedAgents.bSupportsAgent0 = bSupportsAgent0;
@@ -70,6 +61,14 @@ void UNavArea::RegisterArea()
 		SupportedAgents.bSupportsAgent15 = bSupportsAgent15;
 		SupportedAgents.MarkInitialized();
 	}
+
+	if (HasAnyFlags(RF_ClassDefaultObject) && 
+		!HasAnyFlags(RF_NeedInitialization)  // Don't register BP Area that has still not finished loaded their properties, it was also try again to register later via UNavArea::PostLoad()
+		&& !IsReloadActive()
+		)
+	{
+		UNavigationSystemV1::RequestAreaRegistering(GetClass());
+	}
 }
 
 void UNavArea::Serialize(FArchive& Ar)
@@ -107,14 +106,18 @@ void UNavArea::PostEditChangeProperty(FPropertyChangedEvent& PropertyChangedEven
 {
 	static const FName NAME_DefaultCost = GET_MEMBER_NAME_CHECKED(UNavArea, DefaultCost);
 	static const FName NAME_FixedAreaEnteringCost = GET_MEMBER_NAME_CHECKED(UNavArea, FixedAreaEnteringCost);
+	static const FName NAME_SupportedAgents = GET_MEMBER_NAME_CHECKED(UNavArea, SupportedAgents);
 
 	Super::PostEditChangeProperty(PropertyChangedEvent);
 
 	if (HasAnyFlags(RF_ClassDefaultObject) && !IsReloadActive())
 	{
 		const FName PropertyName = (PropertyChangedEvent.Property != nullptr) ? PropertyChangedEvent.Property->GetFName() : NAME_None;
+		const FName MemberPropertyName = (PropertyChangedEvent.MemberProperty != nullptr) ? PropertyChangedEvent.MemberProperty->GetFName() : NAME_None;
+
 		if (PropertyName == NAME_DefaultCost
-			|| PropertyName == NAME_FixedAreaEnteringCost)
+			|| PropertyName == NAME_FixedAreaEnteringCost
+			|| MemberPropertyName == NAME_SupportedAgents)
 		{
 			UNavigationSystemV1::RequestAreaUnregistering(GetClass());
 			RegisterArea();
diff --git a/Engine/Source/Runtime/NavigationSystem/Private/NavMesh/RecastNavMesh.cpp b/Engine/Source/Runtime/NavigationSystem/Private/NavMesh/RecastNavMesh.cpp
--- a/Engine/Source/Runtime/NavigationSystem/Private/NavMesh/RecastNavMesh.cpp
+++ b/Engine/Source/Runtime/NavigationSystem/Private/NavMesh/RecastNavMesh.cpp
@@ -835,6 +835,20 @@ void ARecastNavMesh::OnNavAreaAdded(const UClass* NavAreaClass, int32 AgentIndex
 	}
 }
 
+void ARecastNavMesh::OnNavAreaRemoved(const UClass* NavAreaClass)
+{
+	// In an ideal world we'd reset the DefaultQueryFilter Costs here for the AreaID but its
+	// not really worth changing the API as we shouldn't be using these values anyway.
+
+	FRecastNavMeshGenerator* MyGenerator = static_cast<FRecastNavMeshGenerator*>(GetGenerator());
+	if (MyGenerator)
+	{
+		MyGenerator->OnAreaRemoved(NavAreaClass);
+	}
+
+	Super::OnNavAreaRemoved(NavAreaClass);
+}
+
 void ARecastNavMesh::OnNavAreaChanged()
 {
 	if (RecastNavMeshImpl)
@@ -3628,6 +3642,19 @@ void FRecastNavMeshCachedData::OnAreaAdded(const UClass* AreaClass, int32 AreaID
 	}		
 }
 
+void FRecastNavMeshCachedData::OnAreaRemoved(const UClass* AreaClass)
+{
+	const int32* AreaID = AreaClass ? AreaClassToIdMap.Find(AreaClass) : nullptr;
+
+	if (AreaID != nullptr)
+	{
+		AreaClassToIdMap.Remove(AreaClass);
+
+		FlagsPerArea[*AreaID] = 0;
+		FlagsPerOffMeshLinkArea[*AreaID] = 0;
+	}
+}
+
 uint32 ARecastNavMesh::GetLinkUserId(NavNodeRef LinkPolyID) const
 {
 	return RecastNavMeshImpl ? RecastNavMeshImpl->GetLinkUserId(LinkPolyID) : 0;
diff --git a/Engine/Source/Runtime/NavigationSystem/Private/NavMesh/RecastNavMeshGenerator.cpp b/Engine/Source/Runtime/NavigationSystem/Private/NavMesh/RecastNavMeshGenerator.cpp
--- a/Engine/Source/Runtime/NavigationSystem/Private/NavMesh/RecastNavMeshGenerator.cpp
+++ b/Engine/Source/Runtime/NavigationSystem/Private/NavMesh/RecastNavMeshGenerator.cpp
@@ -5121,6 +5121,11 @@ void FRecastNavMeshGenerator::OnAreaAdded(const UClass* AreaClass, int32 AreaID)
 	AdditionalCachedData.OnAreaAdded(AreaClass, AreaID);
 }
 
+void FRecastNavMeshGenerator::OnAreaRemoved(const UClass* AreaClass)
+{
+	AdditionalCachedData.OnAreaRemoved(AreaClass);
+}
+
 int32 FRecastNavMeshGenerator::FindInclusionBoundEncapsulatingBox(const FBox& Box) const
 {
 	for (int32 Index = 0; Index < InclusionBounds.Num(); ++Index)
diff --git a/Engine/Source/Runtime/NavigationSystem/Private/NavModifierComponent.cpp b/Engine/Source/Runtime/NavigationSystem/Private/NavModifierComponent.cpp
--- a/Engine/Source/Runtime/NavigationSystem/Private/NavModifierComponent.cpp
+++ b/Engine/Source/Runtime/NavigationSystem/Private/NavModifierComponent.cpp
@@ -6,6 +6,7 @@
 #include "AI/NavigationModifier.h"
 #include "NavAreas/NavArea_Null.h"
 #include "PhysicsEngine/BodySetup.h"
+#include "NavigationSystem.h"
 
 #include UE_INLINE_GENERATED_CPP_BY_NAME(NavModifierComponent)
 
@@ -16,6 +17,59 @@ UNavModifierComponent::UNavModifierComponent(const FObjectInitializer& ObjectIni
 	bIncludeAgentHeight = true;
 }
 
+void UNavModifierComponent::OnRegister()
+{
+	Super::OnRegister();
+
+#if WITH_EDITOR
+	if (GIsEditor && !HasAnyFlags(RF_ClassDefaultObject))
+	{
+		OnNavAreaRegisteredDelegateHandle = UNavigationSystemBase::OnNavAreaRegisteredDelegate().AddUObject(this, &UNavModifierComponent::OnNavAreaRegistered);
+		OnNavAreaUnregisteredDelegateHandle = UNavigationSystemBase::OnNavAreaUnregisteredDelegate().AddUObject(this, &UNavModifierComponent::OnNavAreaUnregistered);
+	}
+#endif // WITH_EDITOR 
+}
+
+void UNavModifierComponent::OnUnregister()
+{
+#if WITH_EDITOR
+	if (GIsEditor && !HasAnyFlags(RF_ClassDefaultObject))
+	{
+		UNavigationSystemBase::OnNavAreaRegisteredDelegate().Remove(OnNavAreaRegisteredDelegateHandle);
+		UNavigationSystemBase::OnNavAreaUnregisteredDelegate().Remove(OnNavAreaUnregisteredDelegateHandle);
+	}
+#endif // WITH_EDITOR 
+
+	Super::OnUnregister();
+}
+
+namespace UE::Navigation::Private
+{
+	void OnNavAreaRegistrationChanged(UNavModifierComponent& ModifierComponent, const UWorld& World, const UClass* NavAreaClass)
+	{
+#if WITH_EDITOR
+		if (NavAreaClass && NavAreaClass == ModifierComponent.AreaClass && &World == ModifierComponent.GetWorld())
+		{
+			// Reregister the component, this follows the same pattern in editor to editing the values of components member variables,
+			// actor that owns the component will have it reregistered.
+			ModifierComponent.ReregisterComponent();
+		}
+#endif // WITH_EDITOR 
+	}
+}
+
+// This function is only called if GIsEditor == true for non default objects components that are registered.
+void UNavModifierComponent::OnNavAreaRegistered(const UWorld& World, const UClass* NavAreaClass)
+{
+	UE::Navigation::Private::OnNavAreaRegistrationChanged(*this, World, NavAreaClass);
+}
+
+// This function is only called if GIsEditor == true for non default objects components that are registered.
+void UNavModifierComponent::OnNavAreaUnregistered(const UWorld& World, const UClass* NavAreaClass)
+{
+	UE::Navigation::Private::OnNavAreaRegistrationChanged(*this, World, NavAreaClass);
+}
+
 void UNavModifierComponent::CalcAndCacheBounds() const
 {
 	AActor* MyOwner = GetOwner();
diff --git a/Engine/Source/Runtime/NavigationSystem/Private/NavModifierVolume.cpp b/Engine/Source/Runtime/NavigationSystem/Private/NavModifierVolume.cpp
--- a/Engine/Source/Runtime/NavigationSystem/Private/NavModifierVolume.cpp
+++ b/Engine/Source/Runtime/NavigationSystem/Private/NavModifierVolume.cpp
@@ -28,6 +28,54 @@ ANavModifierVolume::ANavModifierVolume(const FObjectInitializer& ObjectInitializ
 	}
 }
 
+void ANavModifierVolume::PostInitProperties()
+{
+	Super::PostInitProperties();
+
+#if WITH_EDITOR
+	if (GIsEditor && !HasAnyFlags(RF_ClassDefaultObject))
+	{
+		OnNavAreaRegisteredDelegateHandle = UNavigationSystemBase::OnNavAreaRegisteredDelegate().AddUObject(this, &ANavModifierVolume::OnNavAreaRegistered);
+		OnNavAreaUnregisteredDelegateHandle = UNavigationSystemBase::OnNavAreaUnregisteredDelegate().AddUObject(this, &ANavModifierVolume::OnNavAreaUnregistered);
+	}
+#endif // WITH_EDITOR
+}
+
+void ANavModifierVolume::BeginDestroy()
+{
+	Super::BeginDestroy();
+
+#if WITH_EDITOR
+	if (GIsEditor && !HasAnyFlags(RF_ClassDefaultObject))
+	{
+		UNavigationSystemBase::OnNavAreaRegisteredDelegate().Remove(OnNavAreaRegisteredDelegateHandle);
+		UNavigationSystemBase::OnNavAreaUnregisteredDelegate().Remove(OnNavAreaUnregisteredDelegateHandle);
+	}
+#endif // WITH_EDITOR
+}
+
+// This function is only called if GIsEditor == true
+void ANavModifierVolume::OnNavAreaRegistered(const UWorld& World, const UClass* NavAreaClass)
+{
+#if WITH_EDITOR
+	if (NavAreaClass && NavAreaClass == AreaClass && &World == GetWorld())
+	{
+		FNavigationSystem::UpdateActorData(*this);
+	}
+#endif // WITH_EDITOR
+}
+
+// This function is only called if GIsEditor == true
+void ANavModifierVolume::OnNavAreaUnregistered(const UWorld& World, const UClass* NavAreaClass)
+{
+#if WITH_EDITOR
+	if (NavAreaClass && NavAreaClass == AreaClass && &World == GetWorld())
+	{
+		FNavigationSystem::UpdateActorData(*this);
+	}
+#endif // WITH_EDITOR
+}
+
 void ANavModifierVolume::GetNavigationData(FNavigationRelevantData& Data) const
 {
 	if (Brush && AreaClass && AreaClass != FNavigationSystem::GetDefaultWalkableArea())
diff --git a/Engine/Source/Runtime/NavigationSystem/Private/NavigationSystem.cpp b/Engine/Source/Runtime/NavigationSystem/Private/NavigationSystem.cpp
--- a/Engine/Source/Runtime/NavigationSystem/Private/NavigationSystem.cpp
+++ b/Engine/Source/Runtime/NavigationSystem/Private/NavigationSystem.cpp
@@ -2607,6 +2607,12 @@ void UNavigationSystemV1::UnregisterNavAreaClass(UClass* NavAreaClass)
 		// notify navigation data
 		// notify existing nav data
 		OnNavigationAreaEvent(NavAreaClass, ENavAreaEvent::Unregistered);
+
+		const UWorld* const World = GetWorld();
+		if (ensure(World))
+		{
+			UNavigationSystemBase::OnNavAreaUnregisteredDelegate().Broadcast(*World, NavAreaClass);
+		}
 	}
 }
 
@@ -2675,6 +2681,12 @@ void UNavigationSystemV1::RegisterNavAreaClass(UClass* AreaClass)
 		SwitchByAgentCDO->UpdateAgentConfig();
 	}
 #endif
+
+	const UWorld* const World = GetWorld();
+	if (ensure(World))
+	{
+		UNavigationSystemBase::OnNavAreaRegisteredDelegate().Broadcast(*World, AreaClass);
+	}
 }
 
 void UNavigationSystemV1::OnNavigationAreaEvent(UClass* AreaClass, ENavAreaEvent::Type Event)
diff --git a/Engine/Source/Runtime/NavigationSystem/Public/NavMesh/RecastNavMesh.h b/Engine/Source/Runtime/NavigationSystem/Public/NavMesh/RecastNavMesh.h
--- a/Engine/Source/Runtime/NavigationSystem/Public/NavMesh/RecastNavMesh.h
+++ b/Engine/Source/Runtime/NavigationSystem/Public/NavMesh/RecastNavMesh.h
@@ -1164,6 +1164,7 @@ public:
 	//----------------------------------------------------------------------//
 	virtual void OnNavAreaChanged() override;
 	virtual void OnNavAreaAdded(const UClass* NavAreaClass, int32 AgentIndex) override;
+	virtual void OnNavAreaRemoved(const UClass* NavAreaClass) override;
 	virtual int32 GetNewAreaID(const UClass* AreaClass) const override;
 	virtual int32 GetMaxSupportedAreas() const override { return RECAST_MAX_AREAS; }
 
@@ -1468,6 +1469,7 @@ struct FRecastNavMeshCachedData
 
 	static FRecastNavMeshCachedData Construct(const ARecastNavMesh* RecastNavMeshActor);
 	void OnAreaAdded(const UClass* AreaClass, int32 AreaID);
+	void OnAreaRemoved(const UClass* AreaClass);
 };
 
 #endif // WITH_RECAST
diff --git a/Engine/Source/Runtime/NavigationSystem/Public/NavMesh/RecastNavMeshGenerator.h b/Engine/Source/Runtime/NavigationSystem/Public/NavMesh/RecastNavMeshGenerator.h
--- a/Engine/Source/Runtime/NavigationSystem/Public/NavMesh/RecastNavMeshGenerator.h
+++ b/Engine/Source/Runtime/NavigationSystem/Public/NavMesh/RecastNavMeshGenerator.h
@@ -695,6 +695,7 @@ public:
 
 	/** update area data */
 	void OnAreaAdded(const UClass* AreaClass, int32 AreaID);
+	void OnAreaRemoved(const UClass* AreaClass);
 		
 	//--- accessors --- //
 	FORCEINLINE class UWorld* GetWorld() const { return DestNavMesh->GetWorld(); }
diff --git a/Engine/Source/Runtime/NavigationSystem/Public/NavModifierComponent.h b/Engine/Source/Runtime/NavigationSystem/Public/NavModifierComponent.h
--- a/Engine/Source/Runtime/NavigationSystem/Public/NavModifierComponent.h
+++ b/Engine/Source/Runtime/NavigationSystem/Public/NavModifierComponent.h
@@ -36,6 +36,13 @@ class NAVIGATIONSYSTEM_API UNavModifierComponent : public UNavRelevantComponent
 
 protected:
 	void OnTransformUpdated(USceneComponent* RootComponent, EUpdateTransformFlags UpdateTransformFlags, ETeleportType Teleport);
+	void OnNavAreaRegistered(const UWorld& World, const UClass* NavAreaClass);
+	void OnNavAreaUnregistered(const UWorld& World, const UClass* NavAreaClass);
+
+	//~ Begin UActorComponent Interface
+	virtual void OnRegister() override;
+	virtual void OnUnregister() override;
+	//~ End UActorComponent Interface
 
 	struct FRotatedBox
 	{
@@ -51,4 +58,7 @@ protected:
 	/** cached in CalcAndCacheBounds and tested in GetNavigationData to see if
 	 *	cached data is still valid */
 	mutable FTransform CachedTransform;
+
+	FDelegateHandle OnNavAreaRegisteredDelegateHandle;
+	FDelegateHandle OnNavAreaUnregisteredDelegateHandle;
 };
diff --git a/Engine/Source/Runtime/NavigationSystem/Public/NavModifierVolume.h b/Engine/Source/Runtime/NavigationSystem/Public/NavModifierVolume.h
--- a/Engine/Source/Runtime/NavigationSystem/Public/NavModifierVolume.h
+++ b/Engine/Source/Runtime/NavigationSystem/Public/NavModifierVolume.h
@@ -34,6 +34,9 @@ protected:
 	UPROPERTY(EditAnywhere, Category = Default, AdvancedDisplay)
 	ENavigationDataResolution NavMeshResolution;
 
+	FDelegateHandle OnNavAreaRegisteredDelegateHandle;
+	FDelegateHandle OnNavAreaUnregisteredDelegateHandle;
+
 public:
 	ANavModifierVolume(const FObjectInitializer& ObjectInitializer);
 
@@ -50,4 +53,11 @@ public:
 	virtual void PostEditUndo() override;
 	virtual void PostEditChangeProperty(FPropertyChangedEvent& PropertyChangedEvent) override;
 #endif
+
+protected:
+	virtual void PostInitProperties() override;
+	virtual void BeginDestroy() override;
+
+	void OnNavAreaRegistered(const UWorld& World, const UClass* NavAreaClass);
+	void OnNavAreaUnregistered(const UWorld& World, const UClass* NavAreaClass);
 };
