Addition of a display option to display navmesh tile resolution
Offset navigation 2D debug text display on screen

diff --git a/Engine/Source/Runtime/NavigationSystem/Private/NavMesh/NavMeshRenderingComponent.cpp b/Engine/Source/Runtime/NavigationSystem/Private/NavMesh/NavMeshRenderingComponent.cpp
--- a/Engine/Source/Runtime/NavigationSystem/Private/NavMesh/NavMeshRenderingComponent.cpp
+++ b/Engine/Source/Runtime/NavigationSystem/Private/NavMesh/NavMeshRenderingComponent.cpp
@@ -40,6 +40,7 @@ static constexpr FColor NavMeshRenderColor_OffMeshConnectionInvalid(64, 64, 64);
 static const FColor NavMeshRenderColor_PolyForbidden(FColorList::Black);
 
 static constexpr float DefaultEdges_LineThickness = 0.0f;
+static constexpr float TileResolution_LineThickness = 5.f;
 static constexpr float PolyEdges_LineThickness = 1.5f;
 static constexpr float NavMeshEdges_LineThickness = 3.5f;
 static constexpr float LinkLines_LineThickness = 2.0f;
@@ -284,6 +285,7 @@ namespace FNavMeshRenderingHelpers
 			(NavMesh->bDrawFilledPolys ? (1 << static_cast<int32>(ENavMeshDetailFlags::FilledPolys)) : 0) |
 			(NavMesh->bDrawNavMeshEdges ? (1 << static_cast<int32>(ENavMeshDetailFlags::BoundaryEdges)) : 0) |
 			(NavMesh->bDrawTileBounds ? (1 << static_cast<int32>(ENavMeshDetailFlags::TileBounds)) : 0) |
+			(NavMesh->bDrawTileResolutions ? (1 << static_cast<int32>(ENavMeshDetailFlags::TileResolutions)) : 0) |
 			(NavMesh->bDrawPathCollidingGeometry ? (1 << static_cast<int32>(ENavMeshDetailFlags::PathCollidingGeometry)) : 0) |
 			(NavMesh->bDrawTileLabels ? (1 << static_cast<int32>(ENavMeshDetailFlags::TileLabels)) : 0) |
 			(NavMesh->bDrawPolygonLabels ? (1 << static_cast<int32>(ENavMeshDetailFlags::PolygonLabels)) : 0) |
@@ -778,11 +780,12 @@ void FNavMeshSceneProxyData::GatherData(const ARecastNavMesh* NavMesh, int32 InN
 
 		const bool bGatherTileLabels = FNavMeshRenderingHelpers::HasFlag(NavDetailFlags, ENavMeshDetailFlags::TileLabels);
 		const bool bGatherTileBounds = FNavMeshRenderingHelpers::HasFlag(NavDetailFlags, ENavMeshDetailFlags::TileBounds);
+		const bool bGatherTileResolutions = FNavMeshRenderingHelpers::HasFlag(NavDetailFlags, ENavMeshDetailFlags::TileResolutions);
 		const bool bGatherPolygonLabels = FNavMeshRenderingHelpers::HasFlag(NavDetailFlags, ENavMeshDetailFlags::PolygonLabels);
 		const bool bGatherPolygonCost = FNavMeshRenderingHelpers::HasFlag(NavDetailFlags, ENavMeshDetailFlags::PolygonCost);
 		const bool bGatherPolygonFlags = FNavMeshRenderingHelpers::HasFlag(NavDetailFlags, ENavMeshDetailFlags::PolygonFlags);
 
-		if (bGatherTileLabels || bGatherTileBounds || bGatherPolygonLabels || bGatherPolygonCost || bGatherPolygonFlags || bGatherTileBuildTimes)
+		if (bGatherTileLabels || bGatherTileBounds || bGatherTileResolutions || bGatherPolygonLabels || bGatherPolygonCost || bGatherPolygonFlags || bGatherTileBuildTimes)
 		{
 			QUICK_SCOPE_CYCLE_COUNTER(STAT_NavMesh_GatherDebugDrawing_TileIterations);
 			
@@ -912,6 +915,47 @@ void FNavMeshSceneProxyData::GatherData(const ARecastNavMesh* NavMesh, int32 InN
 						ThickLineItems.Add(FDebugRenderSceneProxy::FDebugLine(UR, LR, NavMeshRenderColor_TileBounds, DefaultEdges_LineThickness));
 						ThickLineItems.Add(FDebugRenderSceneProxy::FDebugLine(LR, LL, NavMeshRenderColor_TileBounds, DefaultEdges_LineThickness));
 					}
+
+					if (bGatherTileResolutions)
+					{
+						const FBox TileBox = NavMesh->GetNavMeshTileBounds(TileIndex);
+						const FVector::FReal DrawZ = TileBox.Max.Z + NavMeshDrawOffset.Z;
+						constexpr FVector::FReal InsideOffset = 10.f;
+						const FVector LowerLeft(TileBox.Min.X + InsideOffset, TileBox.Min.Y + InsideOffset, DrawZ);
+						const FVector UpperRight(TileBox.Max.X - InsideOffset, TileBox.Max.Y - InsideOffset, DrawZ);
+						const FVector UpperLeft(LowerLeft.X, UpperRight.Y, DrawZ);
+						const FVector LowerRight(UpperRight.X, LowerLeft.Y, DrawZ);
+
+						FColor TileBoundsColor = FColor::Silver;
+						ENavigationDataResolution Resolution = ENavigationDataResolution::Invalid;
+						
+						if (NavMesh->GetNavmeshTileResolution(TileIndex, Resolution))
+						{
+							switch (Resolution)
+							{
+							case ENavigationDataResolution::Low:
+								TileBoundsColor = FColor::Blue;
+								break;
+							
+							case ENavigationDataResolution::Default:
+								TileBoundsColor = FColor::Green;
+								break;
+
+							case ENavigationDataResolution::High:
+								TileBoundsColor = FColor::Orange;
+								break;
+							
+							default:
+								// Unset
+								break;
+							}
+						}
+						
+						ThickLineItems.Add(FDebugRenderSceneProxy::FDebugLine(LowerLeft, UpperLeft, TileBoundsColor, TileResolution_LineThickness));
+						ThickLineItems.Add(FDebugRenderSceneProxy::FDebugLine(UpperLeft, UpperRight, TileBoundsColor, TileResolution_LineThickness));
+						ThickLineItems.Add(FDebugRenderSceneProxy::FDebugLine(UpperRight, LowerRight, TileBoundsColor, TileResolution_LineThickness));
+						ThickLineItems.Add(FDebugRenderSceneProxy::FDebugLine(LowerRight, LowerLeft, TileBoundsColor, TileResolution_LineThickness));
+					}
 				}
 			}
 
@@ -1576,7 +1620,7 @@ void FNavMeshDebugDrawDelegateHelper::DrawDebugLabels(UCanvas* Canvas, APlayerCo
 	const FSceneView* View = Canvas->SceneView;
 	const UFont* Font = GEngine->GetSmallFont();
 	const FNavMeshSceneProxyData::FDebugText* DebugText = DebugLabels.GetData();
-	float ScreenY = 40.f;
+	float ScreenY = 70.f;
 	for (int32 Idx = 0; Idx < DebugLabels.Num(); ++Idx, ++DebugText)
 	{
 		if (DebugText->Location == FNavigationSystem::InvalidLocation)
diff --git a/Engine/Source/Runtime/NavigationSystem/Private/NavMesh/RecastNavMesh.cpp b/Engine/Source/Runtime/NavigationSystem/Private/NavMesh/RecastNavMesh.cpp
--- a/Engine/Source/Runtime/NavigationSystem/Private/NavMesh/RecastNavMesh.cpp
+++ b/Engine/Source/Runtime/NavigationSystem/Private/NavMesh/RecastNavMesh.cpp
@@ -1162,6 +1162,26 @@ bool ARecastNavMesh::GetNavMeshTileXY(const FVector& Point, int32& OutX, int32&
 	return RecastNavMeshImpl && RecastNavMeshImpl->GetNavMeshTileXY(Point, OutX, OutY);
 }
 
+bool ARecastNavMesh::GetNavmeshTileResolution(int32 TileIndex, ENavigationDataResolution& OutResolution) const
+{
+	if (RecastNavMeshImpl)
+	{
+		if (const dtNavMesh* const RecastNavMesh = RecastNavMeshImpl->DetourNavMesh)
+		{
+			if (const dtMeshTile* const Tile = RecastNavMesh->getTile(TileIndex))
+			{
+				if (const dtMeshHeader* const Header = Tile->header)
+				{
+					OutResolution = ENavigationDataResolution(Header->resolution);
+					return true;
+				}
+			}
+		}
+	}
+
+	return false;
+}
+
 bool ARecastNavMesh::CheckTileIndicesInValidRange(const FVector& Point, bool& bOutInRange) const
 {
 	const dtNavMesh* const DetourNavMesh = RecastNavMeshImpl ? RecastNavMeshImpl->DetourNavMesh : nullptr;
diff --git a/Engine/Source/Runtime/NavigationSystem/Public/NavMesh/NavMeshRenderingComponent.h b/Engine/Source/Runtime/NavigationSystem/Public/NavMesh/NavMeshRenderingComponent.h
--- a/Engine/Source/Runtime/NavigationSystem/Public/NavMesh/NavMeshRenderingComponent.h
+++ b/Engine/Source/Runtime/NavigationSystem/Public/NavMesh/NavMeshRenderingComponent.h
@@ -44,7 +44,8 @@ enum class ENavMeshDetailFlags : uint8
 	NavOctreeDetails,
 	MarkForbiddenPolys,
 	TileBuildTimes,
-	TileBuildTimesHeatMap
+	TileBuildTimesHeatMap,
+	TileResolutions
 };
 
 // exported to API for GameplayDebugger module
diff --git a/Engine/Source/Runtime/NavigationSystem/Public/NavMesh/RecastNavMesh.h b/Engine/Source/Runtime/NavigationSystem/Public/NavMesh/RecastNavMesh.h
--- a/Engine/Source/Runtime/NavigationSystem/Public/NavMesh/RecastNavMesh.h
+++ b/Engine/Source/Runtime/NavigationSystem/Public/NavMesh/RecastNavMesh.h
@@ -635,6 +635,10 @@ class NAVIGATIONSYSTEM_API ARecastNavMesh : public ANavigationData
 	/** Draw the tile boundaries */
 	UPROPERTY(EditAnywhere, Category=Display)
 	uint32 bDrawTileBounds:1;
+
+	/** Draw the tile resolutions */
+	UPROPERTY(EditAnywhere, Category=Display)
+	uint32 bDrawTileResolutions:1;
 	
 	/** Draw input geometry passed to the navmesh generator.  Recommend disabling other geometry rendering via viewport showflags in editor. */
 	UPROPERTY(EditAnywhere, Category=Display)
@@ -1073,6 +1077,9 @@ public:
 	/** Retrieves XY coordinates of tile specified by position */
 	bool GetNavMeshTileXY(const FVector& Point, int32& OutX, int32& OutY) const;
 
+	/** Retrieves the tile resolution */
+	bool GetNavmeshTileResolution(int32 TileIndex, ENavigationDataResolution& OutResolution) const;
+
 	/** Checks the supplied Points tile indicies can fit in the range of an int32 */
 	bool CheckTileIndicesInValidRange(const FVector& Point, bool& bOutInRange) const;
 
