[Navmesh] Addition of CellHeight in FNavMeshResolutionParam

diff --git a/Engine/Config/BaseEngine.ini b/Engine/Config/BaseEngine.ini
--- a/Engine/Config/BaseEngine.ini
+++ b/Engine/Config/BaseEngine.ini
@@ -2520,10 +2520,10 @@ DefaultDrawDistance=5000.0
 ; navmesh generation parameters
 TileSizeUU=1000.f
 CellSize=19.f
-NavMeshResolutionParams[0]=(CellSize=38.000000)
-NavMeshResolutionParams[1]=(CellSize=19.000000)
-NavMeshResolutionParams[2]=(CellSize=19.000000)
 CellHeight=10.f
+NavMeshResolutionParams[0]=(CellSize=38.000000,CellHeight=10.f)
+NavMeshResolutionParams[1]=(CellSize=19.000000,CellHeight=10.f)
+NavMeshResolutionParams[2]=(CellSize=19.000000,CellHeight=10.f)
 AgentRadius=34.f
 AgentHeight=144.f
 AgentMaxHeight=160.f
diff --git a/Engine/Source/Runtime/NavigationSystem/Private/NavMesh/NavMeshRenderingComponent.cpp b/Engine/Source/Runtime/NavigationSystem/Private/NavMesh/NavMeshRenderingComponent.cpp
--- a/Engine/Source/Runtime/NavigationSystem/Private/NavMesh/NavMeshRenderingComponent.cpp
+++ b/Engine/Source/Runtime/NavigationSystem/Private/NavMesh/NavMeshRenderingComponent.cpp
@@ -580,11 +580,11 @@ void FNavMeshSceneProxyData::GatherData(const ARecastNavMesh* NavMesh, int32 InN
 			const FString GenerationMode = Mode == ERuntimeGenerationType::Static ? TEXT("Static") :
 				(Mode == ERuntimeGenerationType::Dynamic ? TEXT("Dynamic") : (Mode == ERuntimeGenerationType::DynamicModifiersOnly ? TEXT("DynamicModifersOnly") : TEXT("Unknown")));
 			DebugLabels.Add(FDebugText(FString::Printf(TEXT("%s (%s%s)"), *NavMesh->GetName(), NavMesh->bIsWorldPartitioned ? TEXT("WP ") : TEXT(""), *GenerationMode)));
-			DebugLabels.Add(FDebugText(FString::Printf(TEXT("AgentRadius %0.1f, AgentHeight %0.1f, CellSizes %0.1f/%0.1f/%0.1f (low/default/high), CellHeight %0.1f"),
+			DebugLabels.Add(FDebugText(FString::Printf(TEXT("AgentRadius %0.1f, AgentHeight %0.1f, CellSizes %0.1f/%0.1f/%0.1f, CellHeights %0.1f/%0.1f/%0.1f (low/default/high)"),
 				NavMesh->AgentRadius,
 				NavMesh->AgentHeight,
 				NavMesh->GetCellSize(ENavigationDataResolution::Low), NavMesh->GetCellSize(ENavigationDataResolution::Default), NavMesh->GetCellSize(ENavigationDataResolution::High),
-				NavMesh->CellHeight)));
+				NavMesh->GetCellHeight(ENavigationDataResolution::Low), NavMesh->GetCellHeight(ENavigationDataResolution::Default), NavMesh->GetCellHeight(ENavigationDataResolution::High))));
 			DebugLabels.Add(FDebugText(FString::Printf(TEXT("Region part %s, Layer part %s"), *GetPartitioningString(NavMesh->RegionPartitioning), *GetPartitioningString(NavMesh->LayerPartitioning))));
 			DebugLabels.Add(FDebugText(TEXT(""))); // empty line
 
diff --git a/Engine/Source/Runtime/NavigationSystem/Private/NavMesh/RecastNavMesh.cpp b/Engine/Source/Runtime/NavigationSystem/Private/NavMesh/RecastNavMesh.cpp
--- a/Engine/Source/Runtime/NavigationSystem/Private/NavMesh/RecastNavMesh.cpp
+++ b/Engine/Source/Runtime/NavigationSystem/Private/NavMesh/RecastNavMesh.cpp
@@ -427,7 +427,7 @@ FRecastNavMeshGenerationProperties::FRecastNavMeshGenerationProperties(const ARe
 	
 	// @todo: Handle navmesh resolution if FRecastNavMeshGenerationProperties is used.
 	CellSize = RecastNavMesh.GetCellSize(ENavigationDataResolution::Default);
-	CellHeight = RecastNavMesh.CellHeight;
+	CellHeight = RecastNavMesh.GetCellHeight(ENavigationDataResolution::Default);
 	AgentRadius = RecastNavMesh.AgentRadius;
 	AgentHeight = RecastNavMesh.AgentHeight;
 	AgentMaxSlope = RecastNavMesh.AgentMaxSlope;
@@ -640,6 +640,15 @@ void ARecastNavMesh::PostLoad()
 			SetCellSize((ENavigationDataResolution)i, CellSize);
 		}
 	}
+
+	// If needed, initialize CellHeight from the deprecated value.
+	if (NavMeshVersion < NAVMESHVER_TILE_RESOLUTIONS_CELLHEIGHT)
+	{
+		for (int i = 0; i < (uint8)ENavigationDataResolution::MAX; ++i)
+		{
+			SetCellHeight((ENavigationDataResolution)i, CellHeight);
+		}
+	}
 	PRAGMA_ENABLE_DEPRECATION_WARNINGS
 	
 	for (uint8 Index = 0; Index < (uint8)ENavigationDataResolution::MAX; Index++)
@@ -738,14 +747,16 @@ void ARecastNavMesh::PostInitProperties()
 
 				NavMeshResolutionParams[i].CellSize = DefaultObjectCellSize;
 			}
-		}
 
-		if (CellHeight != DefOb->CellHeight)
-		{
-			UE_LOG(LogNavigation, Warning, TEXT("%s param: CellHeight(%f) differs from config settings, forcing value %f so it can be used with voxel cache!"),
-				*GetNameSafe(this), CellHeight, DefOb->CellHeight);
+			const float CurrentCellHeight = NavMeshResolutionParams[i].CellHeight;
+			const float DefaultObjectCellHeight = DefOb->NavMeshResolutionParams[i].CellHeight;
+			if (CurrentCellHeight != DefaultObjectCellHeight)
+			{
+				UE_LOG(LogNavigation, Warning, TEXT("%s param: CellHeight(%f) differs from config settings, forcing value %f so it can be used with voxel cache!"),
+					*GetNameSafe(this), CurrentCellHeight, DefaultObjectCellHeight);
 
-			CellHeight = DefOb->CellHeight;
+				NavMeshResolutionParams[i].CellHeight = DefaultObjectCellHeight;
+			}
 		}
 
 		if (AgentMaxSlope != DefOb->AgentMaxSlope)
@@ -3094,6 +3105,7 @@ void ARecastNavMesh::PostEditChangeChainProperty(FPropertyChangedChainEvent& Pro
 		if (CategoryName == NAME_Generation)
 		{
 			const FName PropName = PropertyChangedChainEvent.Property->GetFName();
+			bool bRebuild = false;
 			
 			if (PropName == GET_MEMBER_NAME_CHECKED(FNavMeshResolutionParam, CellSize))
 			{
@@ -3125,18 +3137,32 @@ void ARecastNavMesh::PostEditChangeChainProperty(FPropertyChangedChainEvent& Pro
 					CellSize = NavMeshResolutionParams[(uint8)ENavigationDataResolution::Default].CellSize;
 					PRAGMA_ENABLE_DEPRECATION_WARNINGS
 
-					// update config
-					FillConfig(NavDataConfig);
-
-					UNavigationSystemV1* NavSys = FNavigationSystem::GetCurrent<UNavigationSystemV1>(GetWorld());
-					if (!HasAnyFlags(RF_ClassDefaultObject)
-						&& NavSys && NavSys->GetIsAutoUpdateEnabled()
-						&& PropName != GET_MEMBER_NAME_CHECKED(ARecastNavMesh, MaxSimultaneousTileGenerationJobsCount))
-					{
-						RebuildAll();
-					}
+					bRebuild = true;
 				}
 			}
+			else if (PropName == GET_MEMBER_NAME_CHECKED(FNavMeshResolutionParam, CellHeight))
+			{
+				PRAGMA_DISABLE_DEPRECATION_WARNINGS
+				// Update the deprecated CellHeight to fit the default resolution CellHeight
+				CellHeight = NavMeshResolutionParams[(uint8)ENavigationDataResolution::Default].CellHeight;
+				PRAGMA_ENABLE_DEPRECATION_WARNINGS
+
+				bRebuild = true;
+			}
+
+			if (bRebuild)
+			{
+				// update config
+				FillConfig(NavDataConfig);
+
+				const UNavigationSystemV1* NavSys = FNavigationSystem::GetCurrent<UNavigationSystemV1>(GetWorld());
+				if (!HasAnyFlags(RF_ClassDefaultObject)
+					&& NavSys && NavSys->GetIsAutoUpdateEnabled()
+					&& PropName != GET_MEMBER_NAME_CHECKED(ARecastNavMesh, MaxSimultaneousTileGenerationJobsCount))
+				{
+					RebuildAll();
+				}				
+			}
 		}
 	}
 }
@@ -3300,8 +3326,8 @@ void ARecastNavMesh::UpdateGenerationProperties(const FRecastNavMeshGenerationPr
 	TileSizeUU = GenerationProps.TileSizeUU;
 	PRAGMA_DISABLE_DEPRECATION_WARNINGS
 	CellSize = GenerationProps.CellSize;
-	PRAGMA_ENABLE_DEPRECATION_WARNINGS
 	CellHeight = GenerationProps.CellHeight;
+	PRAGMA_ENABLE_DEPRECATION_WARNINGS
 	AgentRadius = GenerationProps.AgentRadius;
 	AgentHeight = GenerationProps.AgentHeight;
 	AgentMaxSlope = GenerationProps.AgentMaxSlope;
diff --git a/Engine/Source/Runtime/NavigationSystem/Private/NavMesh/RecastNavMeshGenerator.cpp b/Engine/Source/Runtime/NavigationSystem/Private/NavMesh/RecastNavMeshGenerator.cpp
--- a/Engine/Source/Runtime/NavigationSystem/Private/NavMesh/RecastNavMeshGenerator.cpp
+++ b/Engine/Source/Runtime/NavigationSystem/Private/NavMesh/RecastNavMeshGenerator.cpp
@@ -4560,6 +4560,7 @@ void FRecastNavMeshGenerator::SetupTileConfig(const ENavigationDataResolution Ti
 	check(GetOwner());
 	ensure(GetConfig().cs == GetOwner()->GetCellSize(ENavigationDataResolution::Default));
 	const float CellSize = GetOwner()->GetCellSize(TileResolution);
+	const float CellHeight = GetOwner()->GetCellHeight(TileResolution);
 	
 	// Update all settings that depends directly on indirectly of the CellSize
 	OutConfig.TileResolution = TileResolution;
@@ -4578,13 +4579,18 @@ void FRecastNavMeshGenerator::SetupTileConfig(const ENavigationDataResolution Ti
 
 	OutConfig.regionChunkSize = FMath::Max(1, OutConfig.tileSize / FMath::Max(1, DestNavMesh->LayerChunkSplits));
 	OutConfig.TileCacheChunkSize = FMath::Max(1, OutConfig.tileSize / FMath::Max(1, DestNavMesh->RegionChunkSplits));
+
+	// Update all settings that depends directly on indirectly of the CellHeight
+	OutConfig.ch = CellHeight;
+	OutConfig.walkableHeight = DestNavMesh->bMarkLowHeightAreas ? 1 : FMath::CeilToInt(DestNavMesh->AgentHeight / CellHeight);
+	OutConfig.walkableClimb = FMath::CeilToInt(DestNavMesh->AgentMaxStepHeight / CellHeight);
 }
 
 void FRecastNavMeshGenerator::ConfigureBuildProperties(FRecastBuildConfig& OutConfig)
 {
 	// @TODO those variables should be tweakable per navmesh actor
 	const float CellSize = DestNavMesh->GetCellSize(ENavigationDataResolution::Default);
-	const float CellHeight = DestNavMesh->CellHeight;
+	const float CellHeight = DestNavMesh->GetCellHeight(ENavigationDataResolution::Default);
 	const float AgentHeight = DestNavMesh->AgentHeight;
 	const float AgentMaxSlope = DestNavMesh->AgentMaxSlope;
 	const float AgentMaxClimb = DestNavMesh->AgentMaxStepHeight;
@@ -6911,10 +6917,11 @@ void FRecastNavMeshGenerator::GrabDebugSnapshot(struct FVisualLogEntry* Snapshot
 
 								if (Verts.Num())
 								{
+									const float CellHeight = NavData->GetCellHeight(ENavigationDataResolution::Default);
 									Snapshot->AddElement(
 										Verts,
-										InConvexNavAreaData.MinZ - NavData->CellHeight,
-										InConvexNavAreaData.MaxZ + NavData->CellHeight,
+										InConvexNavAreaData.MinZ - CellHeight,
+										InConvexNavAreaData.MaxZ + CellHeight,
 										CategoryName, NavAreaVerbosity, PolygonColor.WithAlpha(255));
 								}
 							};
@@ -7044,8 +7051,9 @@ void FRecastNavMeshGenerator::ExportNavigationData(const FString& FileName) cons
 								GrowConvexHull(NavData->AgentRadius, Points, ConvexVerts);
 								if (ConvexVerts.Num())
 								{
-									ExportInfo.Convex.MinZ -= NavData->CellHeight;
-									ExportInfo.Convex.MaxZ += NavData->CellHeight;
+									const float CellHeight = NavData->GetCellHeight(ENavigationDataResolution::Default);
+									ExportInfo.Convex.MinZ -= CellHeight;
+									ExportInfo.Convex.MaxZ += CellHeight;
 
 									ExportInfo.Convex.Points = UE::LWC::ConvertArrayType<FVector>(ConvexVerts);
 
diff --git a/Engine/Source/Runtime/NavigationSystem/Public/NavMesh/RecastNavMesh.h b/Engine/Source/Runtime/NavigationSystem/Public/NavMesh/RecastNavMesh.h
--- a/Engine/Source/Runtime/NavigationSystem/Public/NavMesh/RecastNavMesh.h
+++ b/Engine/Source/Runtime/NavigationSystem/Public/NavMesh/RecastNavMesh.h
@@ -596,11 +596,15 @@ struct NAVIGATIONSYSTEM_API FNavMeshResolutionParam
 {
 	GENERATED_BODY()
 
-	bool IsValid() const { return CellSize > 0.f; }
+	bool IsValid() const { return CellSize > 0.f && CellHeight > 0.f; }
 	
 	/** horizontal size of voxelization cell */
 	UPROPERTY(EditAnywhere, Category = Generation, config, meta = (ClampMin = "1.0", ClampMax = "1024.0"))
 	float CellSize = 25.f;
+
+	/** vertical size of voxelization cell */
+	UPROPERTY(EditAnywhere, Category = Generation, config, meta = (ClampMin = "1.0", ClampMax = "1024.0"))
+	float CellHeight = 10.f;
 };
 
 UCLASS(config=Engine, defaultconfig, hidecategories=(Input,Rendering,Tags,Transformation,Actor,Layers,Replication), notplaceable)
@@ -722,7 +726,8 @@ class NAVIGATIONSYSTEM_API ARecastNavMesh : public ANavigationData
 	float CellSize;
 
 	/** vertical size of voxelization cell */
-	UPROPERTY(EditAnywhere, Category = Generation, config, meta = (ClampMin = "1.0", ClampMax = "1024.0"))
+	UE_DEPRECATED(5.2, "Set the CellHeight for the required navmesh resolutions in NavMeshResolutionParams.")
+	UPROPERTY(config)
 	float CellHeight;
 
 	/** Resolution params 
@@ -1059,6 +1064,12 @@ public:
 	/** Set the CellSize for the given resolution. */
 	void SetCellSize(const ENavigationDataResolution Resolution, const float Size) { NavMeshResolutionParams[(uint8)Resolution].CellSize = Size; }
 
+	/** Get the CellHeight for the given resolution. */
+	float GetCellHeight(const ENavigationDataResolution Resolution) const { return NavMeshResolutionParams[(uint8)Resolution].CellHeight; }
+
+	/** Set the CellHeight for the given resolution. */
+	void SetCellHeight(const ENavigationDataResolution Resolution, const float Height) { NavMeshResolutionParams[(uint8)Resolution].CellHeight = Height; }
+	
 	/** Returns the tile size in world units. */
 	float GetTileSizeUU() const;
 	
diff --git a/Engine/Source/Runtime/NavigationSystem/Public/NavMesh/RecastVersion.h b/Engine/Source/Runtime/NavigationSystem/Public/NavMesh/RecastVersion.h
--- a/Engine/Source/Runtime/NavigationSystem/Public/NavMesh/RecastVersion.h
+++ b/Engine/Source/Runtime/NavigationSystem/Public/NavMesh/RecastVersion.h
@@ -22,6 +22,7 @@
 #define NAVMESHVER_OPTIM_FIX_SERIALIZE_PARAMS	21 // Fix, serialize params that used to be in the tile and are now in the navmesh.
 #define NAVMESHVER_MAXTILES_COUNT_SKIP_INCLUSION 22
 #define NAVMESHVER_TILE_RESOLUTIONS				23 // Addition of a tile resolution index to the tile header.
+#define NAVMESHVER_TILE_RESOLUTIONS_CELLHEIGHT	24 // Addition of CellHeight in the resolution params, deprecating the original CellHeight.
 
-#define NAVMESHVER_LATEST				NAVMESHVER_TILE_RESOLUTIONS
+#define NAVMESHVER_LATEST				NAVMESHVER_TILE_RESOLUTIONS_CELLHEIGHT
 #define NAVMESHVER_MIN_COMPATIBLE		NAVMESHVER_LWCOORDS_OPTIMIZATION
