Add more info about geometry gathering in RecastNavMeshGenerator

diff --git a/Engine/Source/Runtime/NavigationSystem/Private/NavMesh/RecastNavMeshGenerator.cpp b/Engine/Source/Runtime/NavigationSystem/Private/NavMesh/RecastNavMeshGenerator.cpp
--- a/Engine/Source/Runtime/NavigationSystem/Private/NavMesh/RecastNavMeshGenerator.cpp
+++ b/Engine/Source/Runtime/NavigationSystem/Private/NavMesh/RecastNavMeshGenerator.cpp
@@ -2042,7 +2042,9 @@ void FRecastTileGenerator::GatherNavigationDataGeometry(const TSharedRef<FNaviga
 	if (IsTileDebugActive())
 	{
 		UE_LOG(LogNavigation, Log, TEXT("Gathering geometry for tile (%i,%i): %s."), TileX, TileY, *GetFullNameSafe(ElementData->GetOwner()));
-		UE_LOG(LogNavigation, Log, TEXT("                       bounds: %s"), *ElementData->Bounds.ToString());
+		UE_LOG(LogNavigation, Log, TEXT("                       Bounds: %s"), *ElementData->Bounds.ToString());
+		UE_LOG(LogNavigation, Log, TEXT("                       Geometry: Has=%s Pending=%s Slice=%s"), ElementData->HasGeometry() ? TEXT("true") : TEXT("false"), ElementData->IsPendingLazyGeometryGathering() ? TEXT("true") : TEXT("false"), ElementData->SupportsGatheringGeometrySlices() ? TEXT("true") : TEXT("false"));
+		UE_LOG(LogNavigation, Log, TEXT("                       Modifier: Has=%s Pending=%s"), ElementData->HasModifiers() ? TEXT("true") : TEXT("false"), ElementData->NeedAnyPendingLazyModifiersGathering() ? TEXT("true") : TEXT("false"));
 	}
 #endif // RECAST_INTERNAL_DEBUG_DATA
 
