Deprecate unused FRecastNavMeshGenerationProperties

diff --git a/Engine/Source/Runtime/NavigationSystem/Private/NavMesh/RecastNavMesh.cpp b/Engine/Source/Runtime/NavigationSystem/Private/NavMesh/RecastNavMesh.cpp
--- a/Engine/Source/Runtime/NavigationSystem/Private/NavMesh/RecastNavMesh.cpp
+++ b/Engine/Source/Runtime/NavigationSystem/Private/NavMesh/RecastNavMesh.cpp
@@ -71,7 +71,7 @@ namespace UE::NavMesh::Private
 		return FMath::Clamp<float>(InTileSizeUU, MinTileSize, MaxTileSize);
 	}
 
-	// These should reflect the property clamping of FRecastNavMeshGenerationProperties::CellSize.  
+	// These should reflect the property clamping of FNavMeshResolutionParam::CellSize.  
 	// Minimum cell size.
 	constexpr float ArbitraryMinCellSize = 1.0f; 
 	// Maximum cell size.
@@ -392,6 +392,7 @@ namespace FNavMeshConfig
 	}
 }
 
+// Deprecated
 FRecastNavMeshGenerationProperties::FRecastNavMeshGenerationProperties()
 {
 	TilePoolSize = 1024;
@@ -420,12 +421,12 @@ FRecastNavMeshGenerationProperties::FRecastNavMeshGenerationProperties()
 	bIsWorldPartitioned = false;
 }
 
+// Deprecated
 FRecastNavMeshGenerationProperties::FRecastNavMeshGenerationProperties(const ARecastNavMesh& RecastNavMesh)
 {
 	TilePoolSize = RecastNavMesh.TilePoolSize;
 	TileSizeUU = RecastNavMesh.TileSizeUU;
 	
-	// @todo: Handle navmesh resolution if FRecastNavMeshGenerationProperties is used.
 	CellSize = RecastNavMesh.GetCellSize(ENavigationDataResolution::Default);
 	CellHeight = RecastNavMesh.GetCellHeight(ENavigationDataResolution::Default);
 	AgentRadius = RecastNavMesh.AgentRadius;
@@ -3320,6 +3321,8 @@ void ARecastNavMesh::ConditionalConstructGenerator()
 	}
 }
 
+// Deprecated
+PRAGMA_DISABLE_DEPRECATION_WARNINGS
 void ARecastNavMesh::UpdateGenerationProperties(const FRecastNavMeshGenerationProperties& GenerationProps)
 {
 	TilePoolSize = GenerationProps.TilePoolSize;
@@ -3349,6 +3352,7 @@ void ARecastNavMesh::UpdateGenerationProperties(const FRecastNavMeshGenerationPr
 	bFixedTilePoolSize = GenerationProps.bFixedTilePoolSize;
 	bIsWorldPartitioned = GenerationProps.bIsWorldPartitioned;
 }
+PRAGMA_ENABLE_DEPRECATION_WARNINGS
 
 bool ARecastNavMesh::IsVoxelCacheEnabled()
 {
diff --git a/Engine/Source/Runtime/NavigationSystem/Public/NavMesh/RecastNavMesh.h b/Engine/Source/Runtime/NavigationSystem/Public/NavMesh/RecastNavMesh.h
--- a/Engine/Source/Runtime/NavigationSystem/Public/NavMesh/RecastNavMesh.h
+++ b/Engine/Source/Runtime/NavigationSystem/Public/NavMesh/RecastNavMesh.h
@@ -315,8 +315,7 @@ struct FNavigationWallEdge
 };
 #endif //WITH_RECAST
 
-
-USTRUCT()
+USTRUCT(meta=(Deprecated = "5.2"))
 struct NAVIGATIONSYSTEM_API FRecastNavMeshGenerationProperties
 {
 	GENERATED_BODY()
@@ -1379,7 +1378,12 @@ public:
 	bool IsUsingActiveTilesGeneration(const UNavigationSystemV1& NavSys) const;
 
 	virtual void ConditionalConstructGenerator() override;
+
+PRAGMA_DISABLE_DEPRECATION_WARNINGS	
+	UE_DEPRECATED(5.2, "UpdateGenerationProperties is unused, it will be removed")
 	void UpdateGenerationProperties(const FRecastNavMeshGenerationProperties& GenerationProps);
+PRAGMA_ENABLE_DEPRECATION_WARNINGS	
+	
 	bool ShouldGatherDataOnGameThread() const { return bDoFullyAsyncNavDataGathering == false; }
 	int32 GetTileNumberHardLimit() const { return TileNumberHardLimit; }
 
