Fix balance macro warning

diff --git a/Engine/Source/Runtime/NavigationSystem/Private/NavMesh/RecastNavMesh.cpp b/Engine/Source/Runtime/NavigationSystem/Private/NavMesh/RecastNavMesh.cpp
--- a/Engine/Source/Runtime/NavigationSystem/Private/NavMesh/RecastNavMesh.cpp
+++ b/Engine/Source/Runtime/NavigationSystem/Private/NavMesh/RecastNavMesh.cpp
@@ -3327,10 +3327,10 @@ void ARecastNavMesh::UpdateGenerationProperties(const FRecastNavMeshGenerationPr
 {
 	TilePoolSize = GenerationProps.TilePoolSize;
 	TileSizeUU = GenerationProps.TileSizeUU;
-	PRAGMA_DISABLE_DEPRECATION_WARNINGS
+
 	CellSize = GenerationProps.CellSize;
 	CellHeight = GenerationProps.CellHeight;
-	PRAGMA_ENABLE_DEPRECATION_WARNINGS
+
 	AgentRadius = GenerationProps.AgentRadius;
 	AgentHeight = GenerationProps.AgentHeight;
 	AgentMaxSlope = GenerationProps.AgentMaxSlope;
