Expose option to display Solid and Walkable heightfield in navmesh generation

diff --git a/Engine/Config/BaseEngine.ini b/Engine/Config/BaseEngine.ini
--- a/Engine/Config/BaseEngine.ini
+++ b/Engine/Config/BaseEngine.ini
@@ -1105,6 +1105,10 @@ DefaultWorldPartitionHLODLayer=/Engine/Maps/Templates/HLODs/HLODLayer_Instanced.
 +FunctionRedirects=(OldName="InputDeviceSubsystem.GetPlayerControllerFromInputDevice", NewName="InputDeviceLibrary.GetPlayerControllerFromInputDevice")
 +FunctionRedirects=(OldName="InputDeviceSubsystem.GetPlayerControllerFromPlatformUser", NewName="InputDeviceLibrary.GetPlayerControllerFromPlatformUser")
 +FunctionRedirects=(OldName="InputDeviceSubsystem.IsDevicePropertyHandleValid", NewName="InputDeviceLibrary.IsDevicePropertyHandleValid")
++PropertyRedirects=(OldName="/Script/NavigationSystem.RecastNavMeshTileGenerationDebug.bHeightfieldSolidFromRasterization",NewName="bHeightfieldFromRasterization")
++PropertyRedirects=(OldName="/Script/NavigationSystem.RecastNavMeshTileGenerationDebug.bHeightfieldSolidPostInclusionBoundsFiltering",NewName="bHeightfieldPostInclusionBoundsFiltering")
++PropertyRedirects=(OldName="/Script/NavigationSystem.RecastNavMeshTileGenerationDebug.bHeightfieldSolidPostHeightFiltering",NewName="bHeightfieldPostHeightFiltering")
++PropertyRedirects=(OldName="/Script/NavigationSystem.RecastNavMeshTileGenerationDebug.bHeightfieldSolidBounds",NewName="bHeightfieldBounds")
 
 +StructRedirects=(OldName="/Script/AssetTools.RevisionInfo", NewName="/Script/AssetDefinition.RevisionInfo")
 +PropertyRedirects=(OldName="DecalComponent.SetFadeIn.Duaration", NewName="DecalComponent.SetFadeIn.Duration")
diff --git a/Engine/Source/Runtime/NavigationSystem/Private/NavMesh/RecastNavMesh.cpp b/Engine/Source/Runtime/NavigationSystem/Private/NavMesh/RecastNavMesh.cpp
--- a/Engine/Source/Runtime/NavigationSystem/Private/NavMesh/RecastNavMesh.cpp
+++ b/Engine/Source/Runtime/NavigationSystem/Private/NavMesh/RecastNavMesh.cpp
@@ -457,10 +457,11 @@ FRecastNavMeshTileGenerationDebug::FRecastNavMeshTileGenerationDebug()
 	TileCoordinate = FIntVector::ZeroValue;
 	bGenerateDebugTileOnly = false;
 	bCollisionGeometry = false;
-	bHeightfieldSolidFromRasterization = false;
-	bHeightfieldSolidPostInclusionBoundsFiltering = false;
-	bHeightfieldSolidPostHeightFiltering = false;
-	bHeightfieldSolidBounds = false;
+	HeightFieldRenderMode = EHeightFieldRenderMode::Solid;
+	bHeightfieldFromRasterization = false;
+	bHeightfieldPostInclusionBoundsFiltering = false;
+	bHeightfieldPostHeightFiltering = false;
+	bHeightfieldBounds = false;
 	bCompactHeightfield = false;
 	bCompactHeightfieldEroded = false;
 	bCompactHeightfieldRegions = false;
diff --git a/Engine/Source/Runtime/NavigationSystem/Private/NavMesh/RecastNavMeshGenerator.cpp b/Engine/Source/Runtime/NavigationSystem/Private/NavMesh/RecastNavMeshGenerator.cpp
--- a/Engine/Source/Runtime/NavigationSystem/Private/NavMesh/RecastNavMeshGenerator.cpp
+++ b/Engine/Source/Runtime/NavigationSystem/Private/NavMesh/RecastNavMeshGenerator.cpp
@@ -3415,6 +3415,23 @@ ETimeSliceWorkResult FRecastTileGenerator::GenerateCompressedLayersTimeSliced(FN
 	return ETimeSliceWorkResult::Succeeded;
 }
 
+namespace UE::NavMesh::Private
+{
+#if RECAST_INTERNAL_DEBUG_DATA
+	void DrawHeightfield(const EHeightFieldRenderMode RenderMode, duDebugDraw* dd, const rcHeightfield& hf)
+	{
+		if (RenderMode == EHeightFieldRenderMode::Solid)
+		{
+			duDebugDrawHeightfieldSolid(dd, hf);
+		}
+		else
+		{
+			duDebugDrawHeightfieldWalkable(dd, hf);
+		}
+	}
+#endif //RECAST_INTERNAL_DEBUG_DATA
+};
+
 bool FRecastTileGenerator::GenerateCompressedLayers(FNavMeshBuildContext& BuildContext)
 {
 	SCOPE_CYCLE_COUNTER(STAT_Navigation_RecastBuildCompressedLayers);
@@ -3440,11 +3457,11 @@ bool FRecastTileGenerator::GenerateCompressedLayers(FNavMeshBuildContext& BuildC
 #if RECAST_INTERNAL_DEBUG_DATA
 	if (IsTileDebugActive())
 	{
-		if (TileDebugSettings.bHeightfieldSolidFromRasterization)
+		if (TileDebugSettings.bHeightfieldFromRasterization)
 		{
-			duDebugDrawHeightfieldSolid(&BuildContext.InternalDebugData, *RasterContext.SolidHF);
+			UE::NavMesh::Private::DrawHeightfield(TileDebugSettings.HeightFieldRenderMode, &BuildContext.InternalDebugData, *RasterContext.SolidHF);
 		}
-		if (TileDebugSettings.bHeightfieldSolidBounds)
+		if (TileDebugSettings.bHeightfieldBounds)
 		{
 			duDebugDrawHeightfieldBounds(&BuildContext.InternalDebugData, *RasterContext.SolidHF);
 		}
@@ -3458,18 +3475,18 @@ bool FRecastTileGenerator::GenerateCompressedLayers(FNavMeshBuildContext& BuildC
 	}
 
 #if RECAST_INTERNAL_DEBUG_DATA
-	if (IsTileDebugActive() && TileDebugSettings.bHeightfieldSolidPostInclusionBoundsFiltering)
+	if (IsTileDebugActive() && TileDebugSettings.bHeightfieldPostInclusionBoundsFiltering)
 	{
-		duDebugDrawHeightfieldSolid(&BuildContext.InternalDebugData, *RasterContext.SolidHF);
+		UE::NavMesh::Private::DrawHeightfield(TileDebugSettings.HeightFieldRenderMode, &BuildContext.InternalDebugData, *RasterContext.SolidHF);
 	}
 #endif
 
 	GenerateRecastFilter(BuildContext, RasterContext);
 
 #if RECAST_INTERNAL_DEBUG_DATA
-	if (IsTileDebugActive() && TileDebugSettings.bHeightfieldSolidPostHeightFiltering)
+	if (IsTileDebugActive() && TileDebugSettings.bHeightfieldPostHeightFiltering)
 	{
-		duDebugDrawHeightfieldSolid(&BuildContext.InternalDebugData, *RasterContext.SolidHF);
+		UE::NavMesh::Private::DrawHeightfield(TileDebugSettings.HeightFieldRenderMode, &BuildContext.InternalDebugData, *RasterContext.SolidHF);
 	}
 #endif
 
diff --git a/Engine/Source/Runtime/NavigationSystem/Public/NavMesh/RecastNavMesh.h b/Engine/Source/Runtime/NavigationSystem/Public/NavMesh/RecastNavMesh.h
--- a/Engine/Source/Runtime/NavigationSystem/Public/NavMesh/RecastNavMesh.h
+++ b/Engine/Source/Runtime/NavigationSystem/Public/NavMesh/RecastNavMesh.h
@@ -429,6 +429,13 @@ struct NAVIGATIONSYSTEM_API FRecastNavMeshGenerationProperties
 	FRecastNavMeshGenerationProperties(const ARecastNavMesh& RecastNavMesh);
 };
 
+UENUM()
+enum class EHeightFieldRenderMode : uint8
+{
+	Solid = 0,
+	Walkable
+};
+
 USTRUCT()
 struct NAVIGATIONSYSTEM_API FRecastNavMeshTileGenerationDebug
 {
@@ -455,16 +462,19 @@ struct NAVIGATIONSYSTEM_API FRecastNavMeshTileGenerationDebug
 	uint32 bCollisionGeometry : 1;
 
 	UPROPERTY(EditAnywhere, Category = Debug)
-	uint32 bHeightfieldSolidFromRasterization : 1;
+	EHeightFieldRenderMode HeightFieldRenderMode;
+
+	UPROPERTY(EditAnywhere, Category = Debug)
+	uint32 bHeightfieldFromRasterization : 1;
 
 	UPROPERTY(EditAnywhere, Category = Debug)
-	uint32 bHeightfieldSolidPostInclusionBoundsFiltering : 1;
+	uint32 bHeightfieldPostInclusionBoundsFiltering : 1;
 
 	UPROPERTY(EditAnywhere, Category = Debug)
-	uint32 bHeightfieldSolidPostHeightFiltering : 1;
+	uint32 bHeightfieldPostHeightFiltering : 1;
 
 	UPROPERTY(EditAnywhere, Category = Debug)
-	uint32 bHeightfieldSolidBounds : 1;
+	uint32 bHeightfieldBounds : 1;
 
 	UPROPERTY(EditAnywhere, Category = Debug)
 	uint32 bCompactHeightfield : 1;
