change NavCollision version string to force an update on the DDC which contains invalid information

diff --git a/Engine/Source/Runtime/NavigationSystem/Private/NavCollision.cpp b/Engine/Source/Runtime/NavigationSystem/Private/NavCollision.cpp
--- a/Engine/Source/Runtime/NavigationSystem/Private/NavCollision.cpp
+++ b/Engine/Source/Runtime/NavigationSystem/Private/NavCollision.cpp
@@ -93,7 +93,7 @@ public:
 
 	virtual const TCHAR* GetVersionString() const override
 	{
-		return TEXT("982ABB61BCEA412C832BC21B87BF3538");
+		return TEXT("F33FFAC3B070461781F1B11C4EAC7100");
 	}
 
 	virtual FString GetPluginSpecificCacheKeySuffix() const override
