[Navmesh] Support navmesh resolution when doing async navdata gathering.

diff --git a/Engine/Source/Runtime/NavigationSystem/Private/NavMesh/RecastNavMeshGenerator.cpp b/Engine/Source/Runtime/NavigationSystem/Private/NavMesh/RecastNavMeshGenerator.cpp
--- a/Engine/Source/Runtime/NavigationSystem/Private/NavMesh/RecastNavMeshGenerator.cpp
+++ b/Engine/Source/Runtime/NavigationSystem/Private/NavMesh/RecastNavMeshGenerator.cpp
@@ -1957,8 +1957,11 @@ void FRecastTileGenerator::PrepareGeometrySources(const FRecastNavMeshGenerator&
 	const ARecastNavMesh* const OwnerNav = ParentGenerator.GetOwner();
 	const bool bUseVirtualGeometryFilteringAndDirtying = OwnerNav != nullptr && OwnerNav->bUseVirtualGeometryFilteringAndDirtying;
 
+	ENavigationDataResolution HighestResolution = ENavigationDataResolution::Low;
+	bool bNewResolutionFound = false;
+	
 	NavOctreeInstance->FindElementsWithBoundsTest(ParentGenerator.GrowBoundingBox(TileBB, /*bIncludeAgentHeight*/ false),
-		[&ParentGenerator,this, bGeometryChanged, bUseVirtualGeometryFilteringAndDirtying](const FNavigationOctreeElement& Element)
+		[&HighestResolution, &bNewResolutionFound, &ParentGenerator, this, bGeometryChanged, bUseVirtualGeometryFilteringAndDirtying](const FNavigationOctreeElement& Element)
 	{
 		const bool bShouldUse = bUseVirtualGeometryFilteringAndDirtying ?
 			ParentGenerator.ShouldGenerateGeometryForOctreeElement(Element, NavDataConfig) :
@@ -1971,10 +1974,25 @@ void FRecastTileGenerator::PrepareGeometrySources(const FRecastNavMeshGenerator&
 				Element.Data->Modifiers.HasMetaAreas() == true || 
 				Element.Data->Modifiers.IsEmpty() == false)
 			{
+				// Keep highest resolution that is not the default.
+				const ENavigationDataResolution Resolution = Element.Data->Modifiers.GetNavMeshResolution();
+				if (Resolution != ENavigationDataResolution::Invalid)
+				{
+					HighestResolution = FMath::Max(HighestResolution, Resolution);
+					bNewResolutionFound = true;
+				}
+				
 				NavigationRelevantData.Add(Element.Data);
 			}
 		}
 	});
+	
+	check(HighestResolution != ENavigationDataResolution::Invalid);
+	if (bNewResolutionFound && ParentGenerator.GetOwner()->NavMeshResolutionParams[(uint8)HighestResolution].IsValid())
+	{
+		// Update the TileConfig
+		ParentGenerator.SetupTileConfig(HighestResolution, TileConfig);
+	}
 }
 
 void FRecastTileGenerator::GatherGeometry(const FRecastNavMeshGenerator& ParentGenerator, bool bGeometryChanged)
