Remove line of code that had no functionality as UNavLinkCustomComponent does not render directly.

diff --git a/Engine/Source/Runtime/NavigationSystem/Private/NavLinkCustomComponent.cpp b/Engine/Source/Runtime/NavigationSystem/Private/NavLinkCustomComponent.cpp
--- a/Engine/Source/Runtime/NavigationSystem/Private/NavLinkCustomComponent.cpp
+++ b/Engine/Source/Runtime/NavigationSystem/Private/NavLinkCustomComponent.cpp
@@ -192,7 +192,6 @@ void UNavLinkCustomComponent::SetLinkData(const FVector& RelativeStart, const FV
 	LinkDirection = Direction;
 	
 	RefreshNavigationModifiers();
-	MarkRenderStateDirty();
 }
 
 FNavigationLink UNavLinkCustomComponent::GetLinkModifier() const
