[Navmesh] Fix crash occurring when NavigationRelevantData array is being modified during GC (when FRecastTileGenerator::AddReferencedObjects is called)
The change moves emptying NavigationRelevantData to the main thread so that it does not occur at any moment.

diff --git a/Engine/Source/Runtime/NavigationSystem/Private/NavMesh/RecastNavMeshGenerator.cpp b/Engine/Source/Runtime/NavigationSystem/Private/NavMesh/RecastNavMeshGenerator.cpp
--- a/Engine/Source/Runtime/NavigationSystem/Private/NavMesh/RecastNavMeshGenerator.cpp
+++ b/Engine/Source/Runtime/NavigationSystem/Private/NavMesh/RecastNavMeshGenerator.cpp
@@ -1843,6 +1843,7 @@ ETimeSliceWorkResult FRecastTileGenerator::DoWorkTimeSliced()
 			if (WorkResult != ETimeSliceWorkResult::CallAgainNextTimeSlice)
 			{
 				DumpAsyncData();
+				DumpSyncData();		// Currently, TIME_SLICE_NAV_REGEN can only be active if not async.
 
 				DoWorkTimeSlicedState = EDoWorkTimeSlicedState::Invalid;//Set to Invalid as we never want to call this again on this instance
 			}
@@ -1888,10 +1889,15 @@ void FRecastTileGenerator::DumpAsyncData()
 	Modifiers.Empty();
 	OffmeshLinks.Empty();
 
-	NavigationRelevantData.Empty();
 	NavSystem = nullptr;
 }
 
+void FRecastTileGenerator::DumpSyncData()
+{
+	ensure(IsInGameThread());
+	NavigationRelevantData.Empty();
+}
+	
 void FRecastTileGenerator::GatherGeometryFromSources()
 {
 	QUICK_SCOPE_CYCLE_COUNTER(STAT_RecastNavMeshGenerator_GatherGeometryFromSources);
@@ -6220,10 +6226,11 @@ TArray<FNavTileRef> FRecastNavMeshGenerator::ProcessTileTasksAsyncAndGetUpdatedT
 
 		if (Element.AsyncTask->IsDone())
 		{
+			FRecastTileGenerator& TileGenerator = *(Element.AsyncTask->GetTask().TileGenerator);
+			
 			// Add generated tiles to navmesh
 			if (!Element.bShouldDiscard)
 			{
-				FRecastTileGenerator& TileGenerator = *(Element.AsyncTask->GetTask().TileGenerator);
 				TArray<FNavTileRef> UpdatedTileRefs = AddGeneratedTilesAndGetUpdatedTiles(TileGenerator);
 				UpdatedTiles.Append(UpdatedTileRefs);
 			
@@ -6234,6 +6241,8 @@ TArray<FNavTileRef> FRecastNavMeshGenerator::ProcessTileTasksAsyncAndGetUpdatedT
 #endif
 			}
 
+			TileGenerator.DumpSyncData();
+
 			{
 				QUICK_SCOPE_CYCLE_COUNTER(STAT_RecastNavMeshGenerator_TileGeneratorRemoval);
 
diff --git a/Engine/Source/Runtime/NavigationSystem/Public/NavMesh/RecastNavMeshGenerator.h b/Engine/Source/Runtime/NavigationSystem/Public/NavMesh/RecastNavMeshGenerator.h
--- a/Engine/Source/Runtime/NavigationSystem/Public/NavMesh/RecastNavMeshGenerator.h
+++ b/Engine/Source/Runtime/NavigationSystem/Public/NavMesh/RecastNavMeshGenerator.h
@@ -444,6 +444,7 @@ protected:
 	void AddVoxelCache(TNavStatArray<uint8>& RawVoxelCache, const rcSpanCache* CachedVoxels, const int32 NumCachedVoxels) const;
 
 	void DumpAsyncData();
+	void DumpSyncData();
 
 #if RECAST_INTERNAL_DEBUG_DATA
 	bool IsTileDebugActive() const;
