NavDataHandler improvement :
- move the NavOctree construction to be able to log the info
- reduce NavOctree radius to better fit the Navigation bounds. The NavOctree radius was based on the NavBound extent which can be sqrt(2) bigger than what is actually needed
- uniformize components logging to ease debugging

diff --git a/Engine/Source/Runtime/NavigationSystem/Private/NavigationDataHandler.cpp b/Engine/Source/Runtime/NavigationSystem/Private/NavigationDataHandler.cpp
--- a/Engine/Source/Runtime/NavigationSystem/Private/NavigationDataHandler.cpp
+++ b/Engine/Source/Runtime/NavigationSystem/Private/NavigationDataHandler.cpp
@@ -7,26 +7,47 @@
 
 DEFINE_LOG_CATEGORY_STATIC(LogNavOctree, Warning, All);
 
-namespace
+namespace UE::NavigationHelper::Private
 {
-	int32 GetDirtyFlagHelper(int32 UpdateFlags, int32 DefaultValue)
+	int32 GetDirtyFlag(int32 UpdateFlags, int32 DefaultValue)
 	{
 		return ((UpdateFlags & FNavigationOctreeController::OctreeUpdate_Geometry) != 0) ? ENavigationDirtyFlag::All :
 			((UpdateFlags & FNavigationOctreeController::OctreeUpdate_Modifiers) != 0) ? ENavigationDirtyFlag::DynamicModifier :
 			DefaultValue;
 	}
+
+	FString GetElementName(const UObject& ElementOwner)
+	{
+		if (const UActorComponent* ActorComp = Cast<UActorComponent>(&ElementOwner))
+		{
+			return FString::Format(TEXT("Comp {0} of Actor {1}"), {*GetNameSafe(&ElementOwner), *GetNameSafe(ActorComp->GetOwner())});
+		}
+		return *GetNameSafe(&ElementOwner);
+	}
 }
 	
 FNavigationDataHandler::FNavigationDataHandler(FNavigationOctreeController& InOctreeController, FNavigationDirtyAreasController& InDirtyAreasController)
 		: OctreeController(InOctreeController), DirtyAreasController(InDirtyAreasController)
 {}
 
+void FNavigationDataHandler::ConstructNavOctree(const FVector& Origin, const float Radius, const ENavDataGatheringModeConfig DataGatheringMode, const float GatheringNavModifiersWarningLimitTime)
+{
+	UE_LOG(LogNavOctree, Log, TEXT("CREATE (Origin:%s Radius:%.2f)"), *Origin.ToString(), Radius);
+
+	OctreeController.Reset();
+	OctreeController.NavOctree = MakeShareable(new FNavigationOctree(Origin, Radius));
+	OctreeController.NavOctree->SetDataGatheringMode(DataGatheringMode);
+#if !UE_BUILD_SHIPPING
+	OctreeController.NavOctree->SetGatheringNavModifiersTimeLimitWarning(GatheringNavModifiersWarningLimitTime);
+#endif // !UE_BUILD_SHIPPING
+}
+
 void FNavigationDataHandler::RemoveNavOctreeElementId(const FOctreeElementId2& ElementId, int32 UpdateFlags)
 {
 	if (ensure(OctreeController.IsValidElement(ElementId)))
 	{
 		const FNavigationOctreeElement& ElementData = OctreeController.NavOctree->GetElementById(ElementId);
-		const int32 DirtyFlag = GetDirtyFlagHelper(UpdateFlags, ElementData.Data->GetDirtyFlag());
+		const int32 DirtyFlag = UE::NavigationHelper::Private::GetDirtyFlag(UpdateFlags, ElementData.Data->GetDirtyFlag());
 		// mark area occupied by given actor as dirty
 		DirtyAreasController.AddArea(ElementData.Bounds.GetBox(), DirtyFlag, [&ElementData] { return ElementData.Data->SourceObject.Get(); }, nullptr, "Remove from navoctree");
 		OctreeController.NavOctree->RemoveNode(ElementId);
@@ -51,7 +72,7 @@ FSetElementId FNavigationDataHandler::RegisterNavOctreeElement(UObject& ElementO
 	}
 
 	const bool bIsRelevant = ElementInterface.IsNavigationRelevant();
-	UE_LOG(LogNavOctree, Log, TEXT("REG %s %s"), *ElementOwner.GetName(), bIsRelevant ? TEXT("[relevant]") : TEXT("[skip]"));
+	UE_LOG(LogNavOctree, Log, TEXT("REG %s %s"), *UE::NavigationHelper::Private::GetElementName(ElementOwner), bIsRelevant ? TEXT("[relevant]") : TEXT("[skip]"));
 
 	if (bIsRelevant)
 	{
@@ -70,7 +91,7 @@ FSetElementId FNavigationDataHandler::RegisterNavOctreeElement(UObject& ElementO
 
 		if (bCanAdd)
 		{
-			FNavigationDirtyElement UpdateInfo(&ElementOwner, &ElementInterface, GetDirtyFlagHelper(UpdateFlags, 0), DirtyAreasController.bUseWorldPartitionedDynamicMode);
+			FNavigationDirtyElement UpdateInfo(&ElementOwner, &ElementInterface, UE::NavigationHelper::Private::GetDirtyFlag(UpdateFlags, 0), DirtyAreasController.bUseWorldPartitionedDynamicMode);
 
 			SetId = OctreeController.PendingOctreeUpdates.FindId(UpdateInfo);
 			if (SetId.IsValidId())
@@ -88,18 +109,6 @@ FSetElementId FNavigationDataHandler::RegisterNavOctreeElement(UObject& ElementO
 	return SetId;
 }
 
-namespace UE::NavigationHelper::Private
-{
-	FString GetElementName(const UObject* ElementOwner)
-	{
-		if (const UActorComponent* ActorComp = Cast<UActorComponent>(ElementOwner))
-		{
-			return FString::Format(TEXT("Comp {0} of Actor {1}"), {*GetNameSafe(ElementOwner), *GetNameSafe(ActorComp->GetOwner())});
-		}
-		return *GetNameSafe(ElementOwner);
-	}
-}
-
 void FNavigationDataHandler::AddElementToNavOctree(const FNavigationDirtyElement& DirtyElement)
 {
 	check(OctreeController.NavOctree.IsValid());
@@ -149,17 +158,17 @@ void FNavigationDataHandler::AddElementToNavOctree(const FNavigationDirtyElement
 		const FOctreeElementId2* ElementId = ParentId ? ParentId : OctreeController.GetObjectsNavOctreeId(*NavigationParent);
 		if (ElementId && ensure(OctreeController.IsValidElement(*ElementId)))
 		{
-			UE_LOG(LogNavOctree, Log, TEXT("ADD %s to %s"), *UE::NavigationHelper::Private::GetElementName(ElementOwner), *GetNameSafe(NavigationParent));
+			UE_LOG(LogNavOctree, Log, TEXT("ADD %s to %s"), *UE::NavigationHelper::Private::GetElementName(*ElementOwner), *GetNameSafe(NavigationParent));
 			OctreeController.NavOctree->AppendToNode(*ElementId, DirtyElement.NavInterface, ElementBounds, GeneratedData);
 		}
 		else
 		{
-			UE_LOG(LogNavOctree, Warning, TEXT("Can't add node [%s] - parent [%s] not found in octree!"), *UE::NavigationHelper::Private::GetElementName(ElementOwner), *UE::NavigationHelper::Private::GetElementName(NavigationParent));
+			UE_LOG(LogNavOctree, Warning, TEXT("Can't add node [%s] - parent [%s] not found in octree!"), *UE::NavigationHelper::Private::GetElementName(*ElementOwner), *UE::NavigationHelper::Private::GetElementName(*NavigationParent));
 		}
 	}
 	else
 	{
-		UE_LOG(LogNavOctree, Log, TEXT("ADD %s"), *UE::NavigationHelper::Private::GetElementName(ElementOwner));
+		UE_LOG(LogNavOctree, Log, TEXT("ADD %s"), *UE::NavigationHelper::Private::GetElementName(*ElementOwner));
 		OctreeController.NavOctree->AddNode(ElementOwner, DirtyElement.NavInterface, ElementBounds, GeneratedData);
 	}
 
@@ -189,7 +198,7 @@ bool FNavigationDataHandler::UnregisterNavOctreeElement(UObject& ElementOwner, I
 
 	bool bUnregistered = false;
 	const FOctreeElementId2* ElementId = OctreeController.GetObjectsNavOctreeId(ElementOwner);
-	UE_LOG(LogNavOctree, Log, TEXT("UNREG %s %s"), *ElementOwner.GetName(), ElementId ? TEXT("[exists]") : TEXT("[does\'t exist]"));
+	UE_LOG(LogNavOctree, Log, TEXT("UNREG %s %s"), *UE::NavigationHelper::Private::GetElementName(ElementOwner), ElementId ? TEXT("[exists]") : TEXT("[doesn\'t exist]"));
 
 	if (ElementId != nullptr)
 	{
@@ -460,7 +469,7 @@ void FNavigationDataHandler::AddLevelCollisionToOctree(ULevel& Level)
 				OctreeController.NavOctree->AddNode(&Level, nullptr, Bounds, BSPElem);
 				DirtyAreasController.AddArea(Bounds, ENavigationDirtyFlag::All, [&Level] { return &Level; }, nullptr, "Add level");
 
-				UE_LOG(LogNavOctree, Log, TEXT("ADD %s"), *Level.GetName());
+				UE_LOG(LogNavOctree, Log, TEXT("ADD %s"), *UE::NavigationHelper::Private::GetElementName(Level));
 			}
 		}
 	}
@@ -472,7 +481,7 @@ void FNavigationDataHandler::RemoveLevelCollisionFromOctree(ULevel& Level)
 	if (OctreeController.NavOctree.IsValid())
 	{
 		const FOctreeElementId2* ElementId = OctreeController.GetObjectsNavOctreeId(Level);
-		UE_LOG(LogNavOctree, Log, TEXT("UNREG %s %s"), *Level.GetName(), ElementId ? TEXT("[exists]") : TEXT(""));
+		UE_LOG(LogNavOctree, Log, TEXT("UNREG %s %s"), *UE::NavigationHelper::Private::GetElementName(Level), ElementId ? TEXT("[exists]") : TEXT(""));
 
 		if (ElementId != nullptr)
 		{
diff --git a/Engine/Source/Runtime/NavigationSystem/Private/NavigationSystem.cpp b/Engine/Source/Runtime/NavigationSystem/Private/NavigationSystem.cpp
--- a/Engine/Source/Runtime/NavigationSystem/Private/NavigationSystem.cpp
+++ b/Engine/Source/Runtime/NavigationSystem/Private/NavigationSystem.cpp
@@ -759,23 +759,19 @@ void UNavigationSystemV1::PostInitProperties()
 
 void UNavigationSystemV1::ConstructNavOctree()
 {
-	DefaultOctreeController.Reset();
+	// Default values to keep previous behavior.
+	FVector NavOctreeCenter = FVector::ZeroVector;
+	float NavOctreeRadius = 64000;
 
 	const FBox Bounds = GetNavigableWorldBounds();
 	if(Bounds.IsValid)
 	{
-		DefaultOctreeController.NavOctree = MakeShareable(new FNavigationOctree(Bounds.GetCenter(), 0.5*Bounds.GetSize().Length()));
-	}
-	else
-	{
-		// Fallback to previous default behavior.
-		DefaultOctreeController.NavOctree = MakeShareable(new FNavigationOctree(FVector(0, 0, 0), 64000));
+		NavOctreeCenter = Bounds.GetCenter();
+		NavOctreeRadius = Bounds.GetExtent().GetAbsMax();
 	}
-	
-	DefaultOctreeController.NavOctree->SetDataGatheringMode(DataGatheringMode);
-#if !UE_BUILD_SHIPPING	
-	DefaultOctreeController.NavOctree->SetGatheringNavModifiersTimeLimitWarning(GatheringNavModifiersWarningLimitTime);
-#endif // !UE_BUILD_SHIPPING	
+
+	FNavigationDataHandler NavHandler(DefaultOctreeController, DefaultDirtyAreasController);
+	NavHandler.ConstructNavOctree(NavOctreeCenter, NavOctreeRadius, DataGatheringMode, GatheringNavModifiersWarningLimitTime);
 }
 
 bool UNavigationSystemV1::ConditionalPopulateNavOctree()
diff --git a/Engine/Source/Runtime/NavigationSystem/Public/NavigationDataHandler.h b/Engine/Source/Runtime/NavigationSystem/Public/NavigationDataHandler.h
--- a/Engine/Source/Runtime/NavigationSystem/Public/NavigationDataHandler.h
+++ b/Engine/Source/Runtime/NavigationSystem/Public/NavigationDataHandler.h
@@ -13,6 +13,8 @@ struct NAVIGATIONSYSTEM_API FNavigationDataHandler
 
 	FNavigationDataHandler(FNavigationOctreeController& InOctreeController, FNavigationDirtyAreasController& InDirtyAreasController);
 
+	void ConstructNavOctree(const FVector& Origin, const float Radius, const ENavDataGatheringModeConfig DataGatheringMode, const float GatheringNavModifiersWarningLimitTime);
+
 	void RemoveNavOctreeElementId(const FOctreeElementId2& ElementId, int32 UpdateFlags);
 	FSetElementId RegisterNavOctreeElement(UObject& ElementOwner, INavRelevantInterface& ElementInterface, int32 UpdateFlags);
 	void AddElementToNavOctree(const FNavigationDirtyElement& DirtyElement);
