Fix for the editor viewports not being redrawn when the nav mesh has been rebuilt. This meant when many navigation related actors and components were edited (navmodifiers navlinks etc) the viewport did not reflect this uintil the user interacted with the level using the viewport.

diff --git a/Engine/Source/Runtime/NavigationSystem/Private/NavMesh/RecastNavMesh.cpp b/Engine/Source/Runtime/NavigationSystem/Private/NavMesh/RecastNavMesh.cpp
--- a/Engine/Source/Runtime/NavigationSystem/Private/NavMesh/RecastNavMesh.cpp
+++ b/Engine/Source/Runtime/NavigationSystem/Private/NavMesh/RecastNavMesh.cpp
@@ -24,6 +24,7 @@
 #include "NavigationDataHandler.h"
 
 #if WITH_EDITOR
+#include "EditorSupportDelegates.h"
 #include "ObjectEditorUtils.h"
 #include "WorldPartition/WorldPartition.h"
 #include "WorldPartition/WorldPartitionEditorLoaderAdapter.h"
@@ -571,6 +572,9 @@ void ARecastNavMesh::UpdateNavMeshDrawing()
 	if (NavMeshRenderComp != nullptr && NavMeshRenderComp->GetVisibleFlag() && (NavMeshRenderComp->IsForcingUpdate() || UNavMeshRenderingComponent::IsNavigationShowFlagSet(GetWorld())))
 	{
 		RenderingComp->MarkRenderStateDirty();
+#if WITH_EDITOR
+		FEditorSupportDelegates::RedrawAllViewports.Broadcast();
+#endif
 	}
 #endif // UE_BUILD_SHIPPING
 }
