Fixes for NavLinks not properly updating the Navoctree when NavAreas are edited (when this occurs they are reregistered), specifically when the allowed agent types are modified.

diff --git a/Engine/Source/Runtime/AIModule/Classes/Navigation/NavLinkProxy.h b/Engine/Source/Runtime/AIModule/Classes/Navigation/NavLinkProxy.h
--- a/Engine/Source/Runtime/AIModule/Classes/Navigation/NavLinkProxy.h
+++ b/Engine/Source/Runtime/AIModule/Classes/Navigation/NavLinkProxy.h
@@ -81,13 +81,15 @@ public:
 	virtual void PostEditUndo() override;
 	virtual void PostEditImport() override;
 #endif // WITH_EDITOR
+	virtual void PostInitProperties() override;
+	virtual void BeginDestroy() override;
+
 	virtual void PostRegisterAllComponents() override;
 	virtual void PostLoad() override;
 
 #if ENABLE_VISUAL_LOG
 protected:
 	virtual void BeginPlay() override;
-
 public:
 #endif // ENABLE_VISUAL_LOG
 
@@ -131,6 +133,11 @@ protected:
 
 	void NotifySmartLinkReached(UNavLinkCustomComponent* LinkComp, UObject* PathingAgent, const FVector& DestPoint);
 
+#if WITH_EDITOR
+	void OnNavAreaRegistered(const UWorld& World, const UClass* NavAreaClass);
+	void OnNavAreaUnregistered(const UWorld& World, const UClass* NavAreaClass);
+#endif // WITH_EDITOR
+
 public:
 	/** Returns SmartLinkComp subobject **/
 	UNavLinkCustomComponent* GetSmartLinkComp() const { return SmartLinkComp; }
@@ -140,4 +147,10 @@ public:
 	/** Returns SpriteComponent subobject **/
 	UBillboardComponent* GetSpriteComponent() const { return SpriteComponent; }
 #endif
+
+protected:
+#if WITH_EDITOR
+	FDelegateHandle OnNavAreaRegisteredDelegateHandle;
+	FDelegateHandle OnNavAreaUnregisteredDelegateHandle;
+#endif // WITH_EDITOR
 };
diff --git a/Engine/Source/Runtime/AIModule/Private/Navigation/NavLinkProxy.cpp b/Engine/Source/Runtime/AIModule/Private/Navigation/NavLinkProxy.cpp
--- a/Engine/Source/Runtime/AIModule/Private/Navigation/NavLinkProxy.cpp
+++ b/Engine/Source/Runtime/AIModule/Private/Navigation/NavLinkProxy.cpp
@@ -16,7 +16,6 @@
 #include "VisualLogger/VisualLogger.h"
 #include "NavigationOctree.h"
 #include "ObjectEditorUtils.h"
-#include "EditorSupportDelegates.h"
 #if WITH_EDITOR
 #include "ScopedTransaction.h"
 #endif
@@ -25,6 +24,35 @@
 
 #define LOCTEXT_NAMESPACE "NavLink"
 
+#if WITH_EDITOR
+namespace UE::Navigation::LinkProxy::Private
+{
+	void OnNavAreaRegistrationChanged(ANavLinkProxy& LinkProxy, const UWorld& World, const UClass* NavAreaClass)
+	{
+		if (&World != LinkProxy.GetWorld())
+		{
+			return;
+		}
+
+		bool bUpdateActor = false;
+
+		for (const FNavigationLink& NavLink : LinkProxy.PointLinks)
+		{
+			if (NavLink.GetAreaClass() == NavAreaClass)
+			{
+				bUpdateActor = true;
+				break;
+			}
+		}
+
+		if (bUpdateActor)
+		{
+			FNavigationSystem::UpdateActorData(LinkProxy);
+		}
+	}
+} // UE::Navigation::LinkProxy::Private
+#endif // WITH_EDITOR
+
 ANavLinkProxy::ANavLinkProxy(const FObjectInitializer& ObjectInitializer) : Super(ObjectInitializer)
 {
 	USceneComponent* SceneComponent = CreateDefaultSubobject<USceneComponent>(TEXT("PositionComponent"));
@@ -81,7 +109,45 @@ ANavLinkProxy::ANavLinkProxy(const FObjectInitializer& ObjectInitializer) : Supe
 	SetCanBeDamaged(false);
 }
 
+void ANavLinkProxy::PostInitProperties()
+{
+	Super::PostInitProperties();
+
+#if WITH_EDITOR
+	if (GIsEditor && !HasAnyFlags(RF_ClassDefaultObject))
+	{
+		OnNavAreaRegisteredDelegateHandle = UNavigationSystemBase::OnNavAreaRegisteredDelegate().AddUObject(this, &ANavLinkProxy::OnNavAreaRegistered);
+		OnNavAreaUnregisteredDelegateHandle = UNavigationSystemBase::OnNavAreaUnregisteredDelegate().AddUObject(this, &ANavLinkProxy::OnNavAreaUnregistered);
+	}
+#endif // WITH_EDITOR
+}
+
+void ANavLinkProxy::BeginDestroy()
+{
+#if WITH_EDITOR
+	if (GIsEditor && !HasAnyFlags(RF_ClassDefaultObject))
+	{
+		UNavigationSystemBase::OnNavAreaRegisteredDelegate().Remove(OnNavAreaRegisteredDelegateHandle);
+		UNavigationSystemBase::OnNavAreaUnregisteredDelegate().Remove(OnNavAreaUnregisteredDelegateHandle);
+	}
+#endif // WITH_EDITOR
+
+	Super::BeginDestroy();
+}
+
 #if WITH_EDITOR
+// This function is only called if GIsEditor == true
+void ANavLinkProxy::OnNavAreaRegistered(const UWorld& World, const UClass* NavAreaClass)
+{
+	UE::Navigation::LinkProxy::Private::OnNavAreaRegistrationChanged(*this, World, NavAreaClass);
+}
+
+// This function is only called if GIsEditor == true
+void ANavLinkProxy::OnNavAreaUnregistered(const UWorld& World, const UClass* NavAreaClass)
+{
+	UE::Navigation::LinkProxy::Private::OnNavAreaRegistrationChanged(*this, World, NavAreaClass);
+}
+
 void ANavLinkProxy::PostEditChangeProperty(FPropertyChangedEvent& PropertyChangedEvent) 
 {
 	static const FName NAME_SmartLinkIsRelevant = GET_MEMBER_NAME_CHECKED(ANavLinkProxy, bSmartLinkIsRelevant);
@@ -142,7 +208,6 @@ void ANavLinkProxy::PostEditImport()
 		Link.InitializeAreaClass(/*bForceRefresh=*/true);
 	}
 }
-
 #endif // WITH_EDITOR
 
 void ANavLinkProxy::PostRegisterAllComponents()
@@ -312,7 +377,6 @@ void ANavLinkProxy::CopyEndPointsFromSimpleLinkToSmartLink()
 		{
 			EdRenderComp->MarkRenderStateDirty();
 		}
-		FEditorSupportDelegates::RedrawAllViewports.Broadcast();
 	}
 }
 #endif // WITH_EDITOR
diff --git a/Engine/Source/Runtime/NavigationSystem/Private/NavLinkCustomComponent.cpp b/Engine/Source/Runtime/NavigationSystem/Private/NavLinkCustomComponent.cpp
--- a/Engine/Source/Runtime/NavigationSystem/Private/NavLinkCustomComponent.cpp
+++ b/Engine/Source/Runtime/NavigationSystem/Private/NavLinkCustomComponent.cpp
@@ -17,6 +17,18 @@
 
 #include UE_INLINE_GENERATED_CPP_BY_NAME(NavLinkCustomComponent)
 
+#if WITH_EDITOR
+namespace UE::Navigation::LinkCustomComponent::Private
+{
+	void OnNavAreaRegistrationChanged(UNavLinkCustomComponent& CustomComponent, const UWorld& World, const UClass* NavAreaClass)
+	{
+		if (NavAreaClass && (NavAreaClass == CustomComponent.GetLinkAreaClass() || NavAreaClass == CustomComponent.GetObstacleAreaClass()) && &World == CustomComponent.GetWorld())
+		{
+			CustomComponent.RefreshNavigationModifiers();
+		}
+	}
+} // UE::Navigation::LinkCustomComponent::Private
+#endif // WITH_EDITOR
 
 UNavLinkCustomComponent::UNavLinkCustomComponent(const FObjectInitializer& ObjectInitializer) : Super(ObjectInitializer)
 {
@@ -44,6 +56,20 @@ void UNavLinkCustomComponent::PostLoad()
 	INavLinkCustomInterface::UpdateUniqueId(NavLinkUserId);
 }
 
+#if WITH_EDITOR
+// This function is only called if GIsEditor == true for non default objects components that are registered.
+void UNavLinkCustomComponent::OnNavAreaRegistered(const UWorld& World, const UClass* NavAreaClass)
+{
+	UE::Navigation::LinkCustomComponent::Private::OnNavAreaRegistrationChanged(*this, World, NavAreaClass);
+}
+
+// This function is only called if GIsEditor == true for non default objects components that are registered.
+void UNavLinkCustomComponent::OnNavAreaUnregistered(const UWorld& World, const UClass* NavAreaClass)
+{
+	UE::Navigation::LinkCustomComponent::Private::OnNavAreaRegistrationChanged(*this, World, NavAreaClass);
+}
+#endif // WITH_EDITOR
+
 TStructOnScope<FActorComponentInstanceData> UNavLinkCustomComponent::GetComponentInstanceData() const
 {
 	TStructOnScope<FActorComponentInstanceData> InstanceData = MakeStructOnScope<FActorComponentInstanceData, FNavLinkCustomInstanceData>(this);
@@ -85,7 +111,7 @@ void UNavLinkCustomComponent::PostEditImport()
 	NavLinkUserId = INavLinkCustomInterface::GetUniqueId();
 	UE_LOG(LogNavLink, VeryVerbose, TEXT("%s new navlink id %u."), ANSI_TO_TCHAR(__FUNCTION__), NavLinkUserId);
 }
-#endif
+#endif // WITH_EDITOR
 
 void UNavLinkCustomComponent::GetLinkData(FVector& LeftPt, FVector& RightPt, ENavLinkDirection::Type& Direction) const
 {
@@ -176,12 +202,28 @@ void UNavLinkCustomComponent::OnRegister()
 	}
 
 	UNavigationSystemV1::RequestCustomLinkRegistering(*this, this);
+
+#if WITH_EDITOR
+	if (GIsEditor && !HasAnyFlags(RF_ClassDefaultObject))
+	{
+		OnNavAreaRegisteredDelegateHandle = UNavigationSystemBase::OnNavAreaRegisteredDelegate().AddUObject(this, &UNavLinkCustomComponent::OnNavAreaRegistered);
+		OnNavAreaUnregisteredDelegateHandle = UNavigationSystemBase::OnNavAreaUnregisteredDelegate().AddUObject(this, &UNavLinkCustomComponent::OnNavAreaUnregistered);
+	}
+#endif // WITH_EDITOR 
 }
 
 void UNavLinkCustomComponent::OnUnregister()
 {
 	Super::OnUnregister();
 
+#if WITH_EDITOR
+	if (GIsEditor && !HasAnyFlags(RF_ClassDefaultObject))
+	{
+		UNavigationSystemBase::OnNavAreaRegisteredDelegate().Remove(OnNavAreaRegisteredDelegateHandle);
+		UNavigationSystemBase::OnNavAreaUnregisteredDelegate().Remove(OnNavAreaUnregisteredDelegateHandle);
+	}
+#endif // WITH_EDITOR 
+
 	UNavigationSystemV1::RequestCustomLinkUnregistering(*this, this);
 }
 
diff --git a/Engine/Source/Runtime/NavigationSystem/Public/NavLinkCustomComponent.h b/Engine/Source/Runtime/NavigationSystem/Public/NavLinkCustomComponent.h
--- a/Engine/Source/Runtime/NavigationSystem/Public/NavLinkCustomComponent.h
+++ b/Engine/Source/Runtime/NavigationSystem/Public/NavLinkCustomComponent.h
@@ -111,6 +111,8 @@ class NAVIGATIONSYSTEM_API UNavLinkCustomComponent : public UNavRelevantComponen
 	/** get link end point in world space */
 	FVector GetEndPoint() const;
 
+	TSubclassOf<UNavArea> GetObstacleAreaClass() const { return ObstacleAreaClass; }
+
 	//////////////////////////////////////////////////////////////////////////
 	// helper functions for setting delegates
 
@@ -135,7 +137,13 @@ class NAVIGATIONSYSTEM_API UNavLinkCustomComponent : public UNavRelevantComponen
 	{
 		SetBroadcastFilter(FBroadcastFilter::CreateUObject(TargetOb, InFunc));
 	}
-	
+
+protected:
+#if WITH_EDITOR
+	void OnNavAreaRegistered(const UWorld& World, const UClass* NavAreaClass);
+	void OnNavAreaUnregistered(const UWorld& World, const UClass* NavAreaClass);
+#endif // WITH_EDITOR
+
 protected:
 
 	/** link Id assigned by navigation system */
@@ -223,6 +231,11 @@ protected:
 	
 	/** gather agents to notify about state change */
 	void CollectNearbyAgents(TArray<UObject*>& NotifyList);
+
+#if WITH_EDITOR
+	FDelegateHandle OnNavAreaRegisteredDelegateHandle;
+	FDelegateHandle OnNavAreaUnregisteredDelegateHandle;
+#endif // WITH_EDITOR
 };
 
 /** Used to store navlink data during RerunConstructionScripts */
