Consistent use of WITH_EDITOR as compared with NavModifierComponent, NavLinkProxy and NavLinkCustomComponent

diff --git a/Engine/Source/Runtime/NavigationSystem/Private/NavModifierVolume.cpp b/Engine/Source/Runtime/NavigationSystem/Private/NavModifierVolume.cpp
--- a/Engine/Source/Runtime/NavigationSystem/Private/NavModifierVolume.cpp
+++ b/Engine/Source/Runtime/NavigationSystem/Private/NavModifierVolume.cpp
@@ -13,6 +13,19 @@
 
 #include UE_INLINE_GENERATED_CPP_BY_NAME(NavModifierVolume)
 
+#if WITH_EDITOR
+namespace UE::Navigation::ModVolume::Private
+{
+	void OnNavAreaRegistrationChanged(ANavModifierVolume& ModifierVolume, const UWorld& World, const UClass* NavAreaClass)
+	{
+		if (NavAreaClass && NavAreaClass == ModifierVolume.GetAreaClass() && &World == ModifierVolume.GetWorld())
+		{
+			FNavigationSystem::UpdateActorData(ModifierVolume);
+		}
+	}
+} // UE::Navigation::ModVolumne::Private
+#endif // WITH_EDITOR
+
 //----------------------------------------------------------------------//
 // ANavModifierVolume
 //----------------------------------------------------------------------//
@@ -54,27 +67,19 @@ void ANavModifierVolume::BeginDestroy()
 #endif // WITH_EDITOR
 }
 
+#if WITH_EDITOR
 // This function is only called if GIsEditor == true
 void ANavModifierVolume::OnNavAreaRegistered(const UWorld& World, const UClass* NavAreaClass)
 {
-#if WITH_EDITOR
-	if (NavAreaClass && NavAreaClass == AreaClass && &World == GetWorld())
-	{
-		FNavigationSystem::UpdateActorData(*this);
-	}
-#endif // WITH_EDITOR
+	UE::Navigation::ModVolume::Private::OnNavAreaRegistrationChanged(*this, World, NavAreaClass);
 }
 
 // This function is only called if GIsEditor == true
 void ANavModifierVolume::OnNavAreaUnregistered(const UWorld& World, const UClass* NavAreaClass)
 {
-#if WITH_EDITOR
-	if (NavAreaClass && NavAreaClass == AreaClass && &World == GetWorld())
-	{
-		FNavigationSystem::UpdateActorData(*this);
-	}
-#endif // WITH_EDITOR
+	UE::Navigation::ModVolume::Private::OnNavAreaRegistrationChanged(*this, World, NavAreaClass);
 }
+#endif // WITH_EDITOR
 
 void ANavModifierVolume::GetNavigationData(FNavigationRelevantData& Data) const
 {
diff --git a/Engine/Source/Runtime/NavigationSystem/Public/NavModifierVolume.h b/Engine/Source/Runtime/NavigationSystem/Public/NavModifierVolume.h
--- a/Engine/Source/Runtime/NavigationSystem/Public/NavModifierVolume.h
+++ b/Engine/Source/Runtime/NavigationSystem/Public/NavModifierVolume.h
@@ -34,8 +34,10 @@ protected:
 	UPROPERTY(EditAnywhere, Category = Default, AdvancedDisplay)
 	ENavigationDataResolution NavMeshResolution;
 
+#if WITH_EDITOR
 	FDelegateHandle OnNavAreaRegisteredDelegateHandle;
 	FDelegateHandle OnNavAreaUnregisteredDelegateHandle;
+#endif
 
 public:
 	ANavModifierVolume(const FObjectInitializer& ObjectInitializer);
@@ -58,6 +60,8 @@ protected:
 	virtual void PostInitProperties() override;
 	virtual void BeginDestroy() override;
 
+#if WITH_EDITOR
 	void OnNavAreaRegistered(const UWorld& World, const UClass* NavAreaClass);
 	void OnNavAreaUnregistered(const UWorld& World, const UClass* NavAreaClass);
+#endif
 };
