Proper use of WITH_EDITOR. Also called RefreshNavigationModifiers() rather than completely reregistering the component OnNavAreaRegistrationChanged().

diff --git a/Engine/Source/Runtime/NavigationSystem/Private/NavModifierComponent.cpp b/Engine/Source/Runtime/NavigationSystem/Private/NavModifierComponent.cpp
--- a/Engine/Source/Runtime/NavigationSystem/Private/NavModifierComponent.cpp
+++ b/Engine/Source/Runtime/NavigationSystem/Private/NavModifierComponent.cpp
@@ -10,6 +10,19 @@
 
 #include UE_INLINE_GENERATED_CPP_BY_NAME(NavModifierComponent)
 
+#if WITH_EDITOR
+namespace UE::Navigation::ModComponent::Private
+{
+	void OnNavAreaRegistrationChanged(UNavModifierComponent& ModifierComponent, const UWorld& World, const UClass* NavAreaClass)
+	{
+		if (NavAreaClass && NavAreaClass == ModifierComponent.AreaClass && &World == ModifierComponent.GetWorld())
+		{
+			ModifierComponent.RefreshNavigationModifiers();
+		}
+	}
+} // UE::Navigation::ModComponent::Private
+#endif // WITH_EDITOR
+
 UNavModifierComponent::UNavModifierComponent(const FObjectInitializer& ObjectInitializer) : Super(ObjectInitializer)
 {
 	AreaClass = UNavArea_Null::StaticClass();
@@ -43,32 +56,19 @@ void UNavModifierComponent::OnUnregister()
 	Super::OnUnregister();
 }
 
-namespace UE::Navigation::Private
-{
-	void OnNavAreaRegistrationChanged(UNavModifierComponent& ModifierComponent, const UWorld& World, const UClass* NavAreaClass)
-	{
 #if WITH_EDITOR
-		if (NavAreaClass && NavAreaClass == ModifierComponent.AreaClass && &World == ModifierComponent.GetWorld())
-		{
-			// Reregister the component, this follows the same pattern in editor to editing the values of components member variables,
-			// actor that owns the component will have it reregistered.
-			ModifierComponent.ReregisterComponent();
-		}
-#endif // WITH_EDITOR 
-	}
-}
-
 // This function is only called if GIsEditor == true for non default objects components that are registered.
 void UNavModifierComponent::OnNavAreaRegistered(const UWorld& World, const UClass* NavAreaClass)
 {
-	UE::Navigation::Private::OnNavAreaRegistrationChanged(*this, World, NavAreaClass);
+	UE::Navigation::ModComponent::Private::OnNavAreaRegistrationChanged(*this, World, NavAreaClass);
 }
 
 // This function is only called if GIsEditor == true for non default objects components that are registered.
 void UNavModifierComponent::OnNavAreaUnregistered(const UWorld& World, const UClass* NavAreaClass)
 {
-	UE::Navigation::Private::OnNavAreaRegistrationChanged(*this, World, NavAreaClass);
+	UE::Navigation::ModComponent::Private::OnNavAreaRegistrationChanged(*this, World, NavAreaClass);
 }
+#endif // WITH_EDITOR 
 
 void UNavModifierComponent::CalcAndCacheBounds() const
 {
diff --git a/Engine/Source/Runtime/NavigationSystem/Public/NavModifierComponent.h b/Engine/Source/Runtime/NavigationSystem/Public/NavModifierComponent.h
--- a/Engine/Source/Runtime/NavigationSystem/Public/NavModifierComponent.h
+++ b/Engine/Source/Runtime/NavigationSystem/Public/NavModifierComponent.h
@@ -36,8 +36,11 @@ class NAVIGATIONSYSTEM_API UNavModifierComponent : public UNavRelevantComponent
 
 protected:
 	void OnTransformUpdated(USceneComponent* RootComponent, EUpdateTransformFlags UpdateTransformFlags, ETeleportType Teleport);
+
+#if WITH_EDITOR
 	void OnNavAreaRegistered(const UWorld& World, const UClass* NavAreaClass);
 	void OnNavAreaUnregistered(const UWorld& World, const UClass* NavAreaClass);
+#endif // WITH_EDITOR 
 
 	//~ Begin UActorComponent Interface
 	virtual void OnRegister() override;
@@ -59,6 +62,8 @@ protected:
 	 *	cached data is still valid */
 	mutable FTransform CachedTransform;
 
+#if WITH_EDITOR
 	FDelegateHandle OnNavAreaRegisteredDelegateHandle;
 	FDelegateHandle OnNavAreaUnregisteredDelegateHandle;
+#endif // WITH_EDITOR 
 };
