Rename navigation console commands prefix from n to ai.nav
#rn The prefix for AI navigation console commands has been changed from n to ai.nav.

diff --git a/Engine/Source/Editor/UnrealEd/Private/EditorBuildUtils.cpp b/Engine/Source/Editor/UnrealEd/Private/EditorBuildUtils.cpp
--- a/Engine/Source/Editor/UnrealEd/Private/EditorBuildUtils.cpp
+++ b/Engine/Source/Editor/UnrealEd/Private/EditorBuildUtils.cpp
@@ -88,7 +88,8 @@ FName FEditorBuildUtils::InProgressBuildId;
 namespace UE::EditorBuildUtils
 {
 	static bool bNavmeshAllowPartitionedBuildingFromEditor = false; // Experimental, not enabled by default yet.
-	static FAutoConsoleVariableRef AllowPartitionedBuildingFromEditor(TEXT("n.bNavmeshAllowPartitionedBuildingFromEditor"), bNavmeshAllowPartitionedBuildingFromEditor, TEXT("Enable experimental navmesh partition building."), ECVF_Default);
+	static FAutoConsoleVariableRef AllowPartitionedBuildingFromEditorDeprecated(TEXT("n.bNavmeshAllowPartitionedBuildingFromEditor"), bNavmeshAllowPartitionedBuildingFromEditor, TEXT("Enable experimental navmesh partition building. Deprecated 5.3: use ai.nav.bNavmeshAllowPartitionedBuildingFromEditor instead."), ECVF_Default);
+	static FAutoConsoleVariableRef AllowPartitionedBuildingFromEditor(TEXT("ai.nav.bNavmeshAllowPartitionedBuildingFromEditor"), bNavmeshAllowPartitionedBuildingFromEditor, TEXT("Enable experimental navmesh partition building."), ECVF_Default);
 }
 
 /**
diff --git a/Engine/Source/Runtime/NavigationSystem/Private/NavMesh/RecastNavMeshGenerator.cpp b/Engine/Source/Runtime/NavigationSystem/Private/NavMesh/RecastNavMeshGenerator.cpp
--- a/Engine/Source/Runtime/NavigationSystem/Private/NavMesh/RecastNavMeshGenerator.cpp
+++ b/Engine/Source/Runtime/NavigationSystem/Private/NavMesh/RecastNavMeshGenerator.cpp
@@ -60,31 +60,31 @@ struct dtTileCacheAlloc;
 
 //Experimental debug tools
 static int32 GNavmeshSynchronousTileGeneration = 0;
-static FAutoConsoleVariableRef NavmeshVarSynchronous(TEXT("n.GNavmeshSynchronousTileGeneration"), GNavmeshSynchronousTileGeneration, TEXT(""), ECVF_Default);
+static FAutoConsoleVariableRef NavmeshVarSynchronous(TEXT("ai.nav.GNavmeshSynchronousTileGeneration"), GNavmeshSynchronousTileGeneration, TEXT(""), ECVF_Default);
 
 #if RECAST_INTERNAL_DEBUG_DATA
 static int32 GNavmeshDebugTileX = MAX_int32;
 static int32 GNavmeshDebugTileY = MAX_int32;
 static bool GNavmeshGenerateDebugTileOnly = false;
-static FAutoConsoleVariableRef NavmeshVarDebugTileX(TEXT("n.GNavmeshDebugTileX"), GNavmeshDebugTileX, TEXT(""), ECVF_Default);
-static FAutoConsoleVariableRef NavmeshVarDebugTileY(TEXT("n.GNavmeshDebugTileY"), GNavmeshDebugTileY, TEXT(""), ECVF_Default);
-static FAutoConsoleVariableRef NavmeshVarGenerateDebugTileOnly(TEXT("n.GNavmeshGenerateDebugTileOnly"), GNavmeshGenerateDebugTileOnly, TEXT(""), ECVF_Default);
+static FAutoConsoleVariableRef NavmeshVarDebugTileX(TEXT("ai.nav.GNavmeshDebugTileX"), GNavmeshDebugTileX, TEXT(""), ECVF_Default);
+static FAutoConsoleVariableRef NavmeshVarDebugTileY(TEXT("ai.nav.GNavmeshDebugTileY"), GNavmeshDebugTileY, TEXT(""), ECVF_Default);
+static FAutoConsoleVariableRef NavmeshVarGenerateDebugTileOnly(TEXT("ai.nav.GNavmeshGenerateDebugTileOnly"), GNavmeshGenerateDebugTileOnly, TEXT(""), ECVF_Default);
 #endif //RECAST_INTERNAL_DEBUG_DATA
 
 // Hotfixing this flag without rebuilding the data will cause decompression errors, equivalent to not having prebuilt navmesh data at all.
 static bool GNavmeshUseOodleCompression = true;
-static FAutoConsoleVariableRef NavmeshVarOodleCompression(TEXT("n.NavmeshUseOodleCompression"), GNavmeshUseOodleCompression, TEXT("Use Oodle for run-time tile cache compression/decompression. Optimized for size in editor, optimized for speed in standalone."), ECVF_Default);
+static FAutoConsoleVariableRef NavmeshVarOodleCompression(TEXT("ai.nav.NavmeshUseOodleCompression"), GNavmeshUseOodleCompression, TEXT("Use Oodle for run-time tile cache compression/decompression. Optimized for size in editor, optimized for speed in standalone."), ECVF_Default);
 
 namespace UE::NavMesh::Private
 {
 	static float RecentlyBuildTileDisplayTime = 0.2f;
-	static FAutoConsoleVariableRef CVarRecentlyBuildTileDisplayTime(TEXT("n.RecentlyBuildTileDisplayTime"), RecentlyBuildTileDisplayTime, TEXT("Time (in seconds) to display tiles that have recently been built."), ECVF_Default);
+	static FAutoConsoleVariableRef CVarRecentlyBuildTileDisplayTime(TEXT("ai.nav.RecentlyBuildTileDisplayTime"), RecentlyBuildTileDisplayTime, TEXT("Time (in seconds) to display tiles that have recently been built."), ECVF_Default);
 
 	static bool bUseTightBoundExpansion = true;
-	static FAutoConsoleVariableRef CVarUseTightBoundExpansion(TEXT("n.UseTightBoundExpansion"), bUseTightBoundExpansion, TEXT("Active by default. Use an expansion of one AgentRadius. Set to false to revert to the previous behavior (2 AgentRadius)."), ECVF_Default);
+	static FAutoConsoleVariableRef CVarUseTightBoundExpansion(TEXT("ai.nav.UseTightBoundExpansion"), bUseTightBoundExpansion, TEXT("Active by default. Use an expansion of one AgentRadius. Set to false to revert to the previous behavior (2 AgentRadius)."), ECVF_Default);
 
 	static bool bKeepSteepSlopeForSingleVoxelAgent = true;
-	static FAutoConsoleVariableRef CVarKeepSteepSlopeForSingleVoxelAgent(TEXT("n.KeepSteepSlopeForSingleVoxelAgent"), bKeepSteepSlopeForSingleVoxelAgent, TEXT("Active by default. Fix too wide filtering of steep slope in Recast Heightfield filtering when the agent radius is only 1 voxel. Set to false to revert to the previous behavior."), ECVF_Default);
+	static FAutoConsoleVariableRef CVarKeepSteepSlopeForSingleVoxelAgent(TEXT("ai.nav.KeepSteepSlopeForSingleVoxelAgent"), bKeepSteepSlopeForSingleVoxelAgent, TEXT("Active by default. Fix too wide filtering of steep slope in Recast Heightfield filtering when the agent radius is only 1 voxel. Set to false to revert to the previous behavior."), ECVF_Default);
 }
 
 static FOodleDataCompression::ECompressor GNavmeshTileCacheCompressor = FOodleDataCompression::ECompressor::Mermaid;
