Fix memory corruption from FRecastTileGenerator::MarkDynamicAreas

diff --git a/Engine/Source/Runtime/NavigationSystem/Private/NavMesh/RecastNavMeshGenerator.cpp b/Engine/Source/Runtime/NavigationSystem/Private/NavMesh/RecastNavMeshGenerator.cpp
--- a/Engine/Source/Runtime/NavigationSystem/Private/NavMesh/RecastNavMeshGenerator.cpp
+++ b/Engine/Source/Runtime/NavigationSystem/Private/NavMesh/RecastNavMeshGenerator.cpp
@@ -4058,9 +4058,13 @@ void FRecastTileGenerator::MarkDynamicAreas(dtTileCacheLayer& Layer)
 
 	if (Modifiers.Num())
 	{
-		if (AdditionalCachedData.bUseSortFunction && AdditionalCachedData.ActorOwner && Modifiers.Num() > 1)
+		if (AdditionalCachedData.bUseSortFunction && Modifiers.Num() > 1)
 		{
-			AdditionalCachedData.ActorOwner->SortAreasForGenerator(Modifiers);
+			FGCScopeGuard GCScopeGuard;
+			if (const ARecastNavMesh* ActorOwner = AdditionalCachedData.ActorOwner.Get())
+			{
+				ActorOwner->SortAreasForGenerator(Modifiers);
+			}
 		}
 
 		// 1: if navmesh is using low areas, apply only low area replacements
diff --git a/Engine/Source/Runtime/NavigationSystem/Public/NavMesh/RecastNavMesh.h b/Engine/Source/Runtime/NavigationSystem/Public/NavMesh/RecastNavMesh.h
--- a/Engine/Source/Runtime/NavigationSystem/Public/NavMesh/RecastNavMesh.h
+++ b/Engine/Source/Runtime/NavigationSystem/Public/NavMesh/RecastNavMesh.h
@@ -1496,7 +1496,7 @@ struct FRecastNavMeshCachedData
 	ARecastNavMesh::FNavPolyFlags FlagsPerArea[RECAST_MAX_AREAS];
 	ARecastNavMesh::FNavPolyFlags FlagsPerOffMeshLinkArea[RECAST_MAX_AREAS];
 	TMap<const UClass*, int32> AreaClassToIdMap;
-	const ARecastNavMesh* ActorOwner;
+	TWeakObjectPtr<const ARecastNavMesh> ActorOwner;
 	uint32 bUseSortFunction : 1;
 
 	static FRecastNavMeshCachedData Construct(const ARecastNavMesh* RecastNavMeshActor);
