Fix crash when there is no navigation system.

diff --git a/Engine/Source/Runtime/NavigationSystem/Private/NavMesh/RecastNavMeshGenerator.cpp b/Engine/Source/Runtime/NavigationSystem/Private/NavMesh/RecastNavMeshGenerator.cpp
--- a/Engine/Source/Runtime/NavigationSystem/Private/NavMesh/RecastNavMeshGenerator.cpp
+++ b/Engine/Source/Runtime/NavigationSystem/Private/NavMesh/RecastNavMeshGenerator.cpp
@@ -4688,7 +4688,7 @@ void FRecastNavMeshGenerator::ConfigureBuildProperties(FRecastBuildConfig& OutCo
 	}
 
 	const UNavigationSystemV1* NavSys = FNavigationSystem::GetCurrent<UNavigationSystemV1>(GetWorld());
-	OutConfig.AgentIndex = NavSys->GetSupportedAgentIndex(DestNavMesh);
+	OutConfig.AgentIndex = NavSys ? NavSys->GetSupportedAgentIndex(DestNavMesh) : 0;
 
 	OutConfig.tileSize = FMath::Max(FMath::TruncToInt(DestNavMesh->TileSizeUU / CellSize), 1);
 	UE_CLOG(OutConfig.tileSize == 1, LogNavigation, Error, TEXT("RecastNavMesh TileSize of 1 is highly discouraged. This occurence indicates an issue with RecastNavMesh\'s generation properties (specifically TileSizeUU: %f, CellSize: %f). Please ensure their correctness.")
