[Navmesh] CSV profiler: addition of task count and task wait time per layer

diff --git a/Engine/Source/Runtime/NavigationSystem/Private/NavMesh/RecastNavMeshGenerator.cpp b/Engine/Source/Runtime/NavigationSystem/Private/NavMesh/RecastNavMeshGenerator.cpp
--- a/Engine/Source/Runtime/NavigationSystem/Private/NavMesh/RecastNavMeshGenerator.cpp
+++ b/Engine/Source/Runtime/NavigationSystem/Private/NavMesh/RecastNavMeshGenerator.cpp
@@ -1636,7 +1636,7 @@ uint32 GetTileCacheSizeHelper(TArray<FNavMeshTileData>& CompressedTiles)
 // FRecastTileGenerator
 //----------------------------------------------------------------------//
 
-FRecastTileGenerator::FRecastTileGenerator(FRecastNavMeshGenerator& ParentGenerator, const FIntPoint& Location)
+FRecastTileGenerator::FRecastTileGenerator(FRecastNavMeshGenerator& ParentGenerator, const FIntPoint& Location, const double PendingTileCreationTime)
 	: TimeSliceManager(ParentGenerator.GetTimeSliceManager())
 {
 	bUpdateGeometry = true;
@@ -1645,6 +1645,8 @@ FRecastTileGenerator::FRecastTileGenerator(FRecastNavMeshGenerator& ParentGenera
 	TileX = Location.X;
 	TileY = Location.Y;
 
+	TileCreationTime = PendingTileCreationTime;
+
 	// Copy tile config from parent generator
 	TileConfig = ParentGenerator.GetConfig();
 	
@@ -5278,16 +5280,19 @@ void FRecastNavMeshGenerator::ReAddTiles(const TArray<FIntPoint>& Tiles)
 
 	TSet<FPendingTileElement> DirtyTiles;
 
+	const double CurrentTimeSeconds = FPlatformTime::Seconds();
+
 	// @note we act on assumption all items in Tiles are unique
 	for (const FIntPoint& TileCoords : Tiles)
 	{
 		FPendingTileElement Element;
 		Element.Coord = TileCoords;
 		Element.bRebuildGeometry = true;
+		Element.CreationTime = CurrentTimeSeconds;
 		DirtyTiles.Add(Element);
 	}
 
-	int32 NumTilesMarked = DirtyTiles.Num();
+	const int32 NumTilesMarked = DirtyTiles.Num();
 
 	// Merge all pending tiles into one container
 	for (const FPendingTileElement& Element : PendingDirtyTiles)
@@ -5296,6 +5301,7 @@ void FRecastNavMeshGenerator::ReAddTiles(const TArray<FIntPoint>& Tiles)
 		if (ExistingElement)
 		{
 			ExistingElement->bRebuildGeometry |= Element.bRebuildGeometry;
+			ExistingElement->CreationTime = FMath::Min(Element.CreationTime, ExistingElement->CreationTime);
 			// Append area bounds to existing list 
 			if (ExistingElement->bRebuildGeometry == false)
 			{
@@ -5860,6 +5866,8 @@ void FRecastNavMeshGenerator::MarkDirtyTiles(const TArray<FNavigationDirtyArea>&
 	const FNavigationOctree* const NavOctreeInstance = (bUseVirtualGeometryFilteringAndDirtying && NavSys) ? NavSys->GetNavOctree() : nullptr;
 	const FNavDataConfig* const NavDataConfig = bUseVirtualGeometryFilteringAndDirtying ? &DestNavMesh->GetConfig() : nullptr;
 	// ~
+
+	const double CurrentTimeSeconds = FPlatformTime::Seconds();
 		
 	// find all tiles that need regeneration
 	TSet<FPendingTileElement> DirtyTiles;
@@ -5948,6 +5956,7 @@ void FRecastNavMeshGenerator::MarkDirtyTiles(const TArray<FNavigationDirtyArea>&
 				FPendingTileElement Element;
 				Element.Coord = FIntPoint(TileX, TileY);
 				Element.bRebuildGeometry = DirtyArea.HasFlag(ENavigationDirtyFlag::Geometry) || DirtyArea.HasFlag(ENavigationDirtyFlag::NavigationBounds);
+				Element.CreationTime = CurrentTimeSeconds;
 				if (Element.bRebuildGeometry == false)
 				{
 					Element.DirtyAreas.Add(AdjustedAreaBounds);
@@ -5994,6 +6003,7 @@ void FRecastNavMeshGenerator::MarkDirtyTiles(const TArray<FNavigationDirtyArea>&
 			const FPendingTileElement& DirtyElement = DirtyTiles[Id];
 
 			ExistingElement.bRebuildGeometry |= DirtyElement.bRebuildGeometry;
+			ExistingElement.CreationTime = FMath::Min(DirtyElement.CreationTime, ExistingElement.CreationTime);
 			// Append area bounds to existing list 
 			if (ExistingElement.bRebuildGeometry == false)
 			{
@@ -6103,11 +6113,11 @@ void FRecastNavMeshGenerator::GetSeedLocations(UWorld& World, TArray<FVector2D>&
 	}
 }
 
-TSharedRef<FRecastTileGenerator> FRecastNavMeshGenerator::CreateTileGenerator(const FIntPoint& Coord, const TArray<FBox>& DirtyAreas)
+TSharedRef<FRecastTileGenerator> FRecastNavMeshGenerator::CreateTileGenerator(const FIntPoint& Coord, const TArray<FBox>& DirtyAreas, const double PendingTileCreationTime /*=0.*/)
 {
 	SCOPE_CYCLE_COUNTER(STAT_Navigation_CreateTileGenerator);
 
-	return ConstuctTileGeneratorImpl<FRecastTileGenerator>(Coord, DirtyAreas);
+	return ConstructTileGeneratorImpl<FRecastTileGenerator>(Coord, DirtyAreas, PendingTileCreationTime);
 }
 
 // Deprecated
@@ -6181,7 +6191,7 @@ TArray<FNavTileRef> FRecastNavMeshGenerator::ProcessTileTasksAsyncAndGetUpdatedT
 		if (!RunningDirtyTiles.Contains(RunningElement))
 		{
 			// Spawn async task
-			TUniquePtr<FRecastTileGeneratorTask> TileTask = MakeUnique<FRecastTileGeneratorTask>(CreateTileGenerator(PendingElement.Coord, PendingElement.DirtyAreas));
+			TUniquePtr<FRecastTileGeneratorTask> TileTask = MakeUnique<FRecastTileGeneratorTask>(CreateTileGenerator(PendingElement.Coord, PendingElement.DirtyAreas, PendingElement.CreationTime));
 
 			// Start it in background in case it has something to build
 			if (TileTask->GetTask().TileGenerator->HasDataToBuild())
@@ -6275,7 +6285,8 @@ TSharedRef<FRecastTileGenerator> FRecastNavMeshGenerator::CreateTileGeneratorFro
 	OutTileLocation.X = PendingElement.Coord.X;
 	OutTileLocation.Y = PendingElement.Coord.Y;
 
-	TSharedRef<FRecastTileGenerator> TileGenerator = CreateTileGenerator(PendingElement.Coord, PendingElement.DirtyAreas);
+	TSharedRef<FRecastTileGenerator> TileGenerator = CreateTileGenerator(PendingElement.Coord, PendingElement.DirtyAreas, PendingElement.CreationTime);
+
 	PendingDirtyTiles.RemoveAt(PendingItemIdx, 1, false);
 
 	return TileGenerator;
@@ -6490,6 +6501,17 @@ TArray<FNavTileRef> FRecastNavMeshGenerator::ProcessTileTasksSyncTimeSlicedAndGe
 				SyncTimeSlicedData.TimeSliceManager->PushTileRegenTime(SyncTimeSlicedData.CurrentTileRegenDuration);
 
 				SyncTimeSlicedData.CurrentTileRegenDuration = 0.;
+
+#if !UE_BUILD_SHIPPING					
+				const UNavigationSystemV1* NavSys = FNavigationSystem::GetCurrent<UNavigationSystemV1>(GetWorld());
+				if (NavSys)
+				{
+					const TSharedPtr<FRecastTileGenerator>& TileGenerator = SyncTimeSlicedData.TileGeneratorSync;
+					const double WaitTime = (FPlatformTime::Seconds() - TileGenerator->TileCreationTime);
+					const int32 NavDataIndex = NavSys->NavDataSet.Find(DestNavMesh);	
+					SyncTimeSlicedData.TimeSliceManager->PushTileWaitTime(NavDataIndex, WaitTime);
+				}
+#endif // !UE_BUILD_SHIPPING					
 				
 				MARK_TIMESLICE_SECTION_DEBUG(SyncTimeSlicedData.TimeSliceManager->GetTimeSlicer(), FinishTile);
 
diff --git a/Engine/Source/Runtime/NavigationSystem/Private/NavigationSystem.cpp b/Engine/Source/Runtime/NavigationSystem/Private/NavigationSystem.cpp
--- a/Engine/Source/Runtime/NavigationSystem/Private/NavigationSystem.cpp
+++ b/Engine/Source/Runtime/NavigationSystem/Private/NavigationSystem.cpp
@@ -155,6 +155,7 @@ DEFINE_STAT(STAT_DetourTileClustersMemory);
 DEFINE_STAT(STAT_DetourTilePolyClustersMemory);
 
 CSV_DEFINE_CATEGORY(NavigationSystem, false);
+CSV_DEFINE_CATEGORY(NavTasksDelays, true);
 CSV_DEFINE_CATEGORY(NavTasks, true);
 
 //----------------------------------------------------------------------//
@@ -311,6 +312,43 @@ bool FNavRegenTimeSlicer::TestTimeSliceFinished() const
 	return bTimeSliceFinishedCached;
 }
 
+void FNavRegenTimeSliceManager::PushTileWaitTime(const int32 NavDataIndex, const double NewTime)
+{
+	if (TileWaitTimes.IsValidIndex(NavDataIndex))
+	{
+		TileWaitTimes[NavDataIndex].Add(NewTime);	
+	}
+}
+
+double FNavRegenTimeSliceManager::GetAverageTileWaitTime(const int32 NavDataIndex) const
+{
+	if (!TileWaitTimes.IsValidIndex(NavDataIndex))
+	{
+		return 0.;
+	}
+	
+	double Total = 0.;
+	const TArray<double>& TimeArray = TileWaitTimes[NavDataIndex];
+	if (TimeArray.Num() == 0)
+	{
+		return 0.;			
+	}
+		
+	for (const double Time : TimeArray)
+	{
+		Total += Time;
+	}
+	return Total / TimeArray.Num();
+}
+
+void FNavRegenTimeSliceManager::ResetTileWaitTime(const int32 NavDataIndex)
+{
+	if (TileWaitTimes.IsValidIndex(NavDataIndex))
+	{
+		TileWaitTimes[NavDataIndex].Reset();			
+	}
+}
+
 FNavRegenTimeSliceManager::FNavRegenTimeSliceManager()
 	: MinTimeSliceDuration(0.00075)
 	, MaxTimeSliceDuration(0.004)
@@ -339,7 +377,7 @@ void FNavRegenTimeSliceManager::CalcAverageDeltaTime(uint64 FrameNum)
 	FrameNumOld = FrameNum;
 }
 
-void FNavRegenTimeSliceManager::CalcTimeSliceDuration(int32 NumTilesToRegen, const TArray<double>& CurrentTileRegenDurations)
+void FNavRegenTimeSliceManager::CalcTimeSliceDuration(const TArray<TObjectPtr<ANavigationData>>& NavDataSet, int32 NumTilesToRegen, const TArray<double>& CurrentTileRegenDurations)
 {
 	const float RawDeltaTimesAverage = FloatCastChecked<float>(MovingWindowDeltaTime.GetAverage(), UE::LWC::DefaultFloatPrecision);
 	const float DeltaTimesAverage = (RawDeltaTimesAverage > 0.f) ? RawDeltaTimesAverage : (1.f / 30.f); //use default 33 ms
@@ -371,7 +409,22 @@ void FNavRegenTimeSliceManager::CalcTimeSliceDuration(int32 NumTilesToRegen, con
 	CSV_CUSTOM_STAT(NavigationSystem, NavTilesToAddForLongCurrentTileRegen, TilesToAddForLongCurrentTileRegen, ECsvCustomStatOp::Set);
 	CSV_CUSTOM_STAT(NavigationSystem, NavTileAvRegenTimeMs, static_cast<float>(MovingWindowTileRegenTime.GetAverage() * 1000.), ECsvCustomStatOp::Set);
 	CSV_CUSTOM_STAT(NavigationSystem, NavTileAvRegenDeltaTimeMs, static_cast<float>(MovingWindowDeltaTime.GetAverage() * 1000.), ECsvCustomStatOp::Set);
-#endif
+
+	for (int32 NavDataIndex = 0; NavDataIndex < NavDataSet.Num(); ++NavDataIndex)
+	{
+		if (TileWaitTimes.IsValidIndex(NavDataIndex))
+		{
+#if CSV_PROFILER			
+			const float WaitTime = static_cast<float>(GetAverageTileWaitTime(NavDataIndex) * 1000.);
+
+			const FString StatName = FString::Printf(TEXT("NavTileAvTileWaitTimeMs_%s"), *GetNameSafe(NavDataSet[NavDataIndex])); 
+			FCsvProfiler::RecordCustomStat(*StatName, CSV_CATEGORY_INDEX(NavTasksDelays), WaitTime, ECsvCustomStatOp::Set);
+#endif // CSV_PROFILER
+
+			ResetTileWaitTime(NavDataIndex);
+		}
+	}
+#endif // !UE_BUILD_SHIPPING
 }
 
 void FNavRegenTimeSliceManager::SetMinTimeSliceDuration(double NewMinTimeSliceDuration)
@@ -1129,9 +1182,9 @@ void UNavigationSystemV1::SetCrowdManager(UCrowdManagerBase* NewCrowdManager)
 	}
 }
 
-void UNavigationSystemV1::CalcTimeSlicedUpdateData(TArray<double>& OutCurrentTimeSlicedBuildTaskDurations, TArray<bool>& OutIsTimeSlicingArray, bool& bOutAnyNonTimeSlicedGenerators, int32& OutNumTimeSlicedRemainingBuildTasks)
+void UNavigationSystemV1::CalcTimeSlicedUpdateData(TArray<double>& OutCurrentTimeSlicedBuildTaskDurations, TArray<bool>& OutIsTimeSlicingArray, bool& bOutAnyNonTimeSlicedGenerators, TArray<int32, TInlineAllocator<8>>& OutNumTimeSlicedRemainingBuildTasksArray)
 {
-	OutNumTimeSlicedRemainingBuildTasks = 0;
+	OutNumTimeSlicedRemainingBuildTasksArray.SetNumZeroed(NavDataSet.Num());
 	OutIsTimeSlicingArray.SetNumZeroed(NavDataSet.Num());
 	bOutAnyNonTimeSlicedGenerators = false;
 	OutCurrentTimeSlicedBuildTaskDurations.Reset(NavDataSet.Num());
@@ -1148,7 +1201,7 @@ void UNavigationSystemV1::CalcTimeSlicedUpdateData(TArray<double>& OutCurrentTim
 			if (Generator->GetTimeSliceData(NumRemainingBuildTasksTemp, TimeSlicedBuildTaskDuration))
 			{
 				OutIsTimeSlicingArray[NavDataIdx] = true;
-				OutNumTimeSlicedRemainingBuildTasks += NumRemainingBuildTasksTemp;
+				OutNumTimeSlicedRemainingBuildTasksArray[NavDataIdx] += NumRemainingBuildTasksTemp;
 				if (TimeSlicedBuildTaskDuration > 0.)
 				{
 					OutCurrentTimeSlicedBuildTaskDurations.Push(TimeSlicedBuildTaskDuration);
@@ -1221,25 +1274,33 @@ void UNavigationSystemV1::Tick(float DeltaSeconds)
 
 			if (NavRegenTimeSliceManager.DoTimeSlicedUpdate())
 			{
-				int32 NumTimeSlicedRemainingBuildTasks = 0;
+				TArray<int32, TInlineAllocator<8>> NumTimeSlicedRemainingBuildTasksArray;
+				NumTimeSlicedRemainingBuildTasksArray.SetNumZeroed(NavDataSet.Num());
+				
 				TArray<double> CurrentTimeSlicedBuildTaskDurations;
 				TArray<bool> IsTimeSlicingArray;
 				bool bAnyNonTimeSlicedGenerators = false;
 
 				NavRegenTimeSliceManager.CalcAverageDeltaTime(GFrameCounter);
 
-				CalcTimeSlicedUpdateData(CurrentTimeSlicedBuildTaskDurations, IsTimeSlicingArray, bAnyNonTimeSlicedGenerators, NumTimeSlicedRemainingBuildTasks);
+				CalcTimeSlicedUpdateData(CurrentTimeSlicedBuildTaskDurations, IsTimeSlicingArray, bAnyNonTimeSlicedGenerators, NumTimeSlicedRemainingBuildTasksArray);
 
+				int32 NumTimeSlicedRemainingBuildTasks = 0;
+				for (const int32 NumTasks : NumTimeSlicedRemainingBuildTasksArray)
+				{
+					NumTimeSlicedRemainingBuildTasks += NumTasks;
+				}
+				
 				if (NumTimeSlicedRemainingBuildTasks > 0)
 				{
-					NavRegenTimeSliceManager.CalcTimeSliceDuration(NumTimeSlicedRemainingBuildTasks, CurrentTimeSlicedBuildTaskDurations);
+					NavRegenTimeSliceManager.CalcTimeSliceDuration(NavDataSet, NumTimeSlicedRemainingBuildTasks, CurrentTimeSlicedBuildTaskDurations);
 
 					//The general idea here is to tick any non time sliced generators once per frame. Time sliced generators we aim to tick one per frame and move to the next, next frame. In the
 					//case where one time sliced generator doesn't use the whole time slice we move to the next time sliced generator. That generator will only be considered to have a full frames
 					//processing if either it runs out of work or uses a large % of the time slice. Depending we either tick it again next frame or go to the next time sliced generator (next frame).
 					bool bNavDataIdxSet = false;
 					int32 NavDataIdxTemp = NavRegenTimeSliceManager.GetNavDataIdx();
-					const double RemainingFractionConsideredWholeTick = 0.8;
+					constexpr double RemainingFractionConsideredWholeTick = 0.8;
 					const int32 FirstNavDataIdx = NavDataIdxTemp = NavDataIdxTemp % NavDataSet.Num();
 
 					for (int32 NavDataIter = 0; NavDataIter < NavDataSet.Num(); ++NavDataIter)
@@ -1306,8 +1367,22 @@ void UNavigationSystemV1::Tick(float DeltaSeconds)
 		}
 	}
 
-	CSV_CUSTOM_STAT(NavTasks, NumRemainingTasks, GetNumRemainingBuildTasks(), ECsvCustomStatOp::Set);
-	CSV_CUSTOM_STAT(NavTasks, NumRunningTasks, GetNumRunningBuildTasks(), ECsvCustomStatOp::Set);
+#if !UE_BUILD_SHIPPING && CSV_PROFILER
+	for (const TObjectPtr<ANavigationData>& NavigationData : NavDataSet)
+	{
+		if (NavigationData)
+		{
+			if (const FNavDataGenerator* Generator = NavigationData->GetGenerator())
+			{
+				const int32 BuildTaskNum = Generator->GetNumRemaningBuildTasks();
+				const FString StatName = FString::Printf(TEXT("NumRemainingTasks_%s"), *GetNameSafe(NavigationData)); 
+				FCsvProfiler::RecordCustomStat(*StatName, CSV_CATEGORY_INDEX(NavTasks), BuildTaskNum, ECsvCustomStatOp::Set);
+			}
+		}
+	}
+	
+	CSV_CUSTOM_STAT(NavigationSystem, NumRunningTasks, GetNumRunningBuildTasks(), ECsvCustomStatOp::Set);
+#endif // !UE_BUILD_SHIPPING && CSV_PROFILER
 
 	// In multithreaded configuration we can process async pathfinding queries
 	// in dedicated task while dispatching completed queries results on the main thread.
@@ -2453,6 +2528,8 @@ UNavigationSystemV1::ERegistrationResult UNavigationSystemV1::RegisterNavData(AN
 		Result = RegistrationFailed_AgentNotValid;
 	}
 
+	NavRegenTimeSliceManager.InitTileWaitTimeArrays(NavDataSet);
+
 	// @todo else might consider modifying this NavData to implement navigation for one of the supported agents
 	// care needs to be taken to not make it implement navigation for agent who's real implementation has 
 	// not been loaded yet.
diff --git a/Engine/Source/Runtime/NavigationSystem/Public/NavMesh/RecastNavMeshGenerator.h b/Engine/Source/Runtime/NavigationSystem/Public/NavMesh/RecastNavMeshGenerator.h
--- a/Engine/Source/Runtime/NavigationSystem/Public/NavMesh/RecastNavMeshGenerator.h
+++ b/Engine/Source/Runtime/NavigationSystem/Public/NavMesh/RecastNavMeshGenerator.h
@@ -298,7 +298,7 @@ class NAVIGATIONSYSTEM_API FRecastTileGenerator : public FNoncopyable, public FG
 	friend FRecastNavMeshGenerator;
 
 public:
-	FRecastTileGenerator(FRecastNavMeshGenerator& ParentGenerator, const FIntPoint& Location);
+	FRecastTileGenerator(FRecastNavMeshGenerator& ParentGenerator, const FIntPoint& Location, const double PendingTileCreationTime);
 	virtual ~FRecastTileGenerator();
 		
 	/** Does the work involved with regenerating this tile using time slicing.
@@ -480,6 +480,10 @@ protected:
 	int32 TileX;
 	int32 TileY;
 	uint32 Version;
+
+	/** Time when the tile was requested */
+	double TileCreationTime = 0.;
+	
 	/** Tile's bounding box, Unreal coords */
 	FBox TileBB;
 
@@ -554,16 +558,19 @@ struct FPendingTileElement
 	FIntPoint	Coord;
 	/** distance to seed, used for sorting pending tiles */
 	FVector::FReal SeedDistance;
+	/** time at which the element was first added to the queue */
+	double		CreationTime;
 	/** Whether we need a full rebuild for this tile grid cell */
 	bool		bRebuildGeometry;
 	/** We need to store dirty area bounds to check which cached layers needs to be regenerated
-	 *  In case geometry is changed cached layers data will be fully regenerated without using dirty areas list
+	 *  In case geometry is changed, cached layers data will be fully regenerated without using dirty areas list
 	 */
 	TArray<FBox> DirtyAreas;
 
 	FPendingTileElement()
 		: Coord(FIntPoint::NoneValue)
 		, SeedDistance(TNumericLimits<FVector::FReal>::Max())
+		, CreationTime(-1.)
 		, bRebuildGeometry(false)
 	{
 	}
@@ -583,6 +590,11 @@ struct FPendingTileElement
 		return Other.SeedDistance < SeedDistance;
 	}
 
+	double GetDurationSinceCreation(const double CurrenTimeSeconds) const
+	{
+		return (CreationTime > 0.) ? CurrenTimeSeconds - CreationTime : 0.;
+	}
+
 	friend uint32 GetTypeHash(const FPendingTileElement& Element)
 	{
 		return GetTypeHash(Element.Coord);
@@ -870,12 +882,19 @@ protected:
 	/** Blocks until build for specified list of tiles is complete and discard results */
 	void DiscardCurrentBuildingTasks();
 
-	virtual TSharedRef<FRecastTileGenerator> CreateTileGenerator(const FIntPoint& Coord, const TArray<FBox>& DirtyAreas);
+	virtual TSharedRef<FRecastTileGenerator> CreateTileGenerator(const FIntPoint& Coord, const TArray<FBox>& DirtyAreas, const double PendingTileCreationTime = 0.);
 
 	template <typename T>
-	TSharedRef<T> ConstuctTileGeneratorImpl(const FIntPoint& Coord, const TArray<FBox>& DirtyAreas)
+	UE_DEPRECATED(5.3, "Use ConstructTileGeneratorImpl instead.")
+	TSharedRef<T> ConstuctTileGeneratorImpl(const FIntPoint& Coord, const TArray<FBox>& DirtyAreas, const double PendingTileCreationTime = 0.)
+	{
+		ConstructTileGeneratorImpl<T>(Coord, DirtyAreas, PendingTileCreationTime);
+	}
+	
+	template <typename T>
+	TSharedRef<T> ConstructTileGeneratorImpl(const FIntPoint& Coord, const TArray<FBox>& DirtyAreas, const double PendingTileCreationTime)
 	{
-		TSharedRef<T> TileGenerator = MakeShareable(new T(*this, Coord));
+		TSharedRef<T> TileGenerator = MakeShareable(new T(*this, Coord, PendingTileCreationTime));
 		TileGenerator->Setup(*this, DirtyAreas);
 		return TileGenerator;
 	}
diff --git a/Engine/Source/Runtime/NavigationSystem/Public/NavigationSystem.h b/Engine/Source/Runtime/NavigationSystem/Public/NavigationSystem.h
--- a/Engine/Source/Runtime/NavigationSystem/Public/NavigationSystem.h
+++ b/Engine/Source/Runtime/NavigationSystem/Public/NavigationSystem.h
@@ -193,11 +193,16 @@ public:
 
 	double GetAverageDeltaTime() const { return MovingWindowDeltaTime.GetAverage(); }
 
+	void InitTileWaitTimeArrays(const TArray<TObjectPtr<ANavigationData>>& NavDataSet) { TileWaitTimes.SetNum(NavDataSet.Num()); }
+	void PushTileWaitTime(const int32 NavDataIndex, const double NewTime);
+	double GetAverageTileWaitTime(const int32 NavDataIndex) const;
+	void ResetTileWaitTime(const int32 NavDataIndex);
+
 	bool DoTimeSlicedUpdate() const { return bDoTimeSlicedUpdate; }
 
 	void CalcAverageDeltaTime(uint64 FrameNum);
 
-	void CalcTimeSliceDuration(int32 NumTilesToRegen, const TArray<double>& CurrentTileRegenDurations);
+	void CalcTimeSliceDuration(const TArray<TObjectPtr<ANavigationData>>& NavDataSet, int32 NumTilesToRegen, const TArray<double>& CurrentTileRegenDurations);
 
 	void SetMinTimeSliceDuration(double NewMinTimeSliceDuration);
 
@@ -220,6 +225,9 @@ protected:
 	/** Used to calculate the actual moving window delta time */
 	FMovingWindowAverageFast<double, 256> MovingWindowDeltaTime;
 
+	/** Average tile wait time per NavDataIndex */
+	TArray<TArray<double>> TileWaitTimes;
+
 	/** If there are enough tiles to process this in the Min Time Slice Duration */
 	double MinTimeSliceDuration;
 
@@ -544,7 +552,7 @@ public:
 	UCrowdManagerBase* GetCrowdManager() const { return CrowdManager.Get(); }
 
 protected:
-	void CalcTimeSlicedUpdateData(TArray<double>& OutCurrentTimeSlicedBuildTaskDurations, TArray<bool>& OutIsTimeSlicingArray, bool& bOutAnyNonTimeSlicedGenerators, int32& OutNumTimeSlicedRemainingBuildTasks);
+	void CalcTimeSlicedUpdateData(TArray<double>& OutCurrentTimeSlicedBuildTaskDurations, TArray<bool>& OutIsTimeSlicingArray, bool& bOutAnyNonTimeSlicedGenerators, TArray<int32, TInlineAllocator<8>>& OutNumTimeSlicedRemainingBuildTasksArray);
 
 	/** spawn new crowd manager */
 	virtual void CreateCrowdManager();
