Fix issue in FPImplRecastNavMesh::Raycast when no hit is found.

diff --git a/Engine/Source/Runtime/NavigationSystem/Private/NavMesh/PImplRecastNavMesh.cpp b/Engine/Source/Runtime/NavigationSystem/Private/NavMesh/PImplRecastNavMesh.cpp
--- a/Engine/Source/Runtime/NavigationSystem/Private/NavMesh/PImplRecastNavMesh.cpp
+++ b/Engine/Source/Runtime/NavigationSystem/Private/NavMesh/PImplRecastNavMesh.cpp
@@ -1119,7 +1119,7 @@ void FPImplRecastNavMesh::Raycast(const FVector& StartLoc, const FVector& EndLoc
 	}
 	else
 	{
-		RaycastResult.HitTime = 0.f;
+		RaycastResult.HitTime = ARecastNavMesh::FRaycastResult::InvalidHitTime;
 		RaycastResult.HitNormal = (StartLoc - EndLoc).GetSafeNormal();
 	}
 }
diff --git a/Engine/Source/Runtime/NavigationSystem/Public/NavMesh/RecastNavMesh.h b/Engine/Source/Runtime/NavigationSystem/Public/NavMesh/RecastNavMesh.h
--- a/Engine/Source/Runtime/NavigationSystem/Public/NavMesh/RecastNavMesh.h
+++ b/Engine/Source/Runtime/NavigationSystem/Public/NavMesh/RecastNavMesh.h
@@ -934,10 +934,11 @@ public:
 		FVector::FReal HitTime;
 		FVector HitNormal;
 		uint32 bIsRaycastEndInCorridor : 1;
+		inline static constexpr FVector::FReal InvalidHitTime = TNumericLimits<FVector::FReal>::Max();
 
 		FRaycastResult()
 			: CorridorPolysCount(0)
-			, HitTime(TNumericLimits<FVector::FReal>::Max())
+			, HitTime(InvalidHitTime)
 			, HitNormal(0.f)
 			, bIsRaycastEndInCorridor(false)
 		{
@@ -946,7 +947,7 @@ public:
 		}
 
 		FORCEINLINE int32 GetMaxCorridorSize() const { return MAX_PATH_CORRIDOR_POLYS; }
-		FORCEINLINE bool HasHit() const { return HitTime != TNumericLimits<FVector::FReal>::Max(); }
+		FORCEINLINE bool HasHit() const { return HitTime != InvalidHitTime; }
 		FORCEINLINE NavNodeRef GetLastNodeRef() const { return CorridorPolysCount > 0 ? CorridorPolys[CorridorPolysCount - 1] : INVALID_NAVNODEREF; }
 	};
 
