[Navmesh] Logs: addition of LogNavInvokers log category and addition of logs detailing invoker registration.
CSVProfiler: addition of NavInvokers csv category and track invokers count in csv.

diff --git a/Engine/Source/Runtime/Engine/Classes/AI/NavigationSystemBase.h b/Engine/Source/Runtime/Engine/Classes/AI/NavigationSystemBase.h
--- a/Engine/Source/Runtime/Engine/Classes/AI/NavigationSystemBase.h
+++ b/Engine/Source/Runtime/Engine/Classes/AI/NavigationSystemBase.h
@@ -22,6 +22,7 @@ class UNavAreaBase;
 
 ENGINE_API DECLARE_LOG_CATEGORY_EXTERN(LogNavigation, Warning, All);
 ENGINE_API DECLARE_LOG_CATEGORY_EXTERN(LogNavigationDataBuild, Log, All);
+ENGINE_API DECLARE_LOG_CATEGORY_EXTERN(LogNavInvokers, Warning, All);
 ENGINE_API DECLARE_LOG_CATEGORY_EXTERN(LogNavLink, Warning, All);
 ENGINE_API DECLARE_LOG_CATEGORY_EXTERN(LogAStar, Display, All);
 
diff --git a/Engine/Source/Runtime/Engine/Private/AI/NavigationSystemBase.cpp b/Engine/Source/Runtime/Engine/Private/AI/NavigationSystemBase.cpp
--- a/Engine/Source/Runtime/Engine/Private/AI/NavigationSystemBase.cpp
+++ b/Engine/Source/Runtime/Engine/Private/AI/NavigationSystemBase.cpp
@@ -13,6 +13,7 @@
 
 DEFINE_LOG_CATEGORY(LogNavigation);
 DEFINE_LOG_CATEGORY(LogNavigationDataBuild);
+DEFINE_LOG_CATEGORY(LogNavInvokers);
 DEFINE_LOG_CATEGORY(LogNavLink);
 DEFINE_LOG_CATEGORY(LogAStar);
 
diff --git a/Engine/Source/Runtime/NavigationSystem/Private/NavigationSystem.cpp b/Engine/Source/Runtime/NavigationSystem/Private/NavigationSystem.cpp
--- a/Engine/Source/Runtime/NavigationSystem/Private/NavigationSystem.cpp
+++ b/Engine/Source/Runtime/NavigationSystem/Private/NavigationSystem.cpp
@@ -157,6 +157,7 @@ DEFINE_STAT(STAT_DetourTilePolyClustersMemory);
 CSV_DEFINE_CATEGORY(NavigationSystem, false);
 CSV_DEFINE_CATEGORY(NavTasksDelays, true);
 CSV_DEFINE_CATEGORY(NavTasks, true);
+CSV_DEFINE_CATEGORY(NavInvokers, true);
 
 //----------------------------------------------------------------------//
 // consts
@@ -4800,8 +4801,8 @@ void UNavigationSystemV1::RegisterInvoker(AActor& Invoker, float TileGenerationR
 
 void UNavigationSystemV1::RegisterInvoker(AActor& Invoker, float TileGenerationRadius, float TileRemovalRadius, const FNavAgentSelector& Agents)
 {
-	UE_CVLOG(bGenerateNavigationOnlyAroundNavigationInvokers == false, this, LogNavigation, Warning
-		, TEXT("Trying to register %s as enforcer, but NavigationSystem is not set up for enforcer-centric generation. See GenerateNavigationOnlyAroundNavigationInvokers in NavigationSystem's properties")
+	UE_CVLOG(bGenerateNavigationOnlyAroundNavigationInvokers == false, this, LogNavInvokers, Warning
+		, TEXT("Trying to register %s as invoker, but NavigationSystem is not set up for invoker-centric generation. See GenerateNavigationOnlyAroundNavigationInvokers in NavigationSystem's properties")
 		, *Invoker.GetName());
 
 	TileGenerationRadius = FMath::Clamp(TileGenerationRadius, 0.f, BIG_NUMBER);
@@ -4814,14 +4815,34 @@ void UNavigationSystemV1::RegisterInvoker(AActor& Invoker, float TileGenerationR
 	Data.SupportedAgents = Agents;
 	Data.SupportedAgents.MarkInitialized();
 
-	UE_VLOG_CYLINDER(this, LogNavigation, Log, Invoker.GetActorLocation(), Invoker.GetActorLocation() + FVector(0, 0, 20), TileGenerationRadius, FColorList::LimeGreen
-		, TEXT("%s %.0f %.0f"), *Invoker.GetName(), TileGenerationRadius, TileRemovalRadius);
-	UE_VLOG_CYLINDER(this, LogNavigation, Log, Invoker.GetActorLocation(), Invoker.GetActorLocation() + FVector(0, 0, 20), TileRemovalRadius, FColorList::IndianRed, TEXT(""));
+	UE_SUPPRESS(LogNavInvokers, Log,
+	{
+		TStringBuilder<128> InvokerNavData;
+		for (int32 NavDataIndex = 0; NavDataIndex < NavDataSet.Num(); NavDataIndex++)
+		{
+			const ANavigationData* NavData = NavDataSet[NavDataIndex].Get();
+			if (NavData)
+			{
+				const int32 NavDataSupportedAgentIndex = GetSupportedAgentIndex(NavData);	
+				if (Data.SupportedAgents.Contains(NavDataSupportedAgentIndex))
+				{
+					InvokerNavData.Append(FString::Printf(TEXT("%s "), *NavData->GetName()));
+				}
+			}
+		}
+
+		const FString RegisterText = FString::Printf(TEXT("Register invoker r: %.0f, r area: %.0f m2, removal r: %.0f (%s %s) "),
+			TileGenerationRadius, UE_PI*FMath::Square(TileGenerationRadius/100.f), TileRemovalRadius, *Invoker.GetName(), *InvokerNavData);
+		UE_LOG(LogNavInvokers, Log, TEXT("%s"), *RegisterText);
+
+		UE_VLOG_CYLINDER(this, LogNavInvokers, Log, Invoker.GetActorLocation(), Invoker.GetActorLocation() + FVector(0, 0, 20), TileGenerationRadius, FColorList::LimeGreen, TEXT("%s"), *RegisterText);
+		UE_VLOG_CYLINDER(this, LogNavInvokers, Log, Invoker.GetActorLocation(), Invoker.GetActorLocation() + FVector(0, 0, 20), TileRemovalRadius, FColorList::IndianRed, TEXT(""));
+	});
 }
 
 void UNavigationSystemV1::UnregisterInvoker(AActor& Invoker)
 {
-	UE_VLOG(this, LogNavigation, Log, TEXT("Removing %s from enforcers list"), *Invoker.GetName());
+	UE_VLOG(this, LogNavInvokers, Log, TEXT("Removing %s from invokers list"), *Invoker.GetName());
 	Invokers.Remove(&Invoker);
 }
 
@@ -4861,12 +4882,12 @@ void UNavigationSystemV1::UpdateInvokers()
 
 #if ENABLE_VISUAL_LOG
 			const double CachingFinishTime = FPlatformTime::Seconds();
-			UE_VLOG(this, LogNavigation, Log, TEXT("Caching time %fms"), (CachingFinishTime - StartTime) * 1000);
+			UE_VLOG(this, LogNavInvokers, Log, TEXT("Caching time %fms"), (CachingFinishTime - StartTime) * 1000);
 
 			for (const auto& InvokerData : InvokerLocations)
 			{
-				UE_VLOG_CYLINDER(this, LogNavigation, Log, InvokerData.Location, InvokerData.Location + FVector(0, 0, 20), InvokerData.RadiusMax, FColorList::Blue, TEXT(""));
-				UE_VLOG_CYLINDER(this, LogNavigation, Log, InvokerData.Location, InvokerData.Location + FVector(0, 0, 20), InvokerData.RadiusMin, FColorList::CadetBlue, TEXT(""));
+				UE_VLOG_CYLINDER(this, LogNavInvokers, Log, InvokerData.Location, InvokerData.Location + FVector(0, 0, 20), InvokerData.RadiusMax, FColorList::Blue, TEXT(""));
+				UE_VLOG_CYLINDER(this, LogNavInvokers, Log, InvokerData.Location, InvokerData.Location + FVector(0, 0, 20), InvokerData.RadiusMin, FColorList::CadetBlue, TEXT(""));
 			}
 #endif // ENABLE_VISUAL_LOG
 		}
@@ -4878,12 +4899,43 @@ void UNavigationSystemV1::UpdateInvokers()
 			It->UpdateActiveTiles(InvokerLocations);
 		}
 		const double UpdateEndTime = FPlatformTime::Seconds();
-		UE_VLOG(this, LogNavigation, Log, TEXT("Marking tiles to update %fms (%d invokers)"), (UpdateEndTime - UpdateStartTime) * 1000, InvokerLocations.Num());
+		UE_VLOG(this, LogNavInvokers, Log, TEXT("Marking tiles to update %fms (%d invokers)"), (UpdateEndTime - UpdateStartTime) * 1000, InvokerLocations.Num());
 #endif
 
 		// once per second
 		NextInvokersUpdateTime = CurrentTime + ActiveTilesUpdateInterval;
 	}
+
+#if !UE_BUILD_SHIPPING
+#if CSV_PROFILER
+	if (FCsvProfiler::Get()->IsCapturing())
+	{
+		TArray<int32, TInlineAllocator<8>> InvokerCounts;
+		InvokerCounts.InsertZeroed(0, NavDataSet.Num());
+	
+		for (int32 NavDataIndex = 0; NavDataIndex < NavDataSet.Num(); NavDataIndex++)
+		{
+			const ANavigationData* NavData = NavDataSet[NavDataIndex].Get();
+			if (NavData)
+			{
+				const int32 NavDataSupportedAgentIndex = GetSupportedAgentIndex(NavData);	
+
+				for (auto ItemIterator = Invokers.CreateIterator(); ItemIterator; ++ItemIterator)
+				{
+					const FNavAgentSelector& InvokerSupportedAgents = ItemIterator->Value.SupportedAgents;
+					if (InvokerSupportedAgents.Contains(NavDataSupportedAgentIndex))
+					{
+						InvokerCounts[NavDataIndex]++;
+					}
+				}
+
+				const FString StatName = FString::Printf(TEXT("InvokerCount_%s"), *NavData->GetName()); 
+				FCsvProfiler::RecordCustomStat(*StatName, CSV_CATEGORY_INDEX(NavInvokers), InvokerCounts[NavDataIndex], ECsvCustomStatOp::Set);
+			}
+		}		
+	}
+#endif // CSV_PROFILER
+#endif // !UE_BUILD_SHIPPING
 }
 
 void UNavigationSystemV1::DirtyTilesInBuildBounds()
