Add NavMeshResolution to UNavModifierComponent.

diff --git a/Engine/Source/Runtime/NavigationSystem/Private/NavModifierComponent.cpp b/Engine/Source/Runtime/NavigationSystem/Private/NavModifierComponent.cpp
--- a/Engine/Source/Runtime/NavigationSystem/Private/NavModifierComponent.cpp
+++ b/Engine/Source/Runtime/NavigationSystem/Private/NavModifierComponent.cpp
@@ -28,6 +28,7 @@ UNavModifierComponent::UNavModifierComponent(const FObjectInitializer& ObjectIni
 	AreaClass = UNavArea_Null::StaticClass();
 	FailsafeExtent = FVector(100, 100, 100);
 	bIncludeAgentHeight = true;
+	NavMeshResolution = ENavigationDataResolution::Invalid;
 }
 
 void UNavModifierComponent::OnRegister()
@@ -170,6 +171,8 @@ void UNavModifierComponent::GetNavigationData(FNavigationRelevantData& Data) con
 	{
 		Data.Modifiers.Add(FAreaNavModifier(ComponentBounds[Idx].Box, FTransform(ComponentBounds[Idx].Quat), AreaClass).SetIncludeAgentHeight(bIncludeAgentHeight));
 	}
+
+	Data.Modifiers.SetNavMeshResolution(NavMeshResolution);
 }
 
 void UNavModifierComponent::SetAreaClass(TSubclassOf<UNavArea> NewAreaClass)
diff --git a/Engine/Source/Runtime/NavigationSystem/Public/NavModifierComponent.h b/Engine/Source/Runtime/NavigationSystem/Public/NavModifierComponent.h
--- a/Engine/Source/Runtime/NavigationSystem/Public/NavModifierComponent.h
+++ b/Engine/Source/Runtime/NavigationSystem/Public/NavModifierComponent.h
@@ -23,6 +23,10 @@ class NAVIGATIONSYSTEM_API UNavModifierComponent : public UNavRelevantComponent
 	UPROPERTY(EditAnywhere, Category = Navigation)
 	FVector FailsafeExtent;
 
+	/** Experimental: Indicates which navmesh resolution should be used around the actor. */
+	UPROPERTY(EditAnywhere, Category = Navigation, AdvancedDisplay)
+	ENavigationDataResolution NavMeshResolution;
+
 	/** Setting to 'true' will result in expanding lower bounding box of the nav 
 	 *	modifier by agent's height, before applying to navmesh */
 	UPROPERTY(config, EditAnywhere, Category = Navigation)
