Fix navmesh resolution issue when using LazyModifiersGathering by reading the resolution after GatherNavigationDataGeometry.
Addition of logs when preparing and gathering geometry.

diff --git a/Engine/Source/Runtime/NavigationSystem/Private/NavMesh/RecastNavMeshGenerator.cpp b/Engine/Source/Runtime/NavigationSystem/Private/NavMesh/RecastNavMeshGenerator.cpp
--- a/Engine/Source/Runtime/NavigationSystem/Private/NavMesh/RecastNavMeshGenerator.cpp
+++ b/Engine/Source/Runtime/NavigationSystem/Private/NavMesh/RecastNavMeshGenerator.cpp
@@ -23,6 +23,7 @@
 #include "Chaos/HeightField.h"
 #include "Chaos/TriangleMeshImplicitObject.h"
 #include "NavMesh/PImplRecastNavMesh.h"
+#include "VisualLogger/VisualLogger.h"
 
 // recast includes
 #include "Detour/DetourNavMeshBuilder.h"
@@ -1966,13 +1967,13 @@ void FRecastTileGenerator::PrepareGeometrySources(const FRecastNavMeshGenerator&
 	bUpdateGeometry = bGeometryChanged;
 	
 	const ARecastNavMesh* const OwnerNav = ParentGenerator.GetOwner();
-	const bool bUseVirtualGeometryFilteringAndDirtying = OwnerNav != nullptr && OwnerNav->bUseVirtualGeometryFilteringAndDirtying;
+	const bool bUseVirtualGeometryFilteringAndDirtying = OwnerNav->bUseVirtualGeometryFilteringAndDirtying;
 
 	ENavigationDataResolution HighestResolution = ENavigationDataResolution::Low;
 	bool bNewResolutionFound = false;
 	
 	NavOctreeInstance->FindElementsWithBoundsTest(ParentGenerator.GrowBoundingBox(TileBB, /*bIncludeAgentHeight*/ false),
-		[&HighestResolution, &bNewResolutionFound, &ParentGenerator, this, bGeometryChanged, bUseVirtualGeometryFilteringAndDirtying](const FNavigationOctreeElement& Element)
+		[&HighestResolution, &bNewResolutionFound, &ParentGenerator, OwnerNav, this, bGeometryChanged, bUseVirtualGeometryFilteringAndDirtying](const FNavigationOctreeElement& Element)
 	{
 		const bool bShouldUse = bUseVirtualGeometryFilteringAndDirtying ?
 			ParentGenerator.ShouldGenerateGeometryForOctreeElement(Element, NavDataConfig) :
@@ -1985,6 +1986,8 @@ void FRecastTileGenerator::PrepareGeometrySources(const FRecastNavMeshGenerator&
 				Element.Data->Modifiers.HasMetaAreas() == true || 
 				Element.Data->Modifiers.IsEmpty() == false)
 			{
+				UE_SUPPRESS(LogNavigation, VeryVerbose, UE_VLOG_UELOG(OwnerNav, LogNavigation, VeryVerbose, TEXT("PrepareGeometrySources, adding: %s"), *GetNameSafe(Element.GetOwner())));
+				
 				// Keep highest resolution that is not the default.
 				const ENavigationDataResolution Resolution = Element.Data->Modifiers.GetNavMeshResolution();
 				if (Resolution != ENavigationDataResolution::Invalid)
@@ -2022,12 +2025,12 @@ void FRecastTileGenerator::GatherGeometry(const FRecastNavMeshGenerator& ParentG
 
 	TArray<TSharedRef<FNavigationRelevantData, ESPMode::ThreadSafe>> RelevantDataArray;
 
+	UE_SUPPRESS(LogNavigation, VeryVerbose, UE_VLOG_BOX(OwnerNav, LogNavigation, VeryVerbose, TileBB, FColor::White, TEXT("Tile (%i, %i)"), TileX, TileY));	
+
 	const FBox NewBounds = ParentGenerator.GrowBoundingBox(TileBB, /*bIncludeAgentHeight*/ false);
-	ENavigationDataResolution HighestResolution = ENavigationDataResolution::Low;
-	bool bNewResolutionFound = false;
 	
 	NavigationOctree->FindElementsWithBoundsTest(NewBounds,
-		[&HighestResolution, &bNewResolutionFound, &RelevantDataArray, &OwnerNavDataConfig, &ParentGenerator, this, NavSys, bGeometryChanged, bUseVirtualGeometryFilteringAndDirtying](const FNavigationOctreeElement& Element)
+		[&RelevantDataArray, &OwnerNavDataConfig, &ParentGenerator, this, NavSys, bGeometryChanged, bUseVirtualGeometryFilteringAndDirtying](const FNavigationOctreeElement& Element)
 	{
 		const bool bShouldUse = bUseVirtualGeometryFilteringAndDirtying ?
 			ParentGenerator.ShouldGenerateGeometryForOctreeElement(Element, OwnerNavDataConfig) :
@@ -2036,27 +2039,31 @@ void FRecastTileGenerator::GatherGeometry(const FRecastNavMeshGenerator& ParentG
 		{
 			RelevantDataArray.Add(Element.Data);
 
-			// Keep highest resolution that is not the default.
-			const ENavigationDataResolution Resolution = Element.Data->Modifiers.GetNavMeshResolution();
-			if (Resolution != ENavigationDataResolution::Invalid)
-			{
-				HighestResolution = FMath::Max(HighestResolution, Resolution);
-				bNewResolutionFound = true;
-			}
 		}
 	});
 
+	ENavigationDataResolution HighestResolution = ENavigationDataResolution::Low;
+	bool bNewResolutionFound = false;
+	
+	for (TSharedRef<FNavigationRelevantData, ESPMode::ThreadSafe>& ElementData : RelevantDataArray)
+	{
+		GatherNavigationDataGeometry(ElementData, *NavSys, OwnerNavDataConfig, bGeometryChanged);
+
+		// Keep highest resolution that is not the default.
+		const ENavigationDataResolution Resolution = ElementData->Modifiers.GetNavMeshResolution();
+		if (Resolution != ENavigationDataResolution::Invalid)
+		{
+			HighestResolution = FMath::Max(HighestResolution, Resolution);
+			bNewResolutionFound = true;
+		}
+	}
+
 	check(HighestResolution != ENavigationDataResolution::Invalid);
 	if (bNewResolutionFound && ParentGenerator.GetOwner()->NavMeshResolutionParams[(uint8)HighestResolution].IsValid())
 	{
 		// Update the TileConfig
 		ParentGenerator.SetupTileConfig(HighestResolution, TileConfig);
 	}
-	
-	for(TSharedRef<FNavigationRelevantData, ESPMode::ThreadSafe> ElementData : RelevantDataArray)
-	{
-		GatherNavigationDataGeometry(ElementData, *NavSys, OwnerNavDataConfig, bGeometryChanged);
-	}
 }
 
 void FRecastTileGenerator::GatherNavigationDataGeometry(const TSharedRef<FNavigationRelevantData, ESPMode::ThreadSafe>& ElementData, UNavigationSystemV1& NavSys, const FNavDataConfig& OwnerNavDataConfig, const bool bGeometryChanged)
@@ -5883,7 +5890,7 @@ void FRecastNavMeshGenerator::MarkDirtyTiles(const TArray<FNavigationDirtyArea>&
 			continue;
 		}
 
-		UE_VLOG_BOX(OwnerNav, LogNavigation, VeryVerbose, DirtyArea.Bounds, FColor::Blue, TEXT("DirtyArea"));
+		UE_VLOG_BOX(OwnerNav, LogNavigation, VeryVerbose, DirtyArea.Bounds, FColor::Blue, TEXT("DirtyArea %s"), *GetNameSafe(DirtyArea.OptionalSourceObject.Get()));
 		
 		// (if bUseVirtualGeometryFilteringAndDirtying is true) Ignore dirty areas flagged by a source object that is not supposed to apply to this navmesh
 		if (bUseVirtualGeometryFilteringAndDirtying && NavSys && NavOctreeInstance && NavDataConfig)
