[Navmesh] Addition of a warning when the navmesh version is compatible but not at the latest.

diff --git a/Engine/Source/Runtime/NavigationSystem/Private/NavMesh/RecastNavMesh.cpp b/Engine/Source/Runtime/NavigationSystem/Private/NavMesh/RecastNavMesh.cpp
--- a/Engine/Source/Runtime/NavigationSystem/Private/NavMesh/RecastNavMesh.cpp
+++ b/Engine/Source/Runtime/NavigationSystem/Private/NavMesh/RecastNavMesh.cpp
@@ -1046,14 +1046,19 @@ void ARecastNavMesh::Serialize( FArchive& Ar )
 			CleanUpAndMarkPendingKill();
 		};
 
+		if (NavMeshVersion >= NAVMESHVER_MIN_COMPATIBLE && NavMeshVersion < NAVMESHVER_LATEST)
+		{
+			UE_LOG(LogNavigation, Warning, TEXT("%s: ARecastNavMesh: Navmesh version %d is compatible but not at latest (%d). Rebuild navmesh to use the latest features. \n"), *GetFullName(), NavMeshVersion, NAVMESHVER_LATEST);
+		}
+
 		if (NavMeshVersion < NAVMESHVER_MIN_COMPATIBLE)
 		{
-			UE_LOG(LogNavigation, Warning, TEXT("%s: ARecastNavMesh: Nav mesh version %d < Min compatible %d. Nav mesh needs to be rebuilt. \n"), *GetFullName(), NavMeshVersion, NAVMESHVER_MIN_COMPATIBLE);
+			UE_LOG(LogNavigation, Warning, TEXT("%s: ARecastNavMesh: Navmesh version %d < Min compatible %d. Navmesh needs to be rebuilt. \n"), *GetFullName(), NavMeshVersion, NAVMESHVER_MIN_COMPATIBLE);
 			CleanUpBadVersion();
 		}
 		else if (NavMeshVersion > NAVMESHVER_LATEST)
 		{
-			UE_LOG(LogNavigation, Warning, TEXT("%s: ARecastNavMesh: Nav mesh version %d > NAVMESHVER_LATEST %d. Newer nav mesh should not be loaded by older code. At a minimum the nav mesh needs to be rebuilt. \n"), *GetFullName(), NavMeshVersion, NAVMESHVER_LATEST);
+			UE_LOG(LogNavigation, Warning, TEXT("%s: ARecastNavMesh: Navmesh version %d > NAVMESHVER_LATEST %d. Newer navmesh should not be loaded by older code. At a minimum the nav mesh needs to be rebuilt. \n"), *GetFullName(), NavMeshVersion, NAVMESHVER_LATEST);
 			CleanUpBadVersion();
 		}
 		else if (RecastNavMeshSizeBytes > 4)
