Integrate CL 23032335 from UE5Main

Navigation display improvements
Addition of tile building statistics
Addition of partitioning type, navdata count and triangle count
Addition of profile tags in the tile generation

diff --git a/Engine/Source/Runtime/NavigationSystem/Private/NavMesh/RecastNavMeshGenerator.cpp b/Engine/Source/Runtime/NavigationSystem/Private/NavMesh/RecastNavMeshGenerator.cpp
--- a/Engine/Source/Runtime/NavigationSystem/Private/NavMesh/RecastNavMeshGenerator.cpp
+++ b/Engine/Source/Runtime/NavigationSystem/Private/NavMesh/RecastNavMeshGenerator.cpp
@@ -2744,6 +2744,8 @@ void FRecastTileGenerator::RasterizeGeometryRecast(FNavMeshBuildContext& BuildCo
 	}
 
 #if RECAST_INTERNAL_DEBUG_DATA
+	BuildContext.InternalDebugData.TriangleCount += NumFaces;
+
 	if (IsTileDebugActive() && TileDebugSettings.bCollisionGeometry) 
 	{
 		TArray<FVector::FReal> Normals;
