Addition of AgentMaxStepHeight into FNavMeshResolutionParam

diff --git a/Engine/Config/BaseEngine.ini b/Engine/Config/BaseEngine.ini
--- a/Engine/Config/BaseEngine.ini
+++ b/Engine/Config/BaseEngine.ini
@@ -2626,9 +2626,9 @@ DefaultDrawDistance=5000.0
 TileSizeUU=1000.f
 CellSize=19.f
 CellHeight=10.f
-NavMeshResolutionParams[0]=(CellSize=38.000000,CellHeight=10.f)
-NavMeshResolutionParams[1]=(CellSize=19.000000,CellHeight=10.f)
-NavMeshResolutionParams[2]=(CellSize=19.000000,CellHeight=10.f)
+NavMeshResolutionParams[0]=(CellSize=38.000000,CellHeight=10.f,AgentMaxStepHeight=35.f)
+NavMeshResolutionParams[1]=(CellSize=19.000000,CellHeight=10.f,AgentMaxStepHeight=35.f)
+NavMeshResolutionParams[2]=(CellSize=19.000000,CellHeight=10.f,AgentMaxStepHeight=35.f)
 AgentRadius=34.f
 AgentHeight=144.f
 AgentMaxHeight=160.f
diff --git a/Engine/Source/Runtime/NavigationSystem/Private/NavLinkRenderingComponent.cpp b/Engine/Source/Runtime/NavigationSystem/Private/NavLinkRenderingComponent.cpp
--- a/Engine/Source/Runtime/NavigationSystem/Private/NavLinkRenderingComponent.cpp
+++ b/Engine/Source/Runtime/NavigationSystem/Private/NavLinkRenderingComponent.cpp
@@ -200,9 +200,10 @@ void FNavLinkRenderingProxy::GetDynamicMeshElements(const TArray<const FSceneVie
 				if (NavMesh != NULL)
 				{
 					AgentMask = NavMesh->IsDrawingEnabled() ? AgentMask | (1 << DataIndex) : AgentMask;
-					if (NavMesh->AgentMaxStepHeight > 0 && NavMesh->IsDrawingEnabled())
+					const float AgentMaxStepHeight = NavMesh->GetAgentMaxStepHeight(ENavigationDataResolution::Default);
+					if (AgentMaxStepHeight > 0 && NavMesh->IsDrawingEnabled())
 					{
-						StepHeights.Add(NavMesh->AgentMaxStepHeight);
+						StepHeights.Add(AgentMaxStepHeight);
 					}
 				}
 			}
diff --git a/Engine/Source/Runtime/NavigationSystem/Private/NavMesh/NavMeshRenderingComponent.cpp b/Engine/Source/Runtime/NavigationSystem/Private/NavMesh/NavMeshRenderingComponent.cpp
--- a/Engine/Source/Runtime/NavigationSystem/Private/NavMesh/NavMeshRenderingComponent.cpp
+++ b/Engine/Source/Runtime/NavigationSystem/Private/NavMesh/NavMeshRenderingComponent.cpp
@@ -580,11 +580,13 @@ void FNavMeshSceneProxyData::GatherData(const ARecastNavMesh* NavMesh, int32 InN
 			const FString GenerationMode = Mode == ERuntimeGenerationType::Static ? TEXT("Static") :
 				(Mode == ERuntimeGenerationType::Dynamic ? TEXT("Dynamic") : (Mode == ERuntimeGenerationType::DynamicModifiersOnly ? TEXT("DynamicModifersOnly") : TEXT("Unknown")));
 			DebugLabels.Add(FDebugText(FString::Printf(TEXT("%s (%s%s)"), *NavMesh->GetName(), NavMesh->bIsWorldPartitioned ? TEXT("WP ") : TEXT(""), *GenerationMode)));
-			DebugLabels.Add(FDebugText(FString::Printf(TEXT("AgentRadius %0.1f, AgentHeight %0.1f, CellSizes %0.1f/%0.1f/%0.1f, CellHeights %0.1f/%0.1f/%0.1f (low/default/high)"),
-				NavMesh->AgentRadius,
-				NavMesh->AgentHeight,
+			DebugLabels.Add(FDebugText(FString::Printf(TEXT("AgentRadius %0.1f, AgentHeight %0.1f"), NavMesh->AgentRadius, NavMesh->AgentHeight)));
+			DebugLabels.Add(FDebugText(FString::Printf(
+				TEXT("CellSizes %0.1f/%0.1f/%0.1f, CellHeights %0.1f/%0.1f/%0.1f, AgentMaxStepHeight %0.1f/%0.1f/%0.1f (low/default/high)"),
 				NavMesh->GetCellSize(ENavigationDataResolution::Low), NavMesh->GetCellSize(ENavigationDataResolution::Default), NavMesh->GetCellSize(ENavigationDataResolution::High),
-				NavMesh->GetCellHeight(ENavigationDataResolution::Low), NavMesh->GetCellHeight(ENavigationDataResolution::Default), NavMesh->GetCellHeight(ENavigationDataResolution::High))));
+				NavMesh->GetCellHeight(ENavigationDataResolution::Low), NavMesh->GetCellHeight(ENavigationDataResolution::Default), NavMesh->GetCellHeight(ENavigationDataResolution::High),
+				NavMesh->GetAgentMaxStepHeight(ENavigationDataResolution::Low), NavMesh->GetAgentMaxStepHeight(ENavigationDataResolution::Default), NavMesh->GetAgentMaxStepHeight(ENavigationDataResolution::High)
+				)));
 			DebugLabels.Add(FDebugText(FString::Printf(TEXT("Region part %s, Layer part %s"), *GetPartitioningString(NavMesh->RegionPartitioning), *GetPartitioningString(NavMesh->LayerPartitioning))));
 			DebugLabels.Add(FDebugText(TEXT(""))); // empty line
 
@@ -751,11 +753,11 @@ void FNavMeshSceneProxyData::GatherData(const ARecastNavMesh* NavMesh, int32 InN
 						if (Link.Direction && (Link.ValidEnds & FRecastDebugGeometry::OMLE_Left) == 0)
 						{
 							// left end invalid - mark it
-							FNavMeshRenderingHelpers::DrawWireCylinder(NavLinkLines, V0, FVector(1, 0, 0), FVector(0, 1, 0), FVector(0, 0, 1), NavMeshRenderColor_OffMeshConnectionInvalid, Link.Radius, NavMesh->AgentMaxStepHeight, 16, 0, DefaultEdges_LineThickness);
+							FNavMeshRenderingHelpers::DrawWireCylinder(NavLinkLines, V0, FVector(1, 0, 0), FVector(0, 1, 0), FVector(0, 0, 1), NavMeshRenderColor_OffMeshConnectionInvalid, Link.Radius, NavMesh->GetAgentMaxStepHeight(ENavigationDataResolution::Default), 16, 0, DefaultEdges_LineThickness);
 						}
 						if ((Link.ValidEnds & FRecastDebugGeometry::OMLE_Right) == 0)
 						{
-							FNavMeshRenderingHelpers::DrawWireCylinder(NavLinkLines, V1, FVector(1, 0, 0), FVector(0, 1, 0), FVector(0, 0, 1), NavMeshRenderColor_OffMeshConnectionInvalid, Link.Radius, NavMesh->AgentMaxStepHeight, 16, 0, DefaultEdges_LineThickness);
+							FNavMeshRenderingHelpers::DrawWireCylinder(NavLinkLines, V1, FVector(1, 0, 0), FVector(0, 1, 0), FVector(0, 0, 1), NavMeshRenderColor_OffMeshConnectionInvalid, Link.Radius, NavMesh->GetAgentMaxStepHeight(ENavigationDataResolution::Default), 16, 0, DefaultEdges_LineThickness);
 						}
 					}
 				}
@@ -971,6 +973,7 @@ void FNavMeshSceneProxyData::GatherData(const ARecastNavMesh* NavMesh, int32 InN
 						const double BuildTimeMs = DebugData->BuildTime * SecToMs;
 						const double CompressedLayerTimeMs = DebugData->BuildCompressedLayerTime * SecToMs;
 						const double NavigationDataTimeMs = DebugData->BuildNavigationDataTime * SecToMs;
+						const ENavigationDataResolution Resolution = DebugData->Resolution;
 
 						const double Range = NavMeshGeometry.MaxTileBuildTime - NavMeshGeometry.MinTileBuildTime;
 						int32 Rank = 0;
@@ -980,8 +983,9 @@ void FNavMeshSceneProxyData::GatherData(const ARecastNavMesh* NavMesh, int32 InN
 							Rank = static_cast<int32>(FMath::Clamp<double>(RankCalc, 0., FRecastDebugGeometry::BuildTimeBucketsCount-1));
 						}
 
-						DebugLabels.Add(FDebugText(Pair.Value + NavMeshDrawOffset, FString::Printf(TEXT("%.2f ms\n%.2f comp + %.2f nav\n%i tri"),
-							BuildTimeMs, CompressedLayerTimeMs, NavigationDataTimeMs,
+						const FString ResolutionText = StaticEnum<ENavigationDataResolution>()->GetNameStringByValue((int64)Resolution);
+						DebugLabels.Add(FDebugText(Pair.Value + NavMeshDrawOffset, FString::Printf(TEXT("%s res\n %.2f ms\n%.2f comp + %.2f nav\n%i tri"),
+							*ResolutionText, BuildTimeMs, CompressedLayerTimeMs, NavigationDataTimeMs,
 							DebugData->TriangleCount)));
 					}
 				}
diff --git a/Engine/Source/Runtime/NavigationSystem/Private/NavMesh/PImplRecastNavMesh.cpp b/Engine/Source/Runtime/NavigationSystem/Private/NavMesh/PImplRecastNavMesh.cpp
--- a/Engine/Source/Runtime/NavigationSystem/Private/NavMesh/PImplRecastNavMesh.cpp
+++ b/Engine/Source/Runtime/NavigationSystem/Private/NavMesh/PImplRecastNavMesh.cpp
@@ -584,7 +584,7 @@ void FPImplRecastNavMesh::Serialize( FArchive& Ar, int32 NavMeshVersion )
 	{
 		Params.walkableHeight = NavMeshOwner->AgentHeight;
 		Params.walkableRadius = NavMeshOwner->AgentRadius;
-		Params.walkableClimb = NavMeshOwner->AgentMaxStepHeight;
+		Params.walkableClimb = NavMeshOwner->GetAgentMaxStepHeight(ENavigationDataResolution::Default);
 		const float DefaultQuantFactor =  1.f / NavMeshOwner->GetCellSize(ENavigationDataResolution::Default);
 		for(uint8 Index = 0; Index < (uint8)ENavigationDataResolution::MAX; Index++)
 		{
diff --git a/Engine/Source/Runtime/NavigationSystem/Private/NavMesh/RecastNavMesh.cpp b/Engine/Source/Runtime/NavigationSystem/Private/NavMesh/RecastNavMesh.cpp
--- a/Engine/Source/Runtime/NavigationSystem/Private/NavMesh/RecastNavMesh.cpp
+++ b/Engine/Source/Runtime/NavigationSystem/Private/NavMesh/RecastNavMesh.cpp
@@ -433,7 +433,7 @@ FRecastNavMeshGenerationProperties::FRecastNavMeshGenerationProperties(const ARe
 	AgentRadius = RecastNavMesh.AgentRadius;
 	AgentHeight = RecastNavMesh.AgentHeight;
 	AgentMaxSlope = RecastNavMesh.AgentMaxSlope;
-	AgentMaxStepHeight = RecastNavMesh.AgentMaxStepHeight;
+	AgentMaxStepHeight = RecastNavMesh.GetAgentMaxStepHeight(ENavigationDataResolution::Default); //FRecastNavMeshGenerationProperties is getting deprecated 
 	MinRegionArea = RecastNavMesh.MinRegionArea;
 	MergeRegionSize = RecastNavMesh.MergeRegionSize;
 	MaxSimplificationError = RecastNavMesh.MaxSimplificationError;
@@ -655,6 +655,15 @@ void ARecastNavMesh::PostLoad()
 			SetCellHeight((ENavigationDataResolution)i, CellHeight);
 		}
 	}
+
+	// If needed, initialize AgentMaxStepHeight from the deprecated value.
+	if (NavMeshVersion < NAVMESHVER_TILE_RESOLUTIONS_AGENTMAXSTEPHEIGHT)
+	{
+		for (int i = 0; i < (uint8)ENavigationDataResolution::MAX; ++i)
+		{
+			SetAgentMaxStepHeight((ENavigationDataResolution)i, AgentMaxStepHeight);
+		}
+	}
 	PRAGMA_ENABLE_DEPRECATION_WARNINGS
 	
 	for (uint8 Index = 0; Index < (uint8)ENavigationDataResolution::MAX; Index++)
@@ -763,6 +772,16 @@ void ARecastNavMesh::PostInitProperties()
 
 				NavMeshResolutionParams[i].CellHeight = DefaultObjectCellHeight;
 			}
+
+			const float CurrentAgentMaxStepHeight = NavMeshResolutionParams[i].AgentMaxStepHeight;
+			const float DefaultObjectAgentMaxStepHeight = DefOb->NavMeshResolutionParams[i].AgentMaxStepHeight;
+			if (CurrentAgentMaxStepHeight != DefaultObjectAgentMaxStepHeight)
+			{
+				UE_LOG(LogNavigation, Warning, TEXT("%s param: AgentMaxStepHeight(%f) differs from config settings, forcing value %f so it can be used with voxel cache!"),
+					*GetNameSafe(this), CurrentAgentMaxStepHeight, DefaultObjectAgentMaxStepHeight);
+
+				NavMeshResolutionParams[i].AgentMaxStepHeight = DefaultObjectAgentMaxStepHeight;
+			}			
 		}
 
 		if (AgentMaxSlope != DefOb->AgentMaxSlope)
@@ -772,14 +791,6 @@ void ARecastNavMesh::PostInitProperties()
 
 			AgentMaxSlope = DefOb->AgentMaxSlope;
 		}
-
-		if (AgentMaxStepHeight != DefOb->AgentMaxStepHeight)
-		{
-			UE_LOG(LogNavigation, Warning, TEXT("%s param: AgentMaxStepHeight(%f) differs from config settings, forcing value %f so it can be used with voxel cache!"),
-				*GetNameSafe(this), AgentMaxStepHeight, DefOb->AgentMaxStepHeight);
-
-			AgentMaxStepHeight = DefOb->AgentMaxStepHeight;
-		}
 	}
 	
 #if WITH_EDITORONLY_DATA
@@ -1120,7 +1131,11 @@ void ARecastNavMesh::SetConfig(const FNavDataConfig& Src)
 
 	if (Src.HasStepHeightOverride())
 	{
-		AgentMaxStepHeight = Src.AgentStepHeight;
+		// If there is an override, apply it to all resolutions
+		for (int32 Index = 0; Index < (int32)ENavigationDataResolution::MAX; Index++)
+		{
+			SetAgentMaxStepHeight((ENavigationDataResolution)Index, Src.AgentStepHeight);
+		}
 	}
 }
 
@@ -1129,7 +1144,7 @@ void ARecastNavMesh::FillConfig(FNavDataConfig& Dest)
 	Dest = NavDataConfig;
 	Dest.AgentHeight = AgentHeight;
 	Dest.AgentRadius = AgentRadius;
-	Dest.AgentStepHeight = AgentMaxStepHeight;
+	Dest.AgentStepHeight = GetAgentMaxStepHeight(ENavigationDataResolution::Default);
 }
 
 void ARecastNavMesh::BeginBatchQuery() const
@@ -3344,7 +3359,9 @@ void ARecastNavMesh::UpdateGenerationProperties(const FRecastNavMeshGenerationPr
 	AgentRadius = GenerationProps.AgentRadius;
 	AgentHeight = GenerationProps.AgentHeight;
 	AgentMaxSlope = GenerationProps.AgentMaxSlope;
+	PRAGMA_DISABLE_DEPRECATION_WARNINGS
 	AgentMaxStepHeight = GenerationProps.AgentMaxStepHeight;
+	PRAGMA_ENABLE_DEPRECATION_WARNINGS
 	MinRegionArea = GenerationProps.MinRegionArea;
 	MergeRegionSize = GenerationProps.MergeRegionSize;
 	MaxSimplificationError = GenerationProps.MaxSimplificationError;
diff --git a/Engine/Source/Runtime/NavigationSystem/Private/NavMesh/RecastNavMeshGenerator.cpp b/Engine/Source/Runtime/NavigationSystem/Private/NavMesh/RecastNavMeshGenerator.cpp
--- a/Engine/Source/Runtime/NavigationSystem/Private/NavMesh/RecastNavMeshGenerator.cpp
+++ b/Engine/Source/Runtime/NavigationSystem/Private/NavMesh/RecastNavMeshGenerator.cpp
@@ -2567,6 +2567,7 @@ bool FRecastTileGenerator::GenerateTile()
 	BuildContext.InternalDebugData.BuildTime = EndStamp - StartStamp;
 	BuildContext.InternalDebugData.BuildCompressedLayerTime = PostCompressLayerStamp - StartStamp;
 	BuildContext.InternalDebugData.BuildNavigationDataTime = EndStamp - PostCompressLayerStamp;
+	BuildContext.InternalDebugData.Resolution = TileConfig.TileResolution;
 #endif // RECAST_INTERNAL_DEBUG_DATA
 	
 	// it's possible to have valid generation with empty resulting tile (no navigable geometry in tile)
@@ -4625,8 +4626,9 @@ void FRecastNavMeshGenerator::SetupTileConfig(const ENavigationDataResolution Ti
 	ensure(GetConfig().cs == GetOwner()->GetCellSize(ENavigationDataResolution::Default));
 	const float CellSize = GetOwner()->GetCellSize(TileResolution);
 	const float CellHeight = GetOwner()->GetCellHeight(TileResolution);
+	const float AgentMaxStepHeight = GetOwner()->GetAgentMaxStepHeight(TileResolution);
 	
-	// Update all settings that depends directly on indirectly of the CellSize
+	// Update all settings that depends directly or indirectly of the CellSize
 	OutConfig.TileResolution = TileResolution;
 	OutConfig.cs = CellSize;
 	OutConfig.walkableRadius = FMath::CeilToInt(DestNavMesh->AgentRadius / CellSize);
@@ -4644,10 +4646,13 @@ void FRecastNavMeshGenerator::SetupTileConfig(const ENavigationDataResolution Ti
 	OutConfig.regionChunkSize = FMath::Max(1, OutConfig.tileSize / FMath::Max(1, DestNavMesh->LayerChunkSplits));
 	OutConfig.TileCacheChunkSize = FMath::Max(1, OutConfig.tileSize / FMath::Max(1, DestNavMesh->RegionChunkSplits));
 
-	// Update all settings that depends directly on indirectly of the CellHeight
+	// Update all settings that depends directly or indirectly of the CellHeight
 	OutConfig.ch = CellHeight;
 	OutConfig.walkableHeight = DestNavMesh->bMarkLowHeightAreas ? 1 : FMath::CeilToInt(DestNavMesh->AgentHeight / CellHeight);
-	OutConfig.walkableClimb = FMath::CeilToInt(DestNavMesh->AgentMaxStepHeight / CellHeight);
+	OutConfig.walkableClimb = FMath::CeilToInt(AgentMaxStepHeight / CellHeight);
+
+	// Update all settings that depends directly or indirectly of AgentMaxStepHeight
+	OutConfig.AgentMaxClimb = AgentMaxStepHeight;
 }
 
 void FRecastNavMeshGenerator::ConfigureBuildProperties(FRecastBuildConfig& OutConfig)
@@ -4657,7 +4662,7 @@ void FRecastNavMeshGenerator::ConfigureBuildProperties(FRecastBuildConfig& OutCo
 	const float CellHeight = DestNavMesh->GetCellHeight(ENavigationDataResolution::Default);
 	const float AgentHeight = DestNavMesh->AgentHeight;
 	const float AgentMaxSlope = DestNavMesh->AgentMaxSlope;
-	const float AgentMaxClimb = DestNavMesh->AgentMaxStepHeight;
+	const float AgentMaxClimb = DestNavMesh->GetAgentMaxStepHeight(ENavigationDataResolution::Default);
 	const float AgentRadius = DestNavMesh->AgentRadius;
 
 	OutConfig.Reset();
diff --git a/Engine/Source/Runtime/NavigationSystem/Public/NavMesh/RecastInternalDebugData.h b/Engine/Source/Runtime/NavigationSystem/Public/NavMesh/RecastInternalDebugData.h
--- a/Engine/Source/Runtime/NavigationSystem/Public/NavMesh/RecastInternalDebugData.h
+++ b/Engine/Source/Runtime/NavigationSystem/Public/NavMesh/RecastInternalDebugData.h
@@ -8,11 +8,12 @@
 #include "Math/UnrealMathSSE.h"
 #include "Math/Vector.h"
 #include "NavMesh/RecastHelpers.h"
+#include "AI/Navigation/NavigationTypes.h"
 
 struct FRecastInternalDebugData : public duDebugDraw
 {
-	duDebugDrawPrimitives CurrentPrim;
-	int32 FirstVertexIndex;
+	duDebugDrawPrimitives CurrentPrim = DU_DRAW_POINTS;
+	int32 FirstVertexIndex = 0;
 
 	TArray<uint32> TriangleIndices;
 	TArray<FVector> TriangleVertices;
@@ -29,6 +30,7 @@ struct FRecastInternalDebugData : public duDebugDraw
 	double BuildNavigationDataTime = 0.;
 
 	uint32 TriangleCount = 0;
+	ENavigationDataResolution Resolution = ENavigationDataResolution::Default;
 	
 	FRecastInternalDebugData() {}
 	virtual ~FRecastInternalDebugData() override {}
diff --git a/Engine/Source/Runtime/NavigationSystem/Public/NavMesh/RecastNavMesh.h b/Engine/Source/Runtime/NavigationSystem/Public/NavMesh/RecastNavMesh.h
--- a/Engine/Source/Runtime/NavigationSystem/Public/NavMesh/RecastNavMesh.h
+++ b/Engine/Source/Runtime/NavigationSystem/Public/NavMesh/RecastNavMesh.h
@@ -605,15 +605,19 @@ struct NAVIGATIONSYSTEM_API FNavMeshResolutionParam
 {
 	GENERATED_BODY()
 
-	bool IsValid() const { return CellSize > 0.f && CellHeight > 0.f; }
+	bool IsValid() const { return CellSize > 0.f && CellHeight > 0.f && AgentMaxStepHeight > 0.f; }
 	
-	/** horizontal size of voxelization cell */
+	/** Horizontal size of voxelization cell */
 	UPROPERTY(EditAnywhere, Category = Generation, config, meta = (ClampMin = "1.0", ClampMax = "1024.0"))
 	float CellSize = 25.f;
 
-	/** vertical size of voxelization cell */
+	/** Vertical size of voxelization cell */
 	UPROPERTY(EditAnywhere, Category = Generation, config, meta = (ClampMin = "1.0", ClampMax = "1024.0"))
 	float CellHeight = 10.f;
+
+	/** Largest vertical step the agent can perform */
+	UPROPERTY(EditAnywhere, Category = Generation, config, meta = (ClampMin = "0.0"))
+	float AgentMaxStepHeight = 35.f;
 };
 
 UCLASS(config=Engine, defaultconfig, hidecategories=(Input,Rendering,Tags,Transformation,Actor,Layers,Replication), notplaceable)
@@ -758,7 +762,8 @@ class NAVIGATIONSYSTEM_API ARecastNavMesh : public ANavigationData
 	float AgentMaxSlope;
 
 	/** Largest vertical step the agent can perform */
-	UPROPERTY(EditAnywhere, Category = Generation, config, meta = (ClampMin = "0.0"))
+	UE_DEPRECATED(5.3, "Set the AgentMaxStepHeight for the required navmesh resolutions in NavMeshResolutionParams.")
+	UPROPERTY(config)
 	float AgentMaxStepHeight;
 
 	/* The minimum dimension of area. Areas smaller than this will be discarded */
@@ -1078,6 +1083,12 @@ public:
 
 	/** Set the CellHeight for the given resolution. */
 	void SetCellHeight(const ENavigationDataResolution Resolution, const float Height) { NavMeshResolutionParams[(uint8)Resolution].CellHeight = Height; }
+
+	/** Get the AgentMaxStepHeight for the given resolution. */
+	float GetAgentMaxStepHeight(const ENavigationDataResolution Resolution) const { return NavMeshResolutionParams[(uint8)Resolution].AgentMaxStepHeight; }
+
+	/** Set the AgentMaxStepHeight for the given resolution. */
+	void SetAgentMaxStepHeight(const ENavigationDataResolution Resolution, const float MaxStepHeight) { NavMeshResolutionParams[(uint8)Resolution].AgentMaxStepHeight = MaxStepHeight; }
 	
 	/** Returns the tile size in world units. */
 	float GetTileSizeUU() const;
diff --git a/Engine/Source/Runtime/NavigationSystem/Public/NavMesh/RecastVersion.h b/Engine/Source/Runtime/NavigationSystem/Public/NavMesh/RecastVersion.h
--- a/Engine/Source/Runtime/NavigationSystem/Public/NavMesh/RecastVersion.h
+++ b/Engine/Source/Runtime/NavigationSystem/Public/NavMesh/RecastVersion.h
@@ -23,7 +23,8 @@
 #define NAVMESHVER_MAXTILES_COUNT_SKIP_INCLUSION 22
 #define NAVMESHVER_TILE_RESOLUTIONS				23 // Addition of a tile resolution index to the tile header.
 #define NAVMESHVER_TILE_RESOLUTIONS_CELLHEIGHT	24 // Addition of CellHeight in the resolution params, deprecating the original CellHeight.
-#define NAVMESHVER_1_VOXEL_AGENT_STEEP_SLOPE_FILTER_FIX	25 // Fix, remove steep slope filtering during heightfield ledge filtering when the agent radius is included into a single voxel 
+#define NAVMESHVER_1_VOXEL_AGENT_STEEP_SLOPE_FILTER_FIX	25 // Fix, remove steep slope filtering during heightfield ledge filtering when the agent radius is included into a single voxel
+#define NAVMESHVER_TILE_RESOLUTIONS_AGENTMAXSTEPHEIGHT 26	// Addition of AgentMaxStepHeight in the resolution params, deprecating the original AgentMaxStepHeight.
 
-#define NAVMESHVER_LATEST				NAVMESHVER_1_VOXEL_AGENT_STEEP_SLOPE_FILTER_FIX
+#define NAVMESHVER_LATEST				NAVMESHVER_TILE_RESOLUTIONS_AGENTMAXSTEPHEIGHT
 #define NAVMESHVER_MIN_COMPATIBLE		NAVMESHVER_LWCOORDS_OPTIMIZATION
