Build fix for 'Resolution': unknown override specifier
#rnx

diff --git a/Engine/Source/Runtime/NavigationSystem/Public/NavMesh/RecastInternalDebugData.h b/Engine/Source/Runtime/NavigationSystem/Public/NavMesh/RecastInternalDebugData.h
--- a/Engine/Source/Runtime/NavigationSystem/Public/NavMesh/RecastInternalDebugData.h
+++ b/Engine/Source/Runtime/NavigationSystem/Public/NavMesh/RecastInternalDebugData.h
@@ -30,9 +30,9 @@ struct FRecastInternalDebugData : public duDebugDraw
 	double BuildNavigationDataTime = 0.;
 
 	uint32 TriangleCount = 0;
-	ENavigationDataResolution Resolution = ENavigationDataResolution::Default;
+	ENavigationDataResolution Resolution;
 	
-	FRecastInternalDebugData() {}
+	FRecastInternalDebugData() : Resolution(ENavigationDataResolution::Default) {}
 	virtual ~FRecastInternalDebugData() override {}
 
 	virtual void depthMask(bool state) override { /*unused*/ };
