Build fix for 'Resolution': unknown override specifier
#rnx

diff --git a/Engine/Source/Runtime/NavigationSystem/Private/NavMesh/NavMeshRenderingComponent.cpp b/Engine/Source/Runtime/NavigationSystem/Private/NavMesh/NavMeshRenderingComponent.cpp
--- a/Engine/Source/Runtime/NavigationSystem/Private/NavMesh/NavMeshRenderingComponent.cpp
+++ b/Engine/Source/Runtime/NavigationSystem/Private/NavMesh/NavMeshRenderingComponent.cpp
@@ -973,7 +973,7 @@ void FNavMeshSceneProxyData::GatherData(const ARecastNavMesh* NavMesh, int32 InN
 						const double BuildTimeMs = DebugData->BuildTime * SecToMs;
 						const double CompressedLayerTimeMs = DebugData->BuildCompressedLayerTime * SecToMs;
 						const double NavigationDataTimeMs = DebugData->BuildNavigationDataTime * SecToMs;
-						const ENavigationDataResolution Resolution = DebugData->Resolution;
+						const ENavigationDataResolution Resolution = (ENavigationDataResolution)DebugData->Resolution;
 
 						const double Range = NavMeshGeometry.MaxTileBuildTime - NavMeshGeometry.MinTileBuildTime;
 						int32 Rank = 0;
diff --git a/Engine/Source/Runtime/NavigationSystem/Private/NavMesh/RecastNavMeshGenerator.cpp b/Engine/Source/Runtime/NavigationSystem/Private/NavMesh/RecastNavMeshGenerator.cpp
--- a/Engine/Source/Runtime/NavigationSystem/Private/NavMesh/RecastNavMeshGenerator.cpp
+++ b/Engine/Source/Runtime/NavigationSystem/Private/NavMesh/RecastNavMeshGenerator.cpp
@@ -2567,7 +2567,7 @@ bool FRecastTileGenerator::GenerateTile()
 	BuildContext.InternalDebugData.BuildTime = EndStamp - StartStamp;
 	BuildContext.InternalDebugData.BuildCompressedLayerTime = PostCompressLayerStamp - StartStamp;
 	BuildContext.InternalDebugData.BuildNavigationDataTime = EndStamp - PostCompressLayerStamp;
-	BuildContext.InternalDebugData.Resolution = TileConfig.TileResolution;
+	BuildContext.InternalDebugData.Resolution = (unsigned char)TileConfig.TileResolution;
 #endif // RECAST_INTERNAL_DEBUG_DATA
 	
 	// it's possible to have valid generation with empty resulting tile (no navigable geometry in tile)
diff --git a/Engine/Source/Runtime/NavigationSystem/Public/NavMesh/RecastInternalDebugData.h b/Engine/Source/Runtime/NavigationSystem/Public/NavMesh/RecastInternalDebugData.h
--- a/Engine/Source/Runtime/NavigationSystem/Public/NavMesh/RecastInternalDebugData.h
+++ b/Engine/Source/Runtime/NavigationSystem/Public/NavMesh/RecastInternalDebugData.h
@@ -8,7 +8,6 @@
 #include "Math/UnrealMathSSE.h"
 #include "Math/Vector.h"
 #include "NavMesh/RecastHelpers.h"
-#include "AI/Navigation/NavigationTypes.h"
 
 struct FRecastInternalDebugData : public duDebugDraw
 {
@@ -30,9 +29,9 @@ struct FRecastInternalDebugData : public duDebugDraw
 	double BuildNavigationDataTime = 0.;
 
 	uint32 TriangleCount = 0;
-	ENavigationDataResolution Resolution;
+	unsigned char Resolution = 0;
 	
-	FRecastInternalDebugData() : Resolution(ENavigationDataResolution::Default) {}
+	FRecastInternalDebugData() {}
 	virtual ~FRecastInternalDebugData() override {}
 
 	virtual void depthMask(bool state) override { /*unused*/ };
