Remove unnecessary calls to UNavigationSystemV1::GetMutableOctree accessor

#rnx

diff --git a/Engine/Source/Developer/FunctionalTesting/Private/FunctionalAITest.cpp b/Engine/Source/Developer/FunctionalTesting/Private/FunctionalAITest.cpp
--- a/Engine/Source/Developer/FunctionalTesting/Private/FunctionalAITest.cpp
+++ b/Engine/Source/Developer/FunctionalTesting/Private/FunctionalAITest.cpp
@@ -105,7 +105,7 @@ void AFunctionalAITestBase::OnTimeout()
 		UE_LOG(LogFunctionalTest, Log, TEXT("Test timed out, log details for: %s"), *GetNameSafe(Navmesh));
 		UE_LOG(LogFunctionalTest, Log, TEXT("> dirty areas? %s"), NavSys->HasDirtyAreasQueued() ? TEXT("YES") : TEXT("no"));
 
-		FNavigationOctree* NavigationOctree = NavSys->GetMutableNavOctree();
+		const FNavigationOctree* NavigationOctree = NavSys->GetNavOctree();
 		
 		FNavigationOctreeFilter AreaFilter;
 		AreaFilter.bIncludeAreas = true;
diff --git a/Engine/Source/Runtime/NavigationSystem/Private/NavMesh/RecastNavMeshGenerator.cpp b/Engine/Source/Runtime/NavigationSystem/Private/NavMesh/RecastNavMeshGenerator.cpp
--- a/Engine/Source/Runtime/NavigationSystem/Private/NavMesh/RecastNavMeshGenerator.cpp
+++ b/Engine/Source/Runtime/NavigationSystem/Private/NavMesh/RecastNavMeshGenerator.cpp
@@ -1869,7 +1869,7 @@ bool FRecastTileGenerator::DoWork()
 	SCOPE_CYCLE_COUNTER(STAT_Navigation_DoWork);
 
 	TSharedPtr<FNavDataGenerator, ESPMode::ThreadSafe> ParentGenerator = ParentGeneratorWeakPtr.Pin();
-	bool bSucceess = true;
+	bool bSuccess = true;
 
 	if (ParentGenerator.IsValid())
 	{
@@ -1878,12 +1878,12 @@ bool FRecastTileGenerator::DoWork()
 			GatherGeometryFromSources();
 		}
 
-		bSucceess = GenerateTile();
+		bSuccess = GenerateTile();
 
 		DumpAsyncData();
 	}
 
-	return bSucceess;
+	return bSuccess;
 }
 
 void FRecastTileGenerator::DumpAsyncData()
@@ -1960,7 +1960,7 @@ void FRecastTileGenerator::PrepareGeometrySources(const FRecastNavMeshGenerator&
 	QUICK_SCOPE_CYCLE_COUNTER(STAT_RecastNavMeshGenerator_PrepareGeometrySources);
 
 	UNavigationSystemV1* NavSys = FNavigationSystem::GetCurrent<UNavigationSystemV1>(ParentGenerator.GetWorld());
-	FNavigationOctree* NavOctreeInstance = NavSys ? NavSys->GetMutableNavOctree() : nullptr;
+	const FNavigationOctree* NavOctreeInstance = NavSys ? NavSys->GetNavOctree() : nullptr;
 	check(NavOctreeInstance);
 	NavigationRelevantData.Reset();
 	NavSystem = NavSys;
@@ -2014,7 +2014,7 @@ void FRecastTileGenerator::GatherGeometry(const FRecastNavMeshGenerator& ParentG
 	QUICK_SCOPE_CYCLE_COUNTER(STAT_RecastNavMeshGenerator_GatherGeometry);
 
 	UNavigationSystemV1* NavSys = FNavigationSystem::GetCurrent<UNavigationSystemV1>(ParentGenerator.GetWorld());
-	FNavigationOctree* NavigationOctree = NavSys ? NavSys->GetMutableNavOctree() : nullptr;
+	const FNavigationOctree* NavigationOctree = NavSys ? NavSys->GetNavOctree() : nullptr;
 	if (NavigationOctree == nullptr)
 	{
 		return;
@@ -5082,7 +5082,7 @@ void FRecastNavMeshGenerator::TickAsyncBuild(float DeltaSeconds)
 #endif//WITH_EDITOR
 
 	const UNavigationSystemV1* NavSys = FNavigationSystem::GetCurrent<UNavigationSystemV1>(GetWorld());
-	if (!ensureMsgf(NavSys != nullptr, TEXT("FRecastNavMeshGenerator can't found valid navigation system: Owner=[%s] World=[%s]"), *GetFullNameSafe(GetOwner()), *GetFullNameSafe(GetWorld())))
+	if (!ensureMsgf(NavSys != nullptr, TEXT("FRecastNavMeshGenerator can't find valid navigation system: Owner=[%s] World=[%s]"), *GetFullNameSafe(GetOwner()), *GetFullNameSafe(GetWorld())))
 	{
 		return;
 	}
diff --git a/Engine/Source/Runtime/NavigationSystem/Public/NavMesh/RecastNavMeshGenerator.h b/Engine/Source/Runtime/NavigationSystem/Public/NavMesh/RecastNavMeshGenerator.h
--- a/Engine/Source/Runtime/NavigationSystem/Public/NavMesh/RecastNavMeshGenerator.h
+++ b/Engine/Source/Runtime/NavigationSystem/Public/NavMesh/RecastNavMeshGenerator.h
@@ -161,7 +161,7 @@ struct FRecastRawGeometryElement
 	TArray<int32>				GeomIndices;
 	
 	// Per instance transformations in unreal coords
-	// When empty geometry is in world space
+	// When empty, geometry is in world space
 	TArray<FTransform>	PerInstanceTransform;
 
 	rcRasterizationFlags RasterizationFlags;
@@ -172,7 +172,7 @@ struct FRecastAreaNavModifierElement
 	TArray<FAreaNavModifier> Areas;
 	
 	// Per instance transformations in unreal coords
-	// When empty areas are in world space
+	// When empty, areas are in world space
 	TArray<FTransform>	PerInstanceTransform;
 
 	bool bMaskFillCollisionUnderneathForNavmesh = false;
