Fixed mismatched copy constructor/assignment operator in FPathFindingQuery.

diff --git a/Engine/Source/Runtime/NavigationSystem/Private/NavigationData.cpp b/Engine/Source/Runtime/NavigationSystem/Private/NavigationData.cpp
--- a/Engine/Source/Runtime/NavigationSystem/Private/NavigationData.cpp
+++ b/Engine/Source/Runtime/NavigationSystem/Private/NavigationData.cpp
@@ -47,16 +47,6 @@ FPathFindingQuery::FPathFindingQuery(const INavAgentInterface& InNavAgent, const
 	}
 }
 
-FPathFindingQuery::FPathFindingQuery(const FPathFindingQuery& Source) :
-	FPathFindingQueryData(Source.Owner.Get(), Source.StartLocation, Source.EndLocation, Source.QueryFilter, Source.NavDataFlags, Source.bAllowPartialPaths, Source.CostLimit, Source.bRequireNavigableEndLocation),
-	NavData(Source.NavData), PathInstanceToFill(Source.PathInstanceToFill), NavAgentProperties(Source.NavAgentProperties)
-{
-	if (!QueryFilter.IsValid() && NavData.IsValid())
-	{
-		QueryFilter = NavData->GetDefaultQueryFilter();
-	}
-}
-
 FPathFindingQuery::FPathFindingQuery(FNavPathSharedRef PathToRecalculate, const ANavigationData* NavDataOverride) :
 	FPathFindingQueryData(PathToRecalculate->GetQueryData()),
 	NavData(NavDataOverride ? NavDataOverride : PathToRecalculate->GetNavigationDataUsed()), PathInstanceToFill(PathToRecalculate), NavAgentProperties(FNavAgentProperties::DefaultProperties)
diff --git a/Engine/Source/Runtime/NavigationSystem/Public/NavigationSystemTypes.h b/Engine/Source/Runtime/NavigationSystem/Public/NavigationSystemTypes.h
--- a/Engine/Source/Runtime/NavigationSystem/Public/NavigationSystemTypes.h
+++ b/Engine/Source/Runtime/NavigationSystem/Public/NavigationSystemTypes.h
@@ -52,7 +52,6 @@ struct NAVIGATIONSYSTEM_API FPathFindingQuery : public FPathFindingQueryData
 	FNavAgentProperties NavAgentProperties;
 
 	FPathFindingQuery() : FPathFindingQueryData() {}
-	FPathFindingQuery(const FPathFindingQuery& Source);
 	FPathFindingQuery(const UObject* InOwner, const ANavigationData& InNavData, const FVector& Start, const FVector& End, FSharedConstNavQueryFilter SourceQueryFilter = NULL, FNavPathSharedPtr InPathInstanceToFill = NULL, const FVector::FReal CostLimit = TNumericLimits<FVector::FReal>::Max(), const bool bInRequireNavigableEndLocation = true);
 	FPathFindingQuery(const INavAgentInterface& InNavAgent, const ANavigationData& InNavData, const FVector& Start, const FVector& End, FSharedConstNavQueryFilter SourceQueryFilter = NULL, FNavPathSharedPtr InPathInstanceToFill = NULL, const FVector::FReal CostLimit = TNumericLimits<FVector::FReal>::Max(), const bool bInRequireNavigableEndLocation = true);
 
