Added a constructor with parameters to NavLinkDrawing and NavLinkSegmentDrawing to be able to reuse the member initialization.

#rnx

diff --git a/Engine/Source/Runtime/NavigationSystem/Private/NavLinkRenderingComponent.cpp b/Engine/Source/Runtime/NavigationSystem/Private/NavLinkRenderingComponent.cpp
--- a/Engine/Source/Runtime/NavigationSystem/Private/NavLinkRenderingComponent.cpp
+++ b/Engine/Source/Runtime/NavigationSystem/Private/NavLinkRenderingComponent.cpp
@@ -149,37 +149,19 @@ SIZE_T FNavLinkRenderingProxy::GetTypeHash() const
 
 void FNavLinkRenderingProxy::StorePointLinks(const FTransform& InLocalToWorld, const TArray<FNavigationLink>& LinksArray)
 {
-	const FNavigationLink* Link = LinksArray.GetData();
-	for (int32 LinkIndex = 0; LinkIndex < LinksArray.Num(); ++LinkIndex, ++Link)
-	{	
-		FNavLinkDrawing LinkDrawing;
-		LinkDrawing.Left = InLocalToWorld.TransformPosition(Link->Left);
-		LinkDrawing.Right = InLocalToWorld.TransformPosition(Link->Right);
-		LinkDrawing.Direction = Link->Direction;
-		LinkDrawing.Color = UNavArea::GetColor(Link->GetAreaClass());
-		LinkDrawing.SnapRadius = Link->SnapRadius;
-		LinkDrawing.SnapHeight = Link->bUseSnapHeight ? Link->SnapHeight : -1.0f;
-		LinkDrawing.SupportedAgentsBits = Link->SupportedAgents.PackedBits;
-		OffMeshPointLinks.Add(LinkDrawing);
+	OffMeshPointLinks.Reserve(OffMeshPointLinks.Num() + LinksArray.Num());
+	for (const FNavigationLink& Link : LinksArray)
+	{
+		OffMeshPointLinks.Emplace(InLocalToWorld, Link);
 	}
 }
 
 void FNavLinkRenderingProxy::StoreSegmentLinks(const FTransform& InLocalToWorld, const TArray<FNavigationSegmentLink>& LinksArray)
 {
-	const FNavigationSegmentLink* Link = LinksArray.GetData();
-	for (int32 LinkIndex = 0; LinkIndex < LinksArray.Num(); ++LinkIndex, ++Link)
-	{	
-		FNavLinkSegmentDrawing LinkDrawing;
-		LinkDrawing.LeftStart = InLocalToWorld.TransformPosition(Link->LeftStart);
-		LinkDrawing.LeftEnd = InLocalToWorld.TransformPosition(Link->LeftEnd);
-		LinkDrawing.RightStart = InLocalToWorld.TransformPosition(Link->RightStart);
-		LinkDrawing.RightEnd = InLocalToWorld.TransformPosition(Link->RightEnd);
-		LinkDrawing.Direction = Link->Direction;
-		LinkDrawing.Color = UNavArea::GetColor(Link->GetAreaClass());
-		LinkDrawing.SnapRadius = Link->SnapRadius;
-		LinkDrawing.SnapHeight = Link->bUseSnapHeight ? Link->SnapHeight : -1.0f;
-		LinkDrawing.SupportedAgentsBits = Link->SupportedAgents.PackedBits;
-		OffMeshSegmentLinks.Add(LinkDrawing);
+	OffMeshSegmentLinks.Reserve(OffMeshSegmentLinks.Num() + LinksArray.Num());
+	for (const FNavigationSegmentLink& Link : LinksArray)
+	{
+		OffMeshSegmentLinks.Emplace(InLocalToWorld, Link);
 	}
 }
 
diff --git a/Engine/Source/Runtime/NavigationSystem/Public/NavLinkRenderingProxy.h b/Engine/Source/Runtime/NavigationSystem/Public/NavLinkRenderingProxy.h
--- a/Engine/Source/Runtime/NavigationSystem/Public/NavLinkRenderingProxy.h
+++ b/Engine/Source/Runtime/NavigationSystem/Public/NavLinkRenderingProxy.h
@@ -23,6 +23,17 @@ public:
 
 	struct FNavLinkDrawing
 	{
+		FNavLinkDrawing() {}
+		FNavLinkDrawing(const FTransform& InLocalToWorld, const FNavigationLink& Link)
+			: Left(InLocalToWorld.TransformPosition(Link.Left))
+			, Right(InLocalToWorld.TransformPosition(Link.Right))
+			, Direction(Link.Direction)
+			, Color(UNavArea::GetColor(Link.GetAreaClass()))
+			, SnapRadius(Link.SnapRadius)
+			, SnapHeight(Link.bUseSnapHeight ? Link.SnapHeight : -1.0f)
+			, SupportedAgentsBits(Link.SupportedAgents.PackedBits)
+		{}
+
 		FVector Left;
 		FVector Right;
 		ENavLinkDirection::Type Direction;
@@ -33,6 +44,19 @@ public:
 	};
 	struct FNavLinkSegmentDrawing
 	{
+		FNavLinkSegmentDrawing() {}
+		FNavLinkSegmentDrawing(const FTransform& InLocalToWorld, const FNavigationSegmentLink& Link)
+			: LeftStart(InLocalToWorld.TransformPosition(Link.LeftStart))
+			, LeftEnd(InLocalToWorld.TransformPosition(Link.LeftEnd))
+			, RightStart(InLocalToWorld.TransformPosition(Link.RightStart))
+			, RightEnd(InLocalToWorld.TransformPosition(Link.RightEnd))
+			, Direction(Link.Direction)
+			, Color(UNavArea::GetColor(Link.GetAreaClass()))
+			, SnapRadius(Link.SnapRadius)
+			, SnapHeight(Link.bUseSnapHeight ? Link.SnapHeight : -1.0f)
+			, SupportedAgentsBits(Link.SupportedAgents.PackedBits)
+		{}
+
 		FVector LeftStart, LeftEnd;
 		FVector RightStart, RightEnd;
 		ENavLinkDirection::Type Direction;
