Support GeometryCollection in UNavModifierComponent
Before this change, a UNavModifierComponent on a geometry collection would not find the physic from the proxy mesh and fallback on an arbitrary "FailsafeExtent".

diff --git a/Engine/Source/Runtime/NavigationSystem/Private/NavModifierComponent.cpp b/Engine/Source/Runtime/NavigationSystem/Private/NavModifierComponent.cpp
--- a/Engine/Source/Runtime/NavigationSystem/Private/NavModifierComponent.cpp
+++ b/Engine/Source/Runtime/NavigationSystem/Private/NavModifierComponent.cpp
@@ -7,6 +7,10 @@
 #include "NavAreas/NavArea_Null.h"
 #include "PhysicsEngine/BodySetup.h"
 #include "NavigationSystem.h"
+#include "GeometryCollection/GeometryCollectionComponent.h"
+#include "GeometryCollection/GeometryCollectionObject.h"
+#include "Engine/StaticMesh.h"
+#include "VisualLogger/VisualLogger.h"
 
 #include UE_INLINE_GENERATED_CPP_BY_NAME(NavModifierComponent)
 
@@ -71,6 +75,54 @@ void UNavModifierComponent::OnNavAreaUnregistered(const UWorld& World, const UCl
 }
 #endif // WITH_EDITOR 
 
+void UNavModifierComponent::PopulateComponentBounds(FTransform InParentTransform, const UBodySetup& InBodySetup) const
+{
+	const FVector Scale3D = InParentTransform.GetScale3D();
+	InParentTransform.RemoveScaling();
+	
+	for (int32 SphereIdx = 0; SphereIdx < InBodySetup.AggGeom.SphereElems.Num(); SphereIdx++)
+	{
+		const FKSphereElem& ElemInfo = InBodySetup.AggGeom.SphereElems[SphereIdx];
+		FTransform ElemTM = ElemInfo.GetTransform();
+		ElemTM.ScaleTranslation(Scale3D);
+		ElemTM *= InParentTransform;
+
+		const FBox SphereBounds = FBox::BuildAABB(ElemTM.GetLocation(), ElemInfo.Radius * Scale3D);
+		ComponentBounds.Add(FRotatedBox(SphereBounds, ElemTM.GetRotation()));
+	}
+
+	for (int32 BoxIdx = 0; BoxIdx < InBodySetup.AggGeom.BoxElems.Num(); BoxIdx++)
+	{
+		const FKBoxElem& ElemInfo = InBodySetup.AggGeom.BoxElems[BoxIdx];
+		FTransform ElemTM = ElemInfo.GetTransform();
+		ElemTM.ScaleTranslation(Scale3D);
+		ElemTM *= InParentTransform;
+
+		const FBox BoxBounds = FBox::BuildAABB(ElemTM.GetLocation(), FVector(ElemInfo.X, ElemInfo.Y, ElemInfo.Z) * Scale3D * 0.5f);
+		ComponentBounds.Add(FRotatedBox(BoxBounds, ElemTM.GetRotation()));
+	}
+
+	for (int32 SphylIdx = 0; SphylIdx < InBodySetup.AggGeom.SphylElems.Num(); SphylIdx++)
+	{
+		const FKSphylElem& ElemInfo = InBodySetup.AggGeom.SphylElems[SphylIdx];
+		FTransform ElemTM = ElemInfo.GetTransform();
+		ElemTM.ScaleTranslation(Scale3D);
+		ElemTM *= InParentTransform;
+
+		const FBox SphylBounds = FBox::BuildAABB(ElemTM.GetLocation(), FVector(ElemInfo.Radius, ElemInfo.Radius, ElemInfo.Length) * Scale3D);
+		ComponentBounds.Add(FRotatedBox(SphylBounds, ElemTM.GetRotation()));
+	}
+
+	for (int32 ConvexIdx = 0; ConvexIdx < InBodySetup.AggGeom.ConvexElems.Num(); ConvexIdx++)
+	{
+		const FKConvexElem& ElemInfo = InBodySetup.AggGeom.ConvexElems[ConvexIdx];
+		FTransform ElemTM = ElemInfo.GetTransform();
+
+		const FBox ConvexBounds = FBox::BuildAABB(InParentTransform.TransformPosition(ElemInfo.ElemBox.GetCenter() * Scale3D), ElemInfo.ElemBox.GetExtent() * Scale3D);
+		ComponentBounds.Add(FRotatedBox(ConvexBounds, ElemTM.GetRotation() * InParentTransform.GetRotation()));
+	}
+}
+
 void UNavModifierComponent::CalcAndCacheBounds() const
 {
 	AActor* MyOwner = GetOwner();
@@ -98,53 +150,32 @@ void UNavModifierComponent::CalcAndCacheBounds() const
 				UBodySetup* BodySetup = PrimComp->GetBodySetup();
 				if (BodySetup)
 				{
-					FTransform ParentTM = PrimComp->GetComponentTransform();
-					const FVector Scale3D = ParentTM.GetScale3D();
-					ParentTM.RemoveScaling();
 					Bounds += PrimComp->Bounds.GetBox();
-
-					for (int32 SphereIdx = 0; SphereIdx < BodySetup->AggGeom.SphereElems.Num(); SphereIdx++)
-					{
-						const FKSphereElem& ElemInfo = BodySetup->AggGeom.SphereElems[SphereIdx];
-						FTransform ElemTM = ElemInfo.GetTransform();
-						ElemTM.ScaleTranslation(Scale3D);
-						ElemTM *= ParentTM;
-
-						const FBox SphereBounds = FBox::BuildAABB(ElemTM.GetLocation(), ElemInfo.Radius * Scale3D);
-						ComponentBounds.Add(FRotatedBox(SphereBounds, ElemTM.GetRotation()));
-					}
-
-					for (int32 BoxIdx = 0; BoxIdx < BodySetup->AggGeom.BoxElems.Num(); BoxIdx++)
-					{
-						const FKBoxElem& ElemInfo = BodySetup->AggGeom.BoxElems[BoxIdx];
-						FTransform ElemTM = ElemInfo.GetTransform();
-						ElemTM.ScaleTranslation(Scale3D);
-						ElemTM *= ParentTM;
-
-						const FBox BoxBounds = FBox::BuildAABB(ElemTM.GetLocation(), FVector(ElemInfo.X, ElemInfo.Y, ElemInfo.Z) * Scale3D * 0.5f);
-						ComponentBounds.Add(FRotatedBox(BoxBounds, ElemTM.GetRotation()));
-					}
-
-					for (int32 SphylIdx = 0; SphylIdx < BodySetup->AggGeom.SphylElems.Num(); SphylIdx++)
+					
+					const FTransform& ParentTM = PrimComp->GetComponentTransform();
+					PopulateComponentBounds(ParentTM, *BodySetup);
+				}
+				else if (const UGeometryCollectionComponent* GeometryCollection = Cast<UGeometryCollectionComponent>(PrimComp))
+				{
+					// If it's a GC, use the bodySetups from the proxyMeshes.
+					if (const TObjectPtr<const UGeometryCollection> RestCollection = GeometryCollection->RestCollection)
 					{
-						const FKSphylElem& ElemInfo = BodySetup->AggGeom.SphylElems[SphylIdx];
-						FTransform ElemTM = ElemInfo.GetTransform();
-						ElemTM.ScaleTranslation(Scale3D);
-						ElemTM *= ParentTM;
-
-						const FBox SphylBounds = FBox::BuildAABB(ElemTM.GetLocation(), FVector(ElemInfo.Radius, ElemInfo.Radius, ElemInfo.Length) * Scale3D);
-						ComponentBounds.Add(FRotatedBox(SphylBounds, ElemTM.GetRotation()));
+						Bounds += GeometryCollection->Bounds.GetBox();
+
+						for (const TObjectPtr<UStaticMesh> ProxyMesh : RestCollection->RootProxyData.ProxyMeshes)
+						{
+							if (ProxyMesh != nullptr)
+							{
+								const UBodySetup* Body = ProxyMesh->GetBodySetup();
+								if (Body)
+								{
+									const FTransform& ParentTM = PrimComp->GetComponentTransform();
+									PopulateComponentBounds(ParentTM, *Body);	
+								}
+							}
+						}
 					}
-
-					for (int32 ConvexIdx = 0; ConvexIdx < BodySetup->AggGeom.ConvexElems.Num(); ConvexIdx++)
-					{
-						const FKConvexElem& ElemInfo = BodySetup->AggGeom.ConvexElems[ConvexIdx];
-						FTransform ElemTM = ElemInfo.GetTransform();
-
-						const FBox ConvexBounds = FBox::BuildAABB(ParentTM.TransformPosition(ElemInfo.ElemBox.GetCenter() * Scale3D), ElemInfo.ElemBox.GetExtent() * Scale3D);
-						ComponentBounds.Add(FRotatedBox(ConvexBounds, ElemTM.GetRotation() * ParentTM.GetRotation()));
-					}
-				}
+				}	
 			}
 		}
 
@@ -163,6 +194,20 @@ void UNavModifierComponent::CalcAndCacheBounds() const
 			ComponentBounds[Idx].Box = FBox::BuildAABB(NavModBoxOrigin, BoxExtent);
 		}
 	}
+
+	UE_SUPPRESS(LogNavigation, VeryVerbose,
+	{
+		TArray<FAreaNavModifier> Areas;
+		for (int32 Idx = 0; Idx < ComponentBounds.Num(); Idx++)
+		{
+			Areas.Add(FAreaNavModifier(ComponentBounds[Idx].Box, FTransform(ComponentBounds[Idx].Quat), AreaClass));
+		}
+
+		for(const FAreaNavModifier& Modifier : Areas)
+		{
+			UE_VLOG_BOX(this, LogNavigation, VeryVerbose, Modifier.GetBounds(), FColor::Yellow, TEXT(""));	
+		}
+	});
 }
 
 void UNavModifierComponent::GetNavigationData(FNavigationRelevantData& Data) const
diff --git a/Engine/Source/Runtime/NavigationSystem/Public/NavModifierComponent.h b/Engine/Source/Runtime/NavigationSystem/Public/NavModifierComponent.h
--- a/Engine/Source/Runtime/NavigationSystem/Public/NavModifierComponent.h
+++ b/Engine/Source/Runtime/NavigationSystem/Public/NavModifierComponent.h
@@ -10,6 +10,7 @@
 #include "NavModifierComponent.generated.h"
 
 struct FNavigationRelevantData;
+class UBodySetup;
 
 UCLASS(ClassGroup = (Navigation), meta = (BlueprintSpawnableComponent), hidecategories = (Activation), config = Engine, defaultconfig)
 class NAVIGATIONSYSTEM_API UNavModifierComponent : public UNavRelevantComponent
@@ -51,6 +52,8 @@ protected:
 	virtual void OnUnregister() override;
 	//~ End UActorComponent Interface
 
+	void PopulateComponentBounds(FTransform InParentTransform, const UBodySetup& InBodySetup) const;
+	
 	struct FRotatedBox
 	{
 		FBox Box;
