Improve the "oversized dirty area" warnings to work with active tile generation (including invokers).
This changes fixes multiple unnecessary warnings when dirty areas are not touching the active tiles.

diff --git a/Engine/Source/Runtime/NavigationSystem/Private/NavMesh/RecastNavMeshGenerator.cpp b/Engine/Source/Runtime/NavigationSystem/Private/NavMesh/RecastNavMeshGenerator.cpp
--- a/Engine/Source/Runtime/NavigationSystem/Private/NavMesh/RecastNavMeshGenerator.cpp
+++ b/Engine/Source/Runtime/NavigationSystem/Private/NavMesh/RecastNavMeshGenerator.cpp
@@ -4806,8 +4806,8 @@ void FRecastNavMeshGenerator::Init()
 					}
 					else
 					{
-						UE_LOG(LogNavigation, Warning, TEXT("Recreating dtNavMesh instance due to saved navmesh origin (%s, usually the RecastNavMesh location) not being aligned with tile size (%d uu) ")
-							, *Orig.ToString(), int(TileDim));
+						UE_LOG(LogNavigation, Warning, TEXT("Recreating dtNavMesh instance %s due to saved navmesh origin (%s, usually the RecastNavMesh location) not being aligned with tile size (%d uu) ")
+							, *GetNameSafe(DestNavMesh), *Orig.ToString(), int(TileDim));
 					}
 				}
 
@@ -5876,7 +5876,7 @@ void FRecastNavMeshGenerator::MarkDirtyTiles(const TArray<FNavigationDirtyArea>&
 	const ARecastNavMesh* const OwnerNav = GetOwner();
 	const bool bUseVirtualGeometryFilteringAndDirtying = OwnerNav != nullptr && OwnerNav->bUseVirtualGeometryFilteringAndDirtying;
 	// Those are set only if bUseVirtualGeometryFilteringAndDirtying is enabled since we do not use them for anything else
-	const UNavigationSystemV1* const NavSys = bUseVirtualGeometryFilteringAndDirtying ? FNavigationSystem::GetCurrent<UNavigationSystemV1>(GetWorld()) : nullptr;
+	const UNavigationSystemV1* NavSys = bUseVirtualGeometryFilteringAndDirtying ? FNavigationSystem::GetCurrent<UNavigationSystemV1>(GetWorld()) : nullptr;
 	const FNavigationOctree* const NavOctreeInstance = (bUseVirtualGeometryFilteringAndDirtying && NavSys) ? NavSys->GetNavOctree() : nullptr;
 	const FNavDataConfig* const NavDataConfig = bUseVirtualGeometryFilteringAndDirtying ? &DestNavMesh->GetConfig() : nullptr;
 	// ~
@@ -5938,7 +5938,8 @@ void FRecastNavMeshGenerator::MarkDirtyTiles(const TArray<FNavigationDirtyArea>&
 			// then FindInclusionBoundEncapsulatingBox can produce false negative
 			bDoTileInclusionTest = (FindInclusionBoundEncapsulatingBox(CutDownArea) == INDEX_NONE);
 		}
-		
+
+		uint32 PendingTilesMarked = 0;
 		const FRcTileBox TileBox(AdjustedAreaBounds, RcNavMeshOrigin, TileSizeInWorldUnits);
 
 		for (int32 TileY = TileBox.YMin; TileY <= TileBox.YMax; ++TileY)
@@ -5975,11 +5976,12 @@ void FRecastNavMeshGenerator::MarkDirtyTiles(const TArray<FNavigationDirtyArea>&
 				{
 					Element.DirtyAreas.Add(AdjustedAreaBounds);
 				}
+				PendingTilesMarked++;
 				
 				FPendingTileElement* ExistingElement = DirtyTiles.Find(Element);
 				if (ExistingElement)
 				{
-					ExistingElement->bRebuildGeometry|= Element.bRebuildGeometry;
+					ExistingElement->bRebuildGeometry |= Element.bRebuildGeometry;
 					// Append area bounds to existing list 
 					if (ExistingElement->bRebuildGeometry == false)
 					{
@@ -6004,6 +6006,37 @@ void FRecastNavMeshGenerator::MarkDirtyTiles(const TArray<FNavigationDirtyArea>&
 #endif
 			}
 		}
+
+#if !UE_BUILD_SHIPPING		
+		// Warn if this is from a big dirty area
+		UE_SUPPRESS(LogNavigationDirtyArea, Warning, 
+		{
+			if (PendingTilesMarked > 0)
+			{
+				if (NavSys == nullptr)
+				{
+					// NavSys might not have been initialized yet if not using bUseVirtualGeometryFilteringAndDirtying
+					NavSys = FNavigationSystem::GetCurrent<UNavigationSystemV1>(GetWorld());
+				}
+				
+				// If not using active tile generation, those are reported earlier in FNavigationDirtyAreasController::AddArea
+				if (NavSys && NavSys->GetOperationMode() == FNavigationSystemRunMode::GameMode && NavSys->IsActiveTilesGenerationEnabled() &&
+					AdjustedAreaBounds.GetSize().GetMax() > NavSys->GetDirtyAreaWarningSizeThreshold())
+				{
+					const UObject* const SourceObject = DirtyArea.OptionalSourceObject.Get();
+					const UActorComponent* const ObjectAsComponent = Cast<UActorComponent>(SourceObject);
+					const AActor* const ComponentOwner = ObjectAsComponent ? ObjectAsComponent->GetOwner() : nullptr;
+					const FVector2D BoundsSize(DirtyArea.Bounds.GetSize());
+			
+					UE_LOG(LogNavigationDirtyArea, Warning,
+						TEXT("(navmesh: %s) Added an oversized dirty area | Tiles marked: %u | Source object = %s | Potential comp owner = %s | Bounds size = %s | Threshold: %.0f"),
+						*GetNameSafe(GetOwner()), PendingTilesMarked, *GetFullNameSafe(SourceObject), *GetFullNameSafe(ComponentOwner),
+						*BoundsSize.ToString(), NavSys->GetDirtyAreaWarningSizeThreshold());
+				}
+			}
+		});
+#endif // !UE_BUILD_SHIPPING
+		
 	}
 	
 	int32 NumTilesMarked = DirtyTiles.Num();
diff --git a/Engine/Source/Runtime/NavigationSystem/Private/NavigationSystem.cpp b/Engine/Source/Runtime/NavigationSystem/Private/NavigationSystem.cpp
--- a/Engine/Source/Runtime/NavigationSystem/Private/NavigationSystem.cpp
+++ b/Engine/Source/Runtime/NavigationSystem/Private/NavigationSystem.cpp
@@ -1112,8 +1112,9 @@ void UNavigationSystemV1::OnWorldInitDone(FNavigationSystemRunMode Mode)
 		DefaultDirtyAreasController.DirtyAreas.Empty();
 	}
 
-	// Report oversized dirty areas only in game mode
-	DefaultDirtyAreasController.SetCanReportOversizedDirtyArea(Mode == FNavigationSystemRunMode::GameMode);
+	// Dirty area controller reports oversized dirty areas only in game mode and if we are not using active tile generation.
+	// When using active tile generation, this is reported only if tiles are actually marked dirty (ex: see MarkDirtyTiles).
+	DefaultDirtyAreasController.SetCanReportOversizedDirtyArea(Mode == FNavigationSystemRunMode::GameMode && !IsActiveTilesGenerationEnabled());
 
 	for (const ANavigationData* NavData : NavDataSet)
 	{
diff --git a/Engine/Source/Runtime/NavigationSystem/Public/NavigationSystem.h b/Engine/Source/Runtime/NavigationSystem/Public/NavigationSystem.h
--- a/Engine/Source/Runtime/NavigationSystem/Public/NavigationSystem.h
+++ b/Engine/Source/Runtime/NavigationSystem/Public/NavigationSystem.h
@@ -867,6 +867,7 @@ public:
 	void AddDirtyAreas(const TArray<FBox>& NewAreas, int32 Flags, const FName& DebugReason = NAME_None);
 	bool HasDirtyAreasQueued() const;
 	int32 GetNumDirtyAreas() const;
+	float GetDirtyAreaWarningSizeThreshold() const { return DirtyAreaWarningSizeThreshold; }
 
 	const FNavigationOctree* GetNavOctree() const { return DefaultOctreeController.GetOctree(); }
 	FNavigationOctree* GetMutableNavOctree() { return DefaultOctreeController.GetMutableOctree(); }
@@ -1104,6 +1105,8 @@ public:
 
 	FNavRegenTimeSliceManager& GetMutableNavRegenTimeSliceManager() { return NavRegenTimeSliceManager; }
 
+	FNavigationSystemRunMode GetOperationMode() const { return OperationMode; }
+
 protected:
 
 	UPROPERTY()
