Navlink / offmeshconnection Ids are now based on a deterministic 64  bit city hash rather than a 32 bit incremental ID.

diff --git a/Engine/Source/Runtime/AIModule/Private/Navigation/CrowdManager.cpp b/Engine/Source/Runtime/AIModule/Private/Navigation/CrowdManager.cpp
--- a/Engine/Source/Runtime/AIModule/Private/Navigation/CrowdManager.cpp
+++ b/Engine/Source/Runtime/AIModule/Private/Navigation/CrowdManager.cpp
@@ -835,7 +835,7 @@ void UCrowdManager::UpdateAgentPaths()
 
 				if (AnimInfo.t == 0)
 				{
-					const uint32 NavLinkId = RecastNavData->GetLinkUserId(AnimInfo.polyRef);
+					const FNavLinkId NavLinkId = RecastNavData->GetNavLinkUserId(AnimInfo.polyRef);
 					INavLinkCustomInterface* CustomLink = NavSys->GetCustomLink(NavLinkId);
 
 					if (CustomLink)
diff --git a/Engine/Source/Runtime/AIModule/Private/Navigation/PathFollowingComponent.cpp b/Engine/Source/Runtime/AIModule/Private/Navigation/PathFollowingComponent.cpp
--- a/Engine/Source/Runtime/AIModule/Private/Navigation/PathFollowingComponent.cpp
+++ b/Engine/Source/Runtime/AIModule/Private/Navigation/PathFollowingComponent.cpp
@@ -868,10 +868,10 @@ void UPathFollowingComponent::SetMoveSegment(int32 SegmentStartIndex)
 		bWalkingNavLinkStart = FNavMeshNodeFlags(PathPt0.Flags).IsNavLink();
 
 		// handle moving through custom nav links
-		if (PathPt0.CustomLinkId)
+		if (PathPt0.CustomNavLinkId != FNavLinkId::Invalid)
 		{
 			UNavigationSystemV1* NavSys = FNavigationSystem::GetCurrent<UNavigationSystemV1>(GetWorld());
-			INavLinkCustomInterface* CustomNavLink = NavSys->GetCustomLink(PathPt0.CustomLinkId);
+			INavLinkCustomInterface* CustomNavLink = NavSys->GetCustomLink(PathPt0.CustomNavLinkId);
 			StartUsingCustomLink(CustomNavLink, SegmentEnd);
 		}
 
diff --git a/Engine/Source/Runtime/Core/Public/UObject/FortniteMainBranchObjectVersion.h b/Engine/Source/Runtime/Core/Public/UObject/FortniteMainBranchObjectVersion.h
--- a/Engine/Source/Runtime/Core/Public/UObject/FortniteMainBranchObjectVersion.h
+++ b/Engine/Source/Runtime/Core/Public/UObject/FortniteMainBranchObjectVersion.h
@@ -304,6 +304,9 @@ struct CORE_API FFortniteMainBranchObjectVersion
 		// Serialize invalid bounds in world partition actor descriptors
 		WorldPartitionActorDescSerializeInvalidBounds,
 
+		// Upgrade Navigation Links to use 64 bits for the ID
+		NavigationLinkID32To64,
+
 		// -----<new versions can be added above this line>-------------------------------------------------
 		VersionPlusOne,
 		LatestVersion = VersionPlusOne - 1
diff --git a/Engine/Source/Runtime/Engine/Classes/AI/Navigation/NavLinkDefinition.h b/Engine/Source/Runtime/Engine/Classes/AI/Navigation/NavLinkDefinition.h
--- a/Engine/Source/Runtime/Engine/Classes/AI/Navigation/NavLinkDefinition.h
+++ b/Engine/Source/Runtime/Engine/Classes/AI/Navigation/NavLinkDefinition.h
@@ -7,9 +7,7 @@
 #include "UObject/Object.h"
 #include "UObject/Class.h"
 #include "Templates/SubclassOf.h"
-#if UE_ENABLE_INCLUDE_ORDER_DEPRECATED_IN_5_2
 #include "AI/Navigation/NavigationTypes.h"
-#endif //UE_ENABLE_INCLUDE_ORDER_DEPRECATED_IN_5_2
 #include "AI/Navigation/NavAgentSelector.h"
 #include "NavLinkDefinition.generated.h"
 
@@ -40,10 +38,14 @@ struct ENGINE_API FNavigationLinkBase
 	float MaxFallDownLength;
 
 	/** Needs to be 0 for recast data generator */
-	static constexpr int32 InvalidUserId = 0;
+	UE_DEPRECATED(5.4, "Use FNavLinkId::Invalid instead")
+	static constexpr uint32 InvalidUserId = 0;
 
 	/** ID passed to navigation data generator */
-	int32 UserId;
+	UE_DEPRECATED(5.4, "Use NavLinkId instead, this id is no longer used in the engine")
+	uint32 UserId;
+
+	FNavLinkId NavLinkId;
 
 	UPROPERTY(EditAnywhere, Category=Default, meta=(ClampMin = "1.0"))
 	float SnapRadius;
@@ -150,6 +152,13 @@ struct ENGINE_API FNavigationLinkBase
 
 	FNavigationLinkBase();
 
+	PRAGMA_DISABLE_DEPRECATION_WARNINGS
+	FNavigationLinkBase(const FNavigationLinkBase&) = default;
+	FNavigationLinkBase(FNavigationLinkBase&& Other) = default;
+	FNavigationLinkBase& operator=(const FNavigationLinkBase& Other) = default;
+	FNavigationLinkBase& operator=(FNavigationLinkBase&& Other) = default;
+	PRAGMA_ENABLE_DEPRECATION_WARNINGS
+
 	void InitializeAreaClass(const bool bForceRefresh = false);
 	void SetAreaClass(UClass* InAreaClass);
 	UClass* GetAreaClass() const;
diff --git a/Engine/Source/Runtime/Engine/Classes/AI/Navigation/NavigationTypes.h b/Engine/Source/Runtime/Engine/Classes/AI/Navigation/NavigationTypes.h
--- a/Engine/Source/Runtime/Engine/Classes/AI/Navigation/NavigationTypes.h
+++ b/Engine/Source/Runtime/Engine/Classes/AI/Navigation/NavigationTypes.h
@@ -141,6 +141,116 @@ enum class ENavDataGatheringModeConfig : uint8
 //
 DECLARE_DELEGATE_TwoParams(FNavDataPerInstanceTransformDelegate, const FBox&, TArray<FTransform>&);
 
+/*
+ * There should be one of these Ids generated and stored per Actor or Component that inherits from INavLinkCustomInterface. Primarily
+ * its a way of differentiating between Components that inherit from INavLinkCustomInterface as the ActorInstanceGuid
+ * alone can not do this.See UNavLinkCustomComponent as an example. For Level Instances these will be different for each component in an actor but repeated
+ * in repeated Level Instances.FNavLinkId::GenerateUniqueId() uses this to make the UniqueId
+ */
+USTRUCT()
+struct ENGINE_API FNavLinkAuxiliaryId
+{
+	GENERATED_BODY()
+
+	FNavLinkAuxiliaryId() = default;
+
+	uint64 GetId() const
+	{
+		return Id;
+	}
+
+	bool operator==(const FNavLinkAuxiliaryId& Other) const { return Id == Other.Id; }
+	bool operator!=(const FNavLinkAuxiliaryId& Other) const { return !this->operator==(Other); }
+
+	/**
+	 * Helper function: returns unique Auxiliary ID for custom links.
+	 **/
+	static FNavLinkAuxiliaryId GenerateUniqueAuxiliaryId();
+
+private:
+	FNavLinkAuxiliaryId(uint64 InId)
+		: Id(InId)
+	{}
+
+public:
+	static const FNavLinkAuxiliaryId Invalid;
+
+private:
+	UPROPERTY()
+	uint64 Id = InvalidLinkId;
+
+	static constexpr uint64 InvalidLinkId = 0;
+};
+
+/**
+ * NavLink Id type. Legacy Ids were generated using the incremental ID system which has been deprecated but will be supported for quite some time. The new system uses CityHash using GenerateUniqueId()
+ * Only NavLinks generated this way should have the most significant bit of the id set to 1.
+ */
+USTRUCT()
+struct ENGINE_API FNavLinkId
+{
+	GENERATED_BODY()
+
+	FNavLinkId() = default;
+
+	explicit FNavLinkId(uint64 InId)
+		: Id(InId)
+	{}
+
+	bool operator==(const FNavLinkId& Other) const
+	{
+		return Id == Other.Id;
+	}
+
+
+	bool operator!=(const FNavLinkId& Other) const
+	{
+		return !this->operator==(Other);
+	}
+	
+	uint64 GetId() const
+	{
+		return Id;
+	}
+
+	void SetId(uint64 InId)
+	{
+		Id = InId;
+	}
+
+	/**
+	 * New non Legacy NavLinkIds have the most significant bit set to 1. Only Id's formed from GenerateUniqueId() should have the most significant bit set to 1. We should only see extremely rare clashes
+	 * here anyway with 63 bits for the city hash in the sort of number of CustomLinks we are likely to have and there is a system in place to handle this for dynamic meshes at run time.
+	 * For static meshes the AuxiliaryId / ActorInstanceGuid will need to be generated again in the editor (probably easiest by deleting and readding the actor / component in the editor
+	 * at the moment).
+	 */
+	bool IsLegacyId() const
+	{
+		return (Id & NavLinkIdBitMask) == 0;
+	}
+
+	/**
+	 *  Helper function: returns unique ID number for custom links.
+	 */
+	static FNavLinkId GenerateUniqueId(FNavLinkAuxiliaryId AuxiliaryId, FGuid ActorInstanceGuid);
+
+	friend inline uint32 GetTypeHash(FNavLinkId const& Value)
+	{
+		return GetTypeHash(Value.GetId());
+	}
+
+public:
+	static const FNavLinkId Invalid;
+
+private:
+	UPROPERTY()
+	uint64 Id = InvalidLinkId;
+
+	static constexpr uint64 InvalidLinkId = 0;
+	static constexpr uint64 NavLinkIdBitMask = 1ULL << 63; // Newer non incremental 64 bit NavLinkIds will have a 1 in the most significant bit.
+};
+
+
 //////////////////////////////////////////////////////////////////////////
 // Path
 
@@ -198,16 +308,37 @@ struct FNavPathPoint : public FNavLocation
 	/** extra node flags */
 	uint32 Flags;
 
-	/** unique Id of custom navigation link starting at this point */
+	UE_DEPRECATED(5.4, "Use CustomNavLinkId instead. This id is no longer used in the engine.")
 	uint32 CustomLinkId;
 
-	FNavPathPoint() : Flags(0), CustomLinkId(0) {}
+	/** unique Id of custom navigation link starting at this point */
+	FNavLinkId CustomNavLinkId;
+
+	FNavPathPoint() : Flags(0)
+	{
+		PRAGMA_DISABLE_DEPRECATION_WARNINGS
+		CustomLinkId = 0;
+		PRAGMA_ENABLE_DEPRECATION_WARNINGS
+	}
+
 	FNavPathPoint(const FVector& InLocation, NavNodeRef InNodeRef = INVALID_NAVNODEREF, uint32 InFlags = 0) 
-		: FNavLocation(InLocation, InNodeRef), Flags(InFlags), CustomLinkId(0) {}
+		: FNavLocation(InLocation, InNodeRef), Flags(InFlags)
+	{
+		PRAGMA_DISABLE_DEPRECATION_WARNINGS
+		CustomLinkId = 0;
+		PRAGMA_ENABLE_DEPRECATION_WARNINGS
+	}
+
+	PRAGMA_DISABLE_DEPRECATION_WARNINGS
+	FNavPathPoint(const FNavPathPoint&) = default;
+	FNavPathPoint(FNavPathPoint&& Other) = default;
+	FNavPathPoint& operator=(const FNavPathPoint& Other) = default;
+	FNavPathPoint& operator=(FNavPathPoint&& Other) = default;
+	PRAGMA_ENABLE_DEPRECATION_WARNINGS
 
 	bool operator==(const FNavPathPoint& Other) const
 	{
-		return Flags == Other.Flags && CustomLinkId == Other.CustomLinkId && FNavLocation::operator==(Other);
+		return Flags == Other.Flags && CustomNavLinkId == Other.CustomNavLinkId && FNavLocation::operator==(Other);
 	}
 };
 
diff --git a/Engine/Source/Runtime/Engine/Private/AI/Navigation/NavigationModifier.cpp b/Engine/Source/Runtime/Engine/Private/AI/Navigation/NavigationModifier.cpp
--- a/Engine/Source/Runtime/Engine/Private/AI/Navigation/NavigationModifier.cpp
+++ b/Engine/Source/Runtime/Engine/Private/AI/Navigation/NavigationModifier.cpp
@@ -29,11 +29,15 @@ namespace UE::Navigation::Private
 // FNavigationLinkBase
 //----------------------------------------------------------------------//
 FNavigationLinkBase::FNavigationLinkBase() 
-	: LeftProjectHeight(0.0f), MaxFallDownLength(1000.0f), UserId(InvalidUserId), SnapRadius(30.f), SnapHeight(50.0f),
+	: LeftProjectHeight(0.0f), MaxFallDownLength(1000.0f), SnapRadius(30.f), SnapHeight(50.0f),
 	  Direction(ENavLinkDirection::BothWays), bUseSnapHeight(false), bSnapToCheapestArea(true),
 	  bCustomFlag0(false), bCustomFlag1(false), bCustomFlag2(false), bCustomFlag3(false), bCustomFlag4(false),
 	  bCustomFlag5(false), bCustomFlag6(false), bCustomFlag7(false)
 {
+	PRAGMA_DISABLE_DEPRECATION_WARNINGS
+	UserId = InvalidUserId;
+	PRAGMA_ENABLE_DEPRECATION_WARNINGS
+
 	AreaClass = nullptr;
 	SupportedAgentsBits = 0xFFFFFFFF;
 }
@@ -665,7 +669,7 @@ void FSimpleLinkNavModifier::SetSegmentLinks(const TArray<FNavigationSegmentLink
 	for (int32 Idx = 0; Idx < SegmentLinks.Num(); Idx++)
 	{
 		FNavigationSegmentLink& LinkData = SegmentLinks[Idx];
-		LinkData.UserId = UserId;
+		LinkData.NavLinkId = NavLinkId;
 
 		bHasMetaAreasSegment |= LinkData.HasMetaArea();
 		bHasFallDownLinks |= LinkData.MaxFallDownLength > 0.f;
@@ -698,7 +702,7 @@ void FSimpleLinkNavModifier::AppendSegmentLinks(const TArray<FNavigationSegmentL
 	for (int32 Idx = 0; Idx < InLinks.Num(); Idx++)
 	{
 		FNavigationSegmentLink& LinkData = SegmentLinks[LinkBase + Idx];
-		LinkData.UserId = UserId;
+		LinkData.NavLinkId = NavLinkId;
 
 		bHasMetaAreasSegment |= LinkData.HasMetaArea();
 		bHasFallDownLinks |= LinkData.MaxFallDownLength > 0.f;
@@ -723,7 +727,7 @@ void FSimpleLinkNavModifier::AddSegmentLink(const FNavigationSegmentLink& InLink
 	const int32 LinkIdx = SegmentLinks.Add(InLink);
 
 	FNavigationSegmentLink& LinkData = SegmentLinks[LinkIdx];
-	LinkData.UserId = UserId;
+	LinkData.NavLinkId = NavLinkId;
 
 	bHasMetaAreasSegment |= LinkData.HasMetaArea();
 	bHasFallDownLinks |= LinkData.MaxFallDownLength > 0.f;
diff --git a/Engine/Source/Runtime/Engine/Private/AI/Navigation/NavigationTypes.cpp b/Engine/Source/Runtime/Engine/Private/AI/Navigation/NavigationTypes.cpp
--- a/Engine/Source/Runtime/Engine/Private/AI/Navigation/NavigationTypes.cpp
+++ b/Engine/Source/Runtime/Engine/Private/AI/Navigation/NavigationTypes.cpp
@@ -300,6 +300,43 @@ void FNavHeightfieldSamples::Empty()
 	Holes.Empty();
 }
 
+
+namespace UE::Navigation::NavLinkIdHelpers::Private
+{
+	uint64 MakeIdFromGUID(const FGuid Guid)
+	{
+		// Create array to guarantee contiguous memory layout 
+		const uint32 GuidArray[] = { Guid.A, Guid.B, Guid.C, Guid.D };
+		return CityHash64(reinterpret_cast<const char*>(GuidArray), sizeof(GuidArray));
+	}
+
+	uint64 MakeIdFromGUID(FNavLinkAuxiliaryId AuxiliaryId, const FGuid Guid)
+	{
+		const uint64 ActorGuidHash = MakeIdFromGUID(Guid);
+		return CityHash128to64({ AuxiliaryId.GetId(), ActorGuidHash});
+	}
+} // UE::Navigation::NavLinkHelpers
+
+const FNavLinkId FNavLinkId::Invalid = FNavLinkId();
+const FNavLinkAuxiliaryId FNavLinkAuxiliaryId::Invalid = FNavLinkAuxiliaryId();
+
+FNavLinkAuxiliaryId FNavLinkAuxiliaryId::GenerateUniqueAuxiliaryId()
+{
+	const uint64 AuxiliaryId = UE::Navigation::NavLinkIdHelpers::Private::MakeIdFromGUID(FGuid::NewGuid());
+
+	return FNavLinkAuxiliaryId(AuxiliaryId);
+}
+
+FNavLinkId FNavLinkId::GenerateUniqueId(FNavLinkAuxiliaryId AuxiliaryId, FGuid ActorInstanceGuid)
+{
+	// Apply NavLinkIdBitMask to differentiate Legacy Ids (that do not have the mask set).
+	const uint64 UniqueId = UE::Navigation::NavLinkIdHelpers::Private::MakeIdFromGUID(AuxiliaryId, ActorInstanceGuid) | NavLinkIdBitMask;
+
+	UE_LOG(LogNavLink, VeryVerbose, TEXT("%hs id: %u."), __FUNCTION__, UniqueId);
+
+	return FNavLinkId(UniqueId);
+}
+
 //----------------------------------------------------------------------//
 // FNavAgentProperties
 //----------------------------------------------------------------------//
diff --git a/Engine/Source/Runtime/Engine/Public/AI/NavigationModifier.h b/Engine/Source/Runtime/Engine/Public/AI/NavigationModifier.h
--- a/Engine/Source/Runtime/Engine/Public/AI/NavigationModifier.h
+++ b/Engine/Source/Runtime/Engine/Public/AI/NavigationModifier.h
@@ -173,8 +173,12 @@ struct ENGINE_API FSimpleLinkNavModifier : public FNavigationModifier
 	TArray<FNavigationLink> Links;
 	TArray<FNavigationSegmentLink> SegmentLinks;
 	FTransform LocalToWorld;
+
+	UE_DEPRECATED(5.4, "Use NavLinkId instead, this id is no longer used in the engine.")
 	int32 UserId;
 
+	FNavLinkId NavLinkId;
+
 	FSimpleLinkNavModifier() 
 		: bHasFallDownLinks(false)
 		, bHasMetaAreasPoint(false)
@@ -188,7 +192,7 @@ struct ENGINE_API FSimpleLinkNavModifier : public FNavigationModifier
 		, bHasMetaAreasPoint(false)
 		, bHasMetaAreasSegment(false)
 	{
-		UserId = InLink.UserId;
+		NavLinkId = InLink.NavLinkId;
 		AddLink(InLink);
 	}	
 
@@ -200,7 +204,7 @@ struct ENGINE_API FSimpleLinkNavModifier : public FNavigationModifier
 	{
 		if (InLinks.Num() > 0)
 		{
-			UserId = InLinks[0].UserId;
+			NavLinkId = InLinks[0].NavLinkId;
 			SetLinks(InLinks);
 		}
 	}
@@ -211,7 +215,7 @@ struct ENGINE_API FSimpleLinkNavModifier : public FNavigationModifier
 		, bHasMetaAreasPoint(false)
 		, bHasMetaAreasSegment(false)
 	{
-		UserId = InLink.UserId;
+		NavLinkId = InLink.NavLinkId;
 		AddSegmentLink(InLink);
 	}	
 
@@ -223,11 +227,18 @@ struct ENGINE_API FSimpleLinkNavModifier : public FNavigationModifier
 	{
 		if (InSegmentLinks.Num() > 0)
 		{
-			UserId = InSegmentLinks[0].UserId;
+			NavLinkId = InSegmentLinks[0].NavLinkId;
 			SetSegmentLinks(InSegmentLinks);
 		}
 	}
 
+	PRAGMA_DISABLE_DEPRECATION_WARNINGS
+	FSimpleLinkNavModifier(const FSimpleLinkNavModifier&) = default;
+	FSimpleLinkNavModifier(FSimpleLinkNavModifier&& Other) = default;
+	FSimpleLinkNavModifier& operator=(const FSimpleLinkNavModifier& Other) = default;
+	FSimpleLinkNavModifier& operator=(FSimpleLinkNavModifier&& Other) = default;
+	PRAGMA_ENABLE_DEPRECATION_WARNINGS
+
 	FORCEINLINE bool HasFallDownLinks() const { return !!bHasFallDownLinks; }
 	void SetLinks(const TArray<FNavigationLink>& InLinks);
 	void SetSegmentLinks(const TArray<FNavigationSegmentLink>& InLinks);
diff --git a/Engine/Source/Runtime/NavigationSystem/Private/NavLinkCustomComponent.cpp b/Engine/Source/Runtime/NavigationSystem/Private/NavLinkCustomComponent.cpp
--- a/Engine/Source/Runtime/NavigationSystem/Private/NavLinkCustomComponent.cpp
+++ b/Engine/Source/Runtime/NavigationSystem/Private/NavLinkCustomComponent.cpp
@@ -14,6 +14,7 @@
 #include "NavigationOctree.h"
 #include "AI/NavigationSystemHelpers.h"
 #include "AI/Navigation/PathFollowingAgentInterface.h"
+#include "UObject/FortniteMainBranchObjectVersion.h"
 
 #include UE_INLINE_GENERATED_CPP_BY_NAME(NavLinkCustomComponent)
 
@@ -32,7 +33,9 @@ namespace UE::Navigation::LinkCustomComponent::Private
 
 UNavLinkCustomComponent::UNavLinkCustomComponent(const FObjectInitializer& ObjectInitializer) : Super(ObjectInitializer)
 {
+PRAGMA_DISABLE_DEPRECATION_WARNINGS
 	NavLinkUserId = 0;
+PRAGMA_ENABLE_DEPRECATION_WARNINGS
 	LinkRelativeStart = FVector(70, 0, 0);
 	LinkRelativeEnd = FVector(-70, 0, 0);
 	LinkDirection = ENavLinkDirection::BothWays;
@@ -47,13 +50,34 @@ UNavLinkCustomComponent::UNavLinkCustomComponent(const FObjectInitializer& Objec
 	BroadcastRadius = 0.0f;
 	BroadcastChannel = ECC_Pawn;
 	BroadcastInterval = 0.0f;
+
+	if (!HasAnyFlags(RF_ClassDefaultObject))
+	{
+		AuxiliaryCustomLinkId = FNavLinkAuxiliaryId::GenerateUniqueAuxiliaryId();
+	}
+}
+
+void UNavLinkCustomComponent::Serialize(FArchive& Ar)
+{
+	Super::Serialize(Ar);
+
+	Ar.UsingCustomVersion(FFortniteMainBranchObjectVersion::GUID);
+
+	if (Ar.IsLoading() && Ar.CustomVer(FFortniteMainBranchObjectVersion::GUID) < FFortniteMainBranchObjectVersion::NavigationLinkID32To64)
+	{
+		PRAGMA_DISABLE_DEPRECATION_WARNINGS
+		CustomLinkId = FNavLinkId(NavLinkUserId);
+		PRAGMA_ENABLE_DEPRECATION_WARNINGS
+	}
 }
 
 void UNavLinkCustomComponent::PostLoad()
 {
 	Super::PostLoad();
 
-	INavLinkCustomInterface::UpdateUniqueId(NavLinkUserId);
+	PRAGMA_DISABLE_DEPRECATION_WARNINGS
+	INavLinkCustomInterface::UpdateUniqueId(CustomLinkId);
+	PRAGMA_ENABLE_DEPRECATION_WARNINGS
 }
 
 #if WITH_EDITOR
@@ -74,8 +98,8 @@ TStructOnScope<FActorComponentInstanceData> UNavLinkCustomComponent::GetComponen
 {
 	TStructOnScope<FActorComponentInstanceData> InstanceData = MakeStructOnScope<FActorComponentInstanceData, FNavLinkCustomInstanceData>(this);
 	FNavLinkCustomInstanceData* NavLinkCustomInstanceData = InstanceData.Cast<FNavLinkCustomInstanceData>();
-	NavLinkCustomInstanceData->NavLinkUserId = NavLinkUserId;
-
+	NavLinkCustomInstanceData->CustomLinkId = CustomLinkId;
+	NavLinkCustomInstanceData->AuxiliaryCustomLinkId = AuxiliaryCustomLinkId;
 	return InstanceData;
 }
 
@@ -83,7 +107,8 @@ void UNavLinkCustomComponent::ApplyComponentInstanceData(FNavLinkCustomInstanceD
 {
 	check(NavLinkData);
 
-	if (NavLinkUserId != NavLinkData->NavLinkUserId)
+	if (CustomLinkId != NavLinkData->CustomLinkId
+		|| AuxiliaryCustomLinkId != NavLinkData->AuxiliaryCustomLinkId)
 	{
 		// Registered component has its link registered in the navigation system. 
 		// In such case, we need to unregister current Id and register with the Id from the instance data.
@@ -94,7 +119,8 @@ void UNavLinkCustomComponent::ApplyComponentInstanceData(FNavLinkCustomInstanceD
 			UNavigationSystemV1::RequestCustomLinkUnregistering(*this, this);
 		}
 
-		NavLinkUserId = NavLinkData->NavLinkUserId;
+		AuxiliaryCustomLinkId = NavLinkData->AuxiliaryCustomLinkId;
+		CustomLinkId = NavLinkData->CustomLinkId;
 
 		if (bIsLinkRegistrationUpdateRequired)
 		{
@@ -108,8 +134,10 @@ void UNavLinkCustomComponent::PostEditImport()
 {
 	Super::PostEditImport();
 
-	NavLinkUserId = INavLinkCustomInterface::GetUniqueId();
-	UE_LOG(LogNavLink, VeryVerbose, TEXT("%s new navlink id %u."), ANSI_TO_TCHAR(__FUNCTION__), NavLinkUserId);
+	// Generate a new AuxiliarLinkUserId and set CustomLinkId to Invalid, this is then inline with the constructor in this regard.
+	// CustomLinkId is set to a valid Id later in OnRegister().
+	AuxiliaryCustomLinkId = FNavLinkAuxiliaryId::GenerateUniqueAuxiliaryId();
+	CustomLinkId = FNavLinkId::Invalid;
 }
 #endif // WITH_EDITOR
 
@@ -130,14 +158,32 @@ TSubclassOf<UNavArea> UNavLinkCustomComponent::GetLinkAreaClass() const
 	return bLinkEnabled ? EnabledAreaClass : DisabledAreaClass;
 }
 
-uint32 UNavLinkCustomComponent::GetLinkId() const
+FNavLinkAuxiliaryId UNavLinkCustomComponent::GetAuxiliaryId() const
+{
+	return AuxiliaryCustomLinkId;
+}
+
+FNavLinkId UNavLinkCustomComponent::GetId() const
 {
-	return NavLinkUserId;
+	return CustomLinkId;
 }
 
-void UNavLinkCustomComponent::UpdateLinkId(uint32 NewUniqueId)
+void UNavLinkCustomComponent::UpdateLinkId(FNavLinkId NewUniqueId)
 {
-	NavLinkUserId = NewUniqueId;
+	if (NewUniqueId == CustomLinkId)
+	{
+		return;
+	}
+
+#if WITH_EDITOR
+	UE_CLOG(NewUniqueId.IsLegacyId(), LogNavLink, Verbose, TEXT("%hs Navlink using LinkUserId id %llu [%s] is being updated with old style link ID. Please move to the new system. See FNavLinkId::GenerateUniqueId()"), __FUNCTION__, CustomLinkId.GetId(), *GetFullName());
+#endif
+
+	CustomLinkId = NewUniqueId;
+
+#if WITH_EDITOR
+	Modify(true);
+#endif
 }
 
 bool UNavLinkCustomComponent::IsLinkPathfindingAllowed(const UObject* Querier) const
@@ -195,11 +241,43 @@ void UNavLinkCustomComponent::OnRegister()
 {
 	Super::OnRegister();
 
-	if (NavLinkUserId == 0)
+	// Actor::GetActorInstanceGuid() is only available when WITH_EDITOR is valid.
+#if WITH_EDITOR
+	// Do not convert old Ids.
+	if (CustomLinkId == FNavLinkId::Invalid || CustomLinkId.IsLegacyId() == false)
+	{
+		// Either this is a freshly spawned component or a component loaded after AuxiliaryCustomLinkId has been saved (and is a new style Id), so we can safely use this to generate a new id.
+		// 
+		const AActor* Owner = GetOwner();
+		checkf(Owner, TEXT("We must have an Owner here as we need it to create a unique id."));
+
+		const FNavLinkId OldCustomLinkId = CustomLinkId;
+
+		// We calculate the CustomLinkId deterministically every time we call OnRegister(), as Level Instances with Actors with UNavLinkCustomComponents can not serialize data for each 
+		// duplicated UNavLinkCustomComponent in different Level Instances (in the editor world). 
+		// For cooked data this is not a problem as the Level Instances are expanded in to actual instances of actors and components so we can cook the in game value for the CustomLinkId.
+		CustomLinkId = FNavLinkId::GenerateUniqueId(AuxiliaryCustomLinkId, Owner->GetActorInstanceGuid());
+		UE_LOG(LogNavLink, VeryVerbose, TEXT("%hs navlink id generated %llu [%s]."), __FUNCTION__, CustomLinkId.GetId(), *GetFullName());
+
+		if (OldCustomLinkId != CustomLinkId)
+		{
+			Modify(true);
+		}
+	}
+#else // !WITH_EDITOR
+	// There is an edge case here for run time only level instances in that they will have a valid CustomLinkId but it will be from the level instance so Ids will get duplicated
+	// in different level instances. Currently baked in nav mesh is not supported for these level instances and its anticipated nav mesh will need to be dynamically regenerated
+	// when the level instance is spawned. This means the CustomLinkId does not need to be deterministic as its not relating to anything generated outside of the run time.
+	// This clash gets handled via RequestCustomLinkRegistering() in UNavigationSystemV1::RegisterCustomLink(), not in the next section of code!
+
+	// This will only be true for freshly spawned components in game that are not in level instances.
+	if (CustomLinkId == FNavLinkId::Invalid)
 	{
-		NavLinkUserId = INavLinkCustomInterface::GetUniqueId();
-		UE_LOG(LogNavLink, VeryVerbose, TEXT("%s new navlink id %u [%s]."), ANSI_TO_TCHAR(__FUNCTION__), NavLinkUserId, *GetFullName());
+		CustomLinkId = FNavLinkId::GenerateUniqueId(AuxiliaryCustomLinkId, FGuid::NewGuid());
+
+		UE_LOG(LogNavLink, VeryVerbose, TEXT("%hs incremental navlink id generated %llu [%s]."), __FUNCTION__, CustomLinkId.GetId(), *GetFullName());
 	}
+#endif // WITH_EDITOR
 
 	UNavigationSystemV1::RequestCustomLinkRegistering(*this, this);
 
diff --git a/Engine/Source/Runtime/NavigationSystem/Private/NavMesh/NavMeshPath.cpp b/Engine/Source/Runtime/NavigationSystem/Private/NavMesh/NavMeshPath.cpp
--- a/Engine/Source/Runtime/NavigationSystem/Private/NavMesh/NavMeshPath.cpp
+++ b/Engine/Source/Runtime/NavigationSystem/Private/NavMesh/NavMeshPath.cpp
@@ -41,7 +41,7 @@ void FNavMeshPath::InternalResetNavMeshPath()
 {
 	PathCorridor.Reset();
 	PathCorridorCost.Reset();
-	CustomLinkIds.Reset();
+	CustomNavLinkIds.Reset();
 	PathCorridorEdges.Reset();
 
 	bCorridorEdgesGenerated = false;
@@ -125,7 +125,7 @@ void FNavMeshPath::PerformStringPulling(const FVector& StartLoc, const FVector&
 	const ARecastNavMesh* MyOwner = Cast<ARecastNavMesh>(GetNavigationDataUsed());
 	if (::IsValid(MyOwner) && PathCorridor.Num())
 	{
-		bStringPulled = MyOwner->FindStraightPath(StartLoc, EndLoc, PathCorridor, PathPoints, &CustomLinkIds);
+		bStringPulled = MyOwner->FindStraightPath(StartLoc, EndLoc, PathCorridor, PathPoints, &CustomNavLinkIds);
 	}
 #endif	// WITH_RECAST
 }
diff --git a/Engine/Source/Runtime/NavigationSystem/Private/NavMesh/PImplRecastNavMesh.cpp b/Engine/Source/Runtime/NavigationSystem/Private/NavMesh/PImplRecastNavMesh.cpp
--- a/Engine/Source/Runtime/NavigationSystem/Private/NavMesh/PImplRecastNavMesh.cpp
+++ b/Engine/Source/Runtime/NavigationSystem/Private/NavMesh/PImplRecastNavMesh.cpp
@@ -5,11 +5,8 @@
 
 #if WITH_RECAST
 
-// recast includes
 #include "NavMesh/RecastHelpers.h"
 #include "NavMesh/RecastVersion.h"
-
-// recast includes
 #include "Detour/DetourNode.h"
 #include "Detour/DetourNavMesh.h"
 #include "Recast/RecastAlloc.h"
@@ -22,6 +19,7 @@
 
 #include "Misc/LargeWorldCoordinates.h"
 #include "DebugUtils/DebugDrawLargeWorldCoordinates.h"
+#include "UObject/FortniteMainBranchObjectVersion.h"
 
 //----------------------------------------------------------------------//
 // bunch of compile-time checks to assure types used by Recast and our
@@ -426,9 +424,9 @@ uint16 FRecastQueryFilter::GetExcludeFlags() const
 	return getExcludeFlags();
 }
 
-bool FRecastSpeciaLinkFilter::isLinkAllowed(const int32 UserId) const
+bool FRecastSpeciaLinkFilter::isLinkAllowed(const uint64 UserId) const
 {
-	const INavLinkCustomInterface* CustomLink = NavSys ? NavSys->GetCustomLink(UserId) : NULL;
+	const INavLinkCustomInterface* CustomLink = NavSys ? NavSys->GetCustomLink(FNavLinkId(UserId)) : nullptr;
 	return (CustomLink != NULL) && CustomLink->IsLinkPathfindingAllowed(CachedOwnerOb);
 }
 
@@ -908,7 +906,19 @@ void FPImplRecastNavMesh::SerializeRecastMeshTile(FArchive& Ar, int32 NavMeshVer
 		{
 			dtOffMeshConnection& Conn = OffMeshCons[ConnIdx];
 			Ar << Conn.pos[0] << Conn.pos[1] << Conn.pos[2] << Conn.pos[3] << Conn.pos[4] << Conn.pos[5];
-			Ar << Conn.rad << Conn.poly << Conn.flags << Conn.side << Conn.userId;
+			Ar << Conn.rad << Conn.poly << Conn.flags << Conn.side;
+
+			if (Ar.IsLoading() && Ar.CustomVer(FFortniteMainBranchObjectVersion::GUID) < FFortniteMainBranchObjectVersion::NavigationLinkID32To64)
+			{
+				uint32 Id;
+
+				Ar << Id;
+				Conn.userId = Id;
+			}
+			else
+			{
+				Ar << Conn.userId;
+			}
 		}
 
 		if (NavMeshVersion >= NAVMESHVER_OFFMESH_HEIGHT_BUG)
@@ -1462,7 +1472,7 @@ void FPImplRecastNavMesh::PostProcessPath(dtStatus FindPathStatus, FNavMeshPath&
 				const dtOffMeshConnection* OffMeshCon = DetourNavMesh->getOffMeshConnectionByRef(Path.PathCorridor[Idx]);
 				if (OffMeshCon)
 				{
-					Path.CustomLinkIds.Add(OffMeshCon->userId);
+					Path.CustomNavLinkIds.Add(FNavLinkId(OffMeshCon->userId));
 				}
 			}
 		}
@@ -1476,7 +1486,7 @@ void FPImplRecastNavMesh::PostProcessPath(dtStatus FindPathStatus, FNavMeshPath&
 	}
 }
 
-bool FPImplRecastNavMesh::FindStraightPath(const FVector& StartLoc, const FVector& EndLoc, const TArray<NavNodeRef>& PathCorridor, TArray<FNavPathPoint>& PathPoints, TArray<uint32>* CustomLinks) const
+bool FPImplRecastNavMesh::FindStraightPath(const FVector& StartLoc, const FVector& EndLoc, const TArray<NavNodeRef>& PathCorridor, TArray<FNavPathPoint>& PathPoints, TArray<FNavLinkId>* CustomLinks) const
 {
 	INITIALIZE_NAVQUERY_SIMPLE(NavQuery, RECAST_MAX_SEARCH_NODES);
 
@@ -1522,8 +1532,8 @@ bool FPImplRecastNavMesh::FindStraightPath(const FVector& StartLoc, const FVecto
 				const dtOffMeshConnection* OffMeshCon = DetourNavMesh->getOffMeshConnectionByRef(CurVert->NodeRef);
 				if (OffMeshCon)
 				{
-					CurVert->CustomLinkId = OffMeshCon->userId;
-					CustomLinks->Add(OffMeshCon->userId);
+					CurVert->CustomNavLinkId.SetId(OffMeshCon->userId);
+					CustomLinks->Add(FNavLinkId(OffMeshCon->userId));
 				}
 			}
 
@@ -2079,11 +2089,11 @@ bool FPImplRecastNavMesh::GetPolysWithinPathingDistance(FVector const& StartLoc,
 	return FoundPolys.Num() > 0;
 }
 
-void FPImplRecastNavMesh::UpdateNavigationLinkArea(int32 UserId, uint8 AreaType, uint16 PolyFlags) const
+void FPImplRecastNavMesh::UpdateNavigationLinkArea(FNavLinkId UserId, uint8 AreaType, uint16 PolyFlags) const
 {
 	if (DetourNavMesh)
 	{
-		DetourNavMesh->updateOffMeshConnectionByUserId(UserId, AreaType, PolyFlags);
+		DetourNavMesh->updateOffMeshConnectionByUserId(UserId.GetId(), AreaType, PolyFlags);
 	}
 }
 
@@ -2378,19 +2388,11 @@ bool FPImplRecastNavMesh::GetClosestPointOnPoly(NavNodeRef PolyID, const FVector
 	return false;
 }
 
-uint32 FPImplRecastNavMesh::GetLinkUserId(NavNodeRef LinkPolyID) const
+FNavLinkId FPImplRecastNavMesh::GetNavLinkUserId(NavNodeRef LinkPolyID) const
 {
-	uint32 UserID = 0;
-	if (DetourNavMesh)
-	{
-		const dtOffMeshConnection* offmeshCon = DetourNavMesh->getOffMeshConnectionByRef(LinkPolyID);
-		if (offmeshCon)
-		{
-			UserID = offmeshCon->userId;
-		}
-	}
+	const dtOffMeshConnection* offmeshCon = DetourNavMesh ? DetourNavMesh->getOffMeshConnectionByRef(LinkPolyID) : nullptr;
 
-	return UserID;
+	return offmeshCon ? FNavLinkId(offmeshCon->userId) : FNavLinkId::Invalid;
 }
 
 bool FPImplRecastNavMesh::GetLinkEndPoints(NavNodeRef LinkPolyID, FVector& PointA, FVector& PointB) const
@@ -2417,7 +2419,7 @@ bool FPImplRecastNavMesh::IsCustomLink(NavNodeRef PolyRef) const
 	if (DetourNavMesh)
 	{
 		const dtOffMeshConnection* offMeshCon = DetourNavMesh->getOffMeshConnectionByRef(PolyRef);
-		return offMeshCon && offMeshCon->userId;
+		return offMeshCon && FNavLinkId(offMeshCon->userId) != FNavLinkId::Invalid;
 	}
 
 	return false;
diff --git a/Engine/Source/Runtime/NavigationSystem/Private/NavMesh/RecastNavMesh.cpp b/Engine/Source/Runtime/NavigationSystem/Private/NavMesh/RecastNavMesh.cpp
--- a/Engine/Source/Runtime/NavigationSystem/Private/NavMesh/RecastNavMesh.cpp
+++ b/Engine/Source/Runtime/NavigationSystem/Private/NavMesh/RecastNavMesh.cpp
@@ -37,6 +37,7 @@
 #endif // WITH_RECAST
 
 #include "NavMesh/NavMeshRenderingComponent.h"
+#include "UObject/FortniteMainBranchObjectVersion.h"
 
 #include UE_INLINE_GENERATED_CPP_BY_NAME(RecastNavMesh)
 
@@ -245,6 +246,8 @@ void ARecastNavMesh::Serialize( FArchive& Ar )
 	uint32 NavMeshVersion;
 	Ar << NavMeshVersion;
 
+	Ar.UsingCustomVersion(FFortniteMainBranchObjectVersion::GUID);
+
 	//@todo: How to handle loading nav meshes saved w/ recast when recast isn't present????
 
 	// when writing, write a zero here for now.  will come back and fill it in later.
@@ -2062,7 +2065,7 @@ FVector::FReal ARecastNavMesh::FindDistanceToWall(const FVector& StartLoc, FShar
 void ARecastNavMesh::UpdateCustomLink(const INavLinkCustomInterface* CustomLink)
 {
 	TSubclassOf<UNavArea> AreaClass = CustomLink->GetLinkAreaClass();
-	const int32 UserId = CustomLink->GetLinkId();
+	const FNavLinkId UserId = CustomLink->GetId();
 	const int32 AreaId = GetAreaID(AreaClass);
 	if (AreaId >= 0 && RecastNavMeshImpl)
 	{
@@ -2071,7 +2074,8 @@ void ARecastNavMesh::UpdateCustomLink(const INavLinkCustomInterface* CustomLink)
 
 		RecastNavMeshImpl->UpdateNavigationLinkArea(UserId, IntCastChecked<uint8>(AreaId), PolyFlags);
 #if WITH_NAVMESH_SEGMENT_LINKS
-		RecastNavMeshImpl->UpdateSegmentLinkArea(UserId, IntCastChecked<uint8>(AreaId), PolyFlags);
+		// SegmentLinks are an unsupported feature that was never completed to production quality, for now at least FNavLinkId ids are not supported here.
+		RecastNavMeshImpl->UpdateSegmentLinkArea((int32)UserId.GetId(), IntCastChecked<uint8>(AreaId), PolyFlags);
 #endif // WITH_NAVMESH_SEGMENT_LINKS
 
 #if !UE_BUILD_SHIPPING
@@ -2080,7 +2084,7 @@ void ARecastNavMesh::UpdateCustomLink(const INavLinkCustomInterface* CustomLink)
 	}
 }
 
-void ARecastNavMesh::UpdateNavigationLinkArea(int32 UserId, TSubclassOf<UNavArea> AreaClass) const
+void ARecastNavMesh::UpdateNavigationLinkArea(FNavLinkId UserId, TSubclassOf<UNavArea> AreaClass) const
 {
 	int32 AreaId = GetAreaID(AreaClass);
 	if (AreaId >= 0 && RecastNavMeshImpl)
@@ -3079,15 +3083,9 @@ bool ARecastNavMesh::IsSegmentOnNavmesh(const FVector& SegmentStart, const FVect
 	return Result.bIsRaycastEndInCorridor && !Result.HasHit();
 }
 
-bool ARecastNavMesh::FindStraightPath(const FVector& StartLoc, const FVector& EndLoc, const TArray<NavNodeRef>& PathCorridor, TArray<FNavPathPoint>& PathPoints, TArray<uint32>* CustomLinks) const
+bool ARecastNavMesh::FindStraightPath(const FVector& StartLoc, const FVector& EndLoc, const TArray<NavNodeRef>& PathCorridor, TArray<FNavPathPoint>& PathPoints, TArray<FNavLinkId>* CustomLinks) const
 {
-	bool bResult = false;
-	if (RecastNavMeshImpl)
-	{
-		bResult = RecastNavMeshImpl->FindStraightPath(StartLoc, EndLoc, PathCorridor, PathPoints, CustomLinks);
-	}
-
-	return bResult;
+	return RecastNavMeshImpl && RecastNavMeshImpl->FindStraightPath(StartLoc, EndLoc, PathCorridor, PathPoints, CustomLinks);
 }
 
 int32 ARecastNavMesh::DebugPathfinding(const FPathFindingQuery& Query, TArray<FRecastDebugPathfindingData>& Steps)
@@ -3732,9 +3730,9 @@ void FRecastNavMeshCachedData::OnAreaRemoved(const UClass* AreaClass)
 	}
 }
 
-uint32 ARecastNavMesh::GetLinkUserId(NavNodeRef LinkPolyID) const
+FNavLinkId ARecastNavMesh::GetNavLinkUserId(NavNodeRef LinkPolyID) const
 {
-	return RecastNavMeshImpl ? RecastNavMeshImpl->GetLinkUserId(LinkPolyID) : 0;
+	return RecastNavMeshImpl ? RecastNavMeshImpl->GetNavLinkUserId(LinkPolyID) : FNavLinkId::Invalid;
 }
 
 dtNavMesh* ARecastNavMesh::GetRecastMesh()
diff --git a/Engine/Source/Runtime/NavigationSystem/Private/NavMesh/RecastNavMeshGenerator.cpp b/Engine/Source/Runtime/NavigationSystem/Private/NavMesh/RecastNavMeshGenerator.cpp
--- a/Engine/Source/Runtime/NavigationSystem/Private/NavMesh/RecastNavMeshGenerator.cpp
+++ b/Engine/Source/Runtime/NavigationSystem/Private/NavMesh/RecastNavMeshGenerator.cpp
@@ -1339,7 +1339,7 @@ struct FOffMeshData
 
 			NewInfo.snapRadius = Link.SnapRadius;
 			NewInfo.snapHeight = Link.bUseSnapHeight ? Link.SnapHeight : DefaultSnapHeight;
-			NewInfo.userID = Link.UserId;
+			NewInfo.userID = Link.NavLinkId.GetId();
 
 			UClass* AreaClass = Link.GetAreaClass();
 			const int32* AreaID = AreaClassToIdMap->Find(AreaClass);
@@ -1380,7 +1380,7 @@ struct FOffMeshData
 			NewInfo.type = DT_OFFMESH_CON_SEGMENT | (Link.Direction == ENavLinkDirection::BothWays ? DT_OFFMESH_CON_BIDIR : 0);
 			NewInfo.snapRadius = Link.SnapRadius;
 			NewInfo.snapHeight = Link.bUseSnapHeight ? Link.SnapHeight : DefaultSnapHeight;
-			NewInfo.userID = Link.UserId;
+			NewInfo.userID = Link.NavLinkId.GetId();
 
 			UClass* AreaClass = Link.GetAreaClass();
 			const int32* AreaID = AreaClassToIdMap->Find(AreaClass);
diff --git a/Engine/Source/Runtime/NavigationSystem/Private/NavigationPath.cpp b/Engine/Source/Runtime/NavigationSystem/Private/NavigationPath.cpp
--- a/Engine/Source/Runtime/NavigationSystem/Private/NavigationPath.cpp
+++ b/Engine/Source/Runtime/NavigationSystem/Private/NavigationPath.cpp
@@ -266,11 +266,16 @@ FVector::FReal FNavigationPath::GetLengthFromPosition(FVector SegmentStart, uint
 	return PathDistance;
 }
 
-bool FNavigationPath::ContainsCustomLink(uint32 LinkUniqueId) const
+bool FNavigationPath::ContainsCustomLink(FNavLinkId LinkUniqueId) const
 {
+	if (LinkUniqueId == FNavLinkId::Invalid)
+	{
+		return false;
+	}
+
 	for (int32 i = 0; i < PathPoints.Num(); i++)
 	{
-		if (PathPoints[i].CustomLinkId == LinkUniqueId && LinkUniqueId)
+		if (PathPoints[i].CustomNavLinkId == LinkUniqueId)
 		{
 			return true;
 		}
@@ -283,7 +288,7 @@ bool FNavigationPath::ContainsAnyCustomLink() const
 {
 	for (int32 i = 0; i < PathPoints.Num(); i++)
 	{
-		if (PathPoints[i].CustomLinkId)
+		if (PathPoints[i].CustomNavLinkId != FNavLinkId::Invalid)
 		{
 			return true;
 		}
diff --git a/Engine/Source/Runtime/NavigationSystem/Private/NavigationSystem.cpp b/Engine/Source/Runtime/NavigationSystem/Private/NavigationSystem.cpp
--- a/Engine/Source/Runtime/NavigationSystem/Private/NavigationSystem.cpp
+++ b/Engine/Source/Runtime/NavigationSystem/Private/NavigationSystem.cpp
@@ -2569,37 +2569,71 @@ void UNavigationSystemV1::UnregisterNavData(ANavigationData* NavData)
 
 void UNavigationSystemV1::RegisterCustomLink(INavLinkCustomInterface& CustomLink)
 {
-	ensureMsgf(CustomLink.GetLinkOwner() == nullptr || GetWorld() == CustomLink.GetLinkOwner()->GetWorld(), 
+	ensureMsgf(CustomLink.GetLinkOwner() == nullptr || GetWorld() == CustomLink.GetLinkOwner()->GetWorld(),
 		TEXT("Registering a link from a world different than the navigation system world should not happen."));
 
-	uint32 LinkId = CustomLink.GetLinkId();
+	const FNavLinkId OldId = CustomLink.GetId();
+	FNavLinkId NewId = OldId;
+	bool bGenerateNewId = false;
 
-	// if there's already a link with that Id registered, assign new Id and mark dirty area
-	// this won't fix baked data in static navmesh (in game), but every other case will regenerate affected tiles 
-	if (CustomLinksMap.Contains(LinkId))
+	// Test for Id clash
+	if (CustomNavLinksMap.Contains(OldId))
 	{
-		LinkId = INavLinkCustomInterface::GetUniqueId();
-		UE_LOG(LogNavLink, VeryVerbose, TEXT("%s new navlink id %u."), ANSI_TO_TCHAR(__FUNCTION__), LinkId);
-		CustomLink.UpdateLinkId(LinkId);
+		if (OldId.IsLegacyId() == false)
+		{
+			UWorld* World = GetWorld();
+			check(World);
+
+			// During PIE or game we just generate a new Id, this is most likely to be from a runtime (non editor placed) prefab like a level instance but could be from 
+			// a legitimate but extremely unlikely Id clash after loading.
+			// If this occurs in EWorldType::Editor world it's a legitimate ID clash, currently we do not handle this edge case here as it should be incredibly unlikely to occur
+			// and we do not save changes when cooking or building paths running a commandlet etc.
+			bGenerateNewId = World->WorldType == EWorldType::PIE || World->WorldType == EWorldType::Game;
+			if (ensureMsgf(bGenerateNewId, TEXT("Id clash in non Game and non PIE world. This should be incredibly rare!")))
+			{
+				// Pass in NewGuid() here as EWorldType::Game does not have access to the ActorInstanceGuid in any case and any random Unique Guid is acceptable here 
+				// if we are not in EWorldType::Editor. Editor is different as we need the cook to be deterministic but for level instances individual actors are not 
+				// serialized (but they are when cooked).
+				NewId = FNavLinkId::GenerateUniqueId(CustomLink.GetAuxiliaryId(), FGuid::NewGuid());
+			}
+			
+			// This should be very unlikely to occur, if its causing issues we should add code to handle this being careful to account for the editor world being run as a commandlet to cook and build paths on seperate runs.
+			UE_CLOG(!bGenerateNewId, LogNavLink, Warning, TEXT("%hs navlink ID %llu is clashing with existing ID. This will not be regenerated automatically in editor although for dynamic navmesh this will be handled at run time in game. For static mesh in the editor world the INavLinkCustomInterface implementor should regenerate the ID, deleting the owning actor and or component and placing again should fix this."), __FUNCTION__, CustomLink.GetId().GetId());
+		}
+		else
+		{
+			bGenerateNewId = true;
+			PRAGMA_DISABLE_DEPRECATION_WARNINGS
+			NewId = FNavLinkId(INavLinkCustomInterface::GetUniqueId());
+			PRAGMA_ENABLE_DEPRECATION_WARNINGS
+		}
 
-		const FBox LinkBounds = ComputeCustomLinkBounds(CustomLink);
-		if (LinkBounds.IsValid)
+		// If the Id has changed mark the area dirty, this will fix the clash in the editor world and also in game for dynamic Navmesh, but not in game for static Navmesh.
+		if (NewId != OldId)
 		{
-			AddDirtyArea(LinkBounds, FNavigationOctreeController::OctreeUpdate_Modifiers);
+			CustomLink.UpdateLinkId(NewId);
+			UE_LOG(LogNavLink, VeryVerbose, TEXT("%hs new navlink ID %llu."), __FUNCTION__, CustomLink.GetId());
+
+			const FBox LinkBounds = ComputeCustomLinkBounds(CustomLink);
+			if (LinkBounds.IsValid)
+			{
+				AddDirtyArea(LinkBounds, FNavigationOctreeController::OctreeUpdate_Modifiers);
+			}
 		}
 	}
 
-	CustomLinksMap.Add(LinkId, FNavigationSystem::FCustomLinkOwnerInfo(&CustomLink));
+	UE_CLOG(bGenerateNewId && CustomNavLinksMap.Contains(CustomLink.GetId()), LogNavLink, Warning, TEXT("%hs New navlink ID %llu is clashing with existing ID."), __FUNCTION__, CustomLink.GetId());
+	CustomNavLinksMap.Add(CustomLink.GetId(), FNavigationSystem::FCustomLinkOwnerInfo(&CustomLink));
 }
 
 void UNavigationSystemV1::UnregisterCustomLink(INavLinkCustomInterface& CustomLink)
 {
-	CustomLinksMap.Remove(CustomLink.GetLinkId());
+	CustomNavLinksMap.Remove(CustomLink.GetId());
 }
 
-INavLinkCustomInterface* UNavigationSystemV1::GetCustomLink(uint32 UniqueLinkId) const
+INavLinkCustomInterface* UNavigationSystemV1::GetCustomLink(FNavLinkId UniqueLinkId) const
 {
-	const FNavigationSystem::FCustomLinkOwnerInfo* LinkInfo = CustomLinksMap.Find(UniqueLinkId);
+	const FNavigationSystem::FCustomLinkOwnerInfo* LinkInfo = CustomNavLinksMap.Find(UniqueLinkId);
 	return (LinkInfo && LinkInfo->IsValid()) ? LinkInfo->LinkInterface : nullptr;
 }
 
@@ -4302,7 +4336,9 @@ void UNavigationSystemV1::CleanUp(FNavigationSystem::ECleanupMode Mode)
 		if (MyWorld->WorldType == EWorldType::Game || MyWorld->WorldType == EWorldType::Editor)
 		{
 			UE_LOG(LogNavLink, VeryVerbose, TEXT("Reset navlink id on cleanup."));
-			INavLinkCustomInterface::NextUniqueId = 1;
+			PRAGMA_DISABLE_DEPRECATION_WARNINGS
+			INavLinkCustomInterface::ResetUniqueId();
+			PRAGMA_ENABLE_DEPRECATION_WARNINGS
 		}
 	}
 
diff --git a/Engine/Source/Runtime/NavigationSystem/Private/NavigationSystemTypes.cpp b/Engine/Source/Runtime/NavigationSystem/Private/NavigationSystemTypes.cpp
--- a/Engine/Source/Runtime/NavigationSystem/Private/NavigationSystemTypes.cpp
+++ b/Engine/Source/Runtime/NavigationSystem/Private/NavigationSystemTypes.cpp
@@ -271,21 +271,32 @@ UObject* INavLinkCustomInterface::GetLinkOwner() const
 
 uint32 INavLinkCustomInterface::GetUniqueId()
 {
-	UE_LOG(LogNavLink, VeryVerbose, TEXT("%s id: %u."), ANSI_TO_TCHAR(__FUNCTION__), NextUniqueId);
+	PRAGMA_DISABLE_DEPRECATION_WARNINGS
+	UE_LOG(LogNavLink, VeryVerbose, TEXT("%hs id: %u."), __FUNCTION__, NextUniqueId);
 	return NextUniqueId++;
+	PRAGMA_ENABLE_DEPRECATION_WARNINGS
 }
 
-void INavLinkCustomInterface::UpdateUniqueId(uint32 AlreadyUsedId)
+void INavLinkCustomInterface::UpdateUniqueId(FNavLinkId AlreadyUsedId)
 {
-	UE_CLOG(AlreadyUsedId + 1 > NextUniqueId, LogNavLink, VeryVerbose, TEXT("%s, updating NextUniqueId to: %u."), ANSI_TO_TCHAR(__FUNCTION__), AlreadyUsedId + 1)
-	NextUniqueId = FMath::Max(NextUniqueId, AlreadyUsedId + 1);
+	// Only update NextUniqueId for old style incremental Ids.
+	if (AlreadyUsedId.IsLegacyId())
+	{
+		PRAGMA_DISABLE_DEPRECATION_WARNINGS
+		UE_CLOG(AlreadyUsedId.GetId() + 1 > NextUniqueId, LogNavLink, VeryVerbose, TEXT("%hs, updating NextUniqueId to: %llu."), __FUNCTION__, AlreadyUsedId.GetId() + 1)
+		const uint64 NextId = FMath::Max((uint64)NextUniqueId, AlreadyUsedId.GetId() + 1);
+		ensureMsgf(NextId <= TNumericLimits<uint32>::Max(), TEXT("Overflowing uint32 using legacy nav link id system!"));
+
+		NextUniqueId = (uint32)NextId;
+		PRAGMA_ENABLE_DEPRECATION_WARNINGS
+	}
 }
 
 FNavigationLink INavLinkCustomInterface::GetModifier(const INavLinkCustomInterface* CustomNavLink)
 {
 	FNavigationLink LinkMod;
 	LinkMod.SetAreaClass(CustomNavLink->GetLinkAreaClass());
-	LinkMod.UserId = CustomNavLink->GetLinkId();
+	LinkMod.NavLinkId = CustomNavLink->GetId();
 
 	ENavLinkDirection::Type LinkDirection = ENavLinkDirection::BothWays;
 	CustomNavLink->GetLinkData(LinkMod.Left, LinkMod.Right, LinkDirection);
@@ -297,11 +308,15 @@ FNavigationLink INavLinkCustomInterface::GetModifier(const INavLinkCustomInterfa
 
 void INavLinkCustomInterface::OnPreWorldInitialization(UWorld* World, const UWorld::InitializationValues IVS)
 {
+	PRAGMA_DISABLE_DEPRECATION_WARNINGS
 	ResetUniqueId();
+	PRAGMA_ENABLE_DEPRECATION_WARNINGS
 }
 
 void INavLinkCustomInterface::ResetUniqueId()
 {
 	UE_LOG(LogNavLink, VeryVerbose, TEXT("Reset navlink id."));
+	PRAGMA_DISABLE_DEPRECATION_WARNINGS
 	NextUniqueId = 1;
+	PRAGMA_ENABLE_DEPRECATION_WARNINGS
 }
diff --git a/Engine/Source/Runtime/NavigationSystem/Public/NavLinkCustomComponent.h b/Engine/Source/Runtime/NavigationSystem/Public/NavLinkCustomComponent.h
--- a/Engine/Source/Runtime/NavigationSystem/Public/NavLinkCustomComponent.h
+++ b/Engine/Source/Runtime/NavigationSystem/Public/NavLinkCustomComponent.h
@@ -37,8 +37,9 @@ class NAVIGATIONSYSTEM_API UNavLinkCustomComponent : public UNavRelevantComponen
 	virtual void GetLinkData(FVector& LeftPt, FVector& RightPt, ENavLinkDirection::Type& Direction) const override;
 	virtual void GetSupportedAgents(FNavAgentSelector& OutSupportedAgents) const override;
 	virtual TSubclassOf<UNavArea> GetLinkAreaClass() const override;
-	virtual uint32 GetLinkId() const override;
-	virtual void UpdateLinkId(uint32 NewUniqueId) override;
+	virtual FNavLinkAuxiliaryId GetAuxiliaryId() const override;
+	virtual FNavLinkId GetId() const override;
+	virtual void UpdateLinkId(FNavLinkId NewUniqueId) override;
 	virtual bool IsLinkPathfindingAllowed(const UObject* Querier) const override;
 	virtual bool OnLinkMoveStarted(UObject* PathComp, const FVector& DestPoint) override;
 	virtual void OnLinkMoveFinished(UObject* PathComp) override;
@@ -58,6 +59,7 @@ class NAVIGATIONSYSTEM_API UNavLinkCustomComponent : public UNavRelevantComponen
 	void ApplyComponentInstanceData(struct FNavLinkCustomInstanceData* ComponentInstanceData);
 
 	//~ Begin UObject Interface
+	virtual void Serialize(FArchive& Ar) override;
 	virtual void PostLoad() override;
 #if WITH_EDITOR
 	virtual void PostEditImport() override;
@@ -146,10 +148,23 @@ protected:
 
 protected:
 
-	/** link Id assigned by navigation system */
+	UE_DEPRECATED(5.4, "LinkIds are now based on FNavLinkId. Use CustomLinkId instead. This Id is no longer used by the engine.")
 	UPROPERTY()
 	uint32 NavLinkUserId;
 
+	/** link Id assigned by navigation system */
+	UPROPERTY()
+	FNavLinkId CustomLinkId;
+
+	/** 
+	 *  Assigned in the constructor. This uniquely identifies a component in an Actor, but will not be unique between duplicate level instances.
+	 *  containing the same Actor.
+	 *  This is Hashed with the Actor Instance FGuid to create the CustomLinkId so that Actors with more than one UNavLinkCustomComponent can have a 
+	 *  completely unique ID per UNavLinkCustomComponent even across level instances.
+	 **/
+	UPROPERTY()
+	FNavLinkAuxiliaryId AuxiliaryCustomLinkId;
+
 	/** area class to use when link is enabled */
 	UPROPERTY(EditAnywhere, Category=SmartLink)
 	TSubclassOf<UNavArea> EnabledAreaClass;
@@ -248,7 +263,6 @@ public:
 	FNavLinkCustomInstanceData() = default;
 	FNavLinkCustomInstanceData(const UNavLinkCustomComponent* SourceComponent)
 		: FActorComponentInstanceData(SourceComponent)
-		, NavLinkUserId(0)
 	{}
 
 	virtual ~FNavLinkCustomInstanceData() = default;
@@ -269,5 +283,8 @@ public:
 	}
 
 	UPROPERTY()
-	uint32 NavLinkUserId = 0;
+	FNavLinkId CustomLinkId;
+
+	UPROPERTY()
+	FNavLinkAuxiliaryId AuxiliaryCustomLinkId;
 };
\ No newline at end of file
diff --git a/Engine/Source/Runtime/NavigationSystem/Public/NavLinkCustomInterface.h b/Engine/Source/Runtime/NavigationSystem/Public/NavLinkCustomInterface.h
--- a/Engine/Source/Runtime/NavigationSystem/Public/NavLinkCustomInterface.h
+++ b/Engine/Source/Runtime/NavigationSystem/Public/NavLinkCustomInterface.h
@@ -44,13 +44,21 @@ class NAVIGATIONSYSTEM_API INavLinkCustomInterface
 	/** Get basic link data: area class (null = default walkable) */
 	virtual TSubclassOf<UNavArea> GetLinkAreaClass() const { return nullptr; }
 
+	virtual FNavLinkAuxiliaryId GetAuxiliaryId() const { return FNavLinkAuxiliaryId::Invalid; }
+
+	UE_DEPRECATED(5.4, "LinkIds are now based on a FNavLinkId. Call GetId() instead. This function only returns Invalid Id.")
+	virtual uint32 GetLinkId() const final { return FNavLinkId::Invalid.GetId(); }
+
 	/** Get unique ID number for custom link
 	 *  Owner should get its unique ID by calling INavLinkCustomInterface::GetUniqueId() and store it
 	 */
-	virtual uint32 GetLinkId() const { return 0; }
+	virtual FNavLinkId GetId() const { return FNavLinkId::Invalid; }
 
-	/** Update unique ID number for custom link by navigation system */
-	virtual void UpdateLinkId(uint32 NewUniqueId) { }
+	UE_DEPRECATED(5.4, "LinkIds are now based on a FNavLinkId. Call the version of this function that takes a FNavLinkId. This function now has no effect.")
+	virtual void UpdateLinkId(uint32 NewUniqueId) final {}
+
+	/** Update unique ID number for custom link by navigation system. */
+	virtual void UpdateLinkId(FNavLinkId NewUniqueId) {}
 
 	/** Get object owner of navigation link, used for creating containers with multiple links */
 	virtual UObject* GetLinkOwner() const;
@@ -68,17 +76,25 @@ class NAVIGATIONSYSTEM_API INavLinkCustomInterface
 	/** Notify called when agent finishes using this link for movement */
 	virtual void OnLinkMoveFinished(class UObject* PathComp) {}
 
-	/** Helper function: returns unique ID number for custom links */
+	UE_DEPRECATED(5.4, "LinkIds are now based on FNavLinkId using FNavLinkId::GenerateUniqueId(). This function will still generate an incremental Id however it does not work well in all circumstances.")
 	static uint32 GetUniqueId();
 
 	/** Helper function: bump unique ID numbers above given one */
-	static void UpdateUniqueId(uint32 AlreadyUsedId);
+	UE_DEPRECATED(5.4, "LinkIds are now based on FNavLinkId. If your project is still using any of the old incremental Ids (saved in actors in levels or licensee code) then this function must be called still (typically by existing engine code), otherwise it is not necessary.")
+	static void UpdateUniqueId(FNavLinkId AlreadyUsedId);
+
+	UE_DEPRECATED(5.4, "LinkIds are now based on FNavLinkId. You may need to call the other version of this function that takes a FNavLinkId. This function has no effect.")
+	static void UpdateUniqueId(uint32 AlreadyUsedId) {}
 
 	/** Helper function: create modifier for navigation data export */
 	static FNavigationLink GetModifier(const INavLinkCustomInterface* CustomNavLink);
+	
+	UE_DEPRECATED(5.4, "LinkIds are now based on a FNavLinkId Hash. If your project is still using any of the old incremental Ids then this function must be called still (typically by existing engine code), otherwise it is not necessary.")
+	static void ResetUniqueId();
 
 	static void OnPreWorldInitialization(UWorld* World, const UWorld::InitializationValues IVS);
-	static void ResetUniqueId();
 
+	UE_DEPRECATED(5.4, "LinkIds are now based on FNavLinkId using FNavLinkId::GenerateUniqueId().")
 	static uint32 NextUniqueId;
 };
+
diff --git a/Engine/Source/Runtime/NavigationSystem/Public/NavMesh/NavMeshPath.h b/Engine/Source/Runtime/NavigationSystem/Public/NavMesh/NavMeshPath.h
--- a/Engine/Source/Runtime/NavigationSystem/Public/NavMesh/NavMeshPath.h
+++ b/Engine/Source/Runtime/NavigationSystem/Public/NavMesh/NavMeshPath.h
@@ -39,6 +39,13 @@ struct NAVIGATIONSYSTEM_API FNavMeshPath : public FNavigationPath
 
 	FNavMeshPath();
 
+	PRAGMA_DISABLE_DEPRECATION_WARNINGS
+	FNavMeshPath(const FNavMeshPath&) = default;
+	FNavMeshPath(FNavMeshPath&& Other) = default;
+	FNavMeshPath& operator=(const FNavMeshPath& Other) = default;
+	FNavMeshPath& operator=(FNavMeshPath&& Other) = default;
+	PRAGMA_ENABLE_DEPRECATION_WARNINGS
+
 	FORCEINLINE void SetWantsStringPulling(const bool bNewWantsStringPulling) { bWantsStringPulling = bNewWantsStringPulling; }
 	FORCEINLINE bool WantsStringPulling() const { return bWantsStringPulling; }
 	FORCEINLINE bool IsStringPulled() const { return bStringPulled; }
@@ -93,8 +100,8 @@ struct NAVIGATIONSYSTEM_API FNavMeshPath : public FNavigationPath
 	/** check if path (all polys in corridor) contains given node */
 	virtual bool ContainsNode(NavNodeRef NodeRef) const override { return PathCorridor.Contains(NodeRef); }
 
-	virtual bool ContainsCustomLink(uint32 UniqueLinkId) const override { return CustomLinkIds.Contains(UniqueLinkId); }
-	virtual bool ContainsAnyCustomLink() const override { return CustomLinkIds.Num() > 0; }
+	virtual bool ContainsCustomLink(FNavLinkId UniqueLinkId) const override { return CustomNavLinkIds.Contains(UniqueLinkId); }
+	virtual bool ContainsAnyCustomLink() const override { return CustomNavLinkIds.Num() > 0; }
 
 	bool IsPathSegmentANavLink(const int32 PathSegmentStartIndex) const;
 
@@ -141,8 +148,11 @@ public:
 	TArray<FVector::FReal> PathCorridorCost;
 
 	/** set of unique link Ids */
+	UE_DEPRECATED(5.4, "LinkIds are now based on FNavLinkId. Use CustomNavLinkIds instead. CustomLinkIds array is no longer populated or used in the engine")
 	TArray<uint32> CustomLinkIds;
 
+	TArray<FNavLinkId> CustomNavLinkIds;
+
 private:
 	/** sequence of FVector pairs where each pair represents navmesh portal edge between two polygons navigation corridor.
 	*	Note, that it should always be accessed via GetPathCorridorEdges() since PathCorridorEdges content is generated
diff --git a/Engine/Source/Runtime/NavigationSystem/Public/NavMesh/PImplRecastNavMesh.h b/Engine/Source/Runtime/NavigationSystem/Public/NavMesh/PImplRecastNavMesh.h
--- a/Engine/Source/Runtime/NavigationSystem/Public/NavMesh/PImplRecastNavMesh.h
+++ b/Engine/Source/Runtime/NavigationSystem/Public/NavMesh/PImplRecastNavMesh.h
@@ -27,6 +27,7 @@
 #endif
 
 class FRecastNavMeshGenerator;
+struct FNavLinkId;
 
 #if WITH_RECAST
 
@@ -149,7 +150,11 @@ public:
 	void GetEdgesForPathCorridor(const TArray<NavNodeRef>* PathCorridor, TArray<FNavigationPortalEdge>* PathCorridorEdges) const;
 
 	/** finds stringpulled path from given corridor */
-	bool FindStraightPath(const FVector& StartLoc, const FVector& EndLoc, const TArray<NavNodeRef>& PathCorridor, TArray<FNavPathPoint>& PathPoints, TArray<uint32>* CustomLinks = NULL) const;
+	bool FindStraightPath(const FVector& StartLoc, const FVector& EndLoc, const TArray<NavNodeRef>& PathCorridor, TArray<FNavPathPoint>& PathPoints, TArray<FNavLinkId>* CustomLinks = NULL) const;
+
+	/** finds stringpulled path from given corridor */
+	UE_DEPRECATED(5.4, "Please use FindStraightPath with the TArray<FNavPathPoint>* CustomLinks. This function has no effect.")
+	bool FindStraightPath(const FVector& StartLoc, const FVector& EndLoc, const TArray<NavNodeRef>& PathCorridor, TArray<FNavPathPoint>& PathPoints, TArray<uint32>* CustomLinks) const { return false; }
 
 	/** Filters nav polys in PolyRefs with Filter */
 	bool FilterPolys(TArray<NavNodeRef>& PolyRefs, const FRecastQueryFilter* Filter, const UObject* Owner) const;
@@ -157,8 +162,11 @@ public:
 	/** Get all polys from tile */
 	bool GetPolysInTile(int32 TileIndex, TArray<FNavPoly>& Polys) const;
 
+	UE_DEPRECATED(5.4, "Please use the version of this function that takes a FNavLinkId. This function has no effect.")
+	void UpdateNavigationLinkArea(int32 UserId, uint8 AreaType, uint16 PolyFlags) const {}
+
 	/** Updates area on polygons creating point-to-point connection with given UserId */
-	void UpdateNavigationLinkArea(int32 UserId, uint8 AreaType, uint16 PolyFlags) const;
+	void UpdateNavigationLinkArea(FNavLinkId UserId, uint8 AreaType, uint16 PolyFlags) const;
 
 #if WITH_NAVMESH_SEGMENT_LINKS
 	/** Updates area on polygons creating segment-to-segment connection with given UserId */
@@ -190,7 +198,12 @@ public:
 	/** Decode poly ID into FNavTileRef and poly index */
 	bool GetPolyTileRef(NavNodeRef PolyId, uint32& OutPolyIndex, FNavTileRef& OutTileRef) const;
 	/** Retrieves user ID for given offmesh link poly */
-	uint32 GetLinkUserId(NavNodeRef LinkPolyID) const;
+	UE_DEPRECATED(5.4, "Please use GetNavLinkUserId() instead. This function only returns Invalid.")
+	uint32 GetLinkUserId(NavNodeRef LinkPolyID) const
+	{
+		return FNavLinkId::Invalid.GetId();
+	}
+	FNavLinkId GetNavLinkUserId(NavNodeRef LinkPolyID) const;
 	/** Retrieves start and end point of offmesh link */
 	bool GetLinkEndPoints(NavNodeRef LinkPolyID, FVector& PointA, FVector& PointB) const;
 	/** Check if poly is a custom link */
diff --git a/Engine/Source/Runtime/NavigationSystem/Public/NavMesh/RecastNavMesh.h b/Engine/Source/Runtime/NavigationSystem/Public/NavMesh/RecastNavMesh.h
--- a/Engine/Source/Runtime/NavigationSystem/Public/NavMesh/RecastNavMesh.h
+++ b/Engine/Source/Runtime/NavigationSystem/Public/NavMesh/RecastNavMesh.h
@@ -1232,8 +1232,11 @@ public:
 
 	virtual void UpdateCustomLink(const INavLinkCustomInterface* CustomLink) override;
 
+	UE_DEPRECATED(5.4, "Use version of this function that takes a FNavLinkId. This function now has no effect.")
+	void UpdateNavigationLinkArea(int32 UserId, TSubclassOf<UNavArea> AreaClass) const {}
+
 	/** update area class and poly flags for all offmesh links with given UserId */
-	void UpdateNavigationLinkArea(int32 UserId, TSubclassOf<UNavArea> AreaClass) const;
+	void UpdateNavigationLinkArea(FNavLinkId UserId, TSubclassOf<UNavArea> AreaClass) const;
 
 #if WITH_NAVMESH_SEGMENT_LINKS
 	/** update area class and poly flags for all offmesh segment links with given UserId */
@@ -1260,8 +1263,14 @@ public:
 	/** dtNavMesh getter */
 	const dtNavMesh* GetRecastMesh() const;
 
+	UE_DEPRECATED(5.4, "Please use GetNavLinkUserId() instead. This function only returns Invalid.")
+	int32 GetLinkUserId(NavNodeRef LinkPolyID) const
+	{
+		return FNavLinkId::Invalid.GetId();
+	}
+
 	/** Retrieves LinkUserID associated with indicated PolyID */
-	uint32 GetLinkUserId(NavNodeRef LinkPolyID) const;
+	FNavLinkId GetNavLinkUserId(NavNodeRef LinkPolyID) const;
 
 	FColor GetAreaIDColor(uint8 AreaID) const;
 
@@ -1379,8 +1388,11 @@ public:
 	/** Check if poly is a custom link */
 	bool IsCustomLink(NavNodeRef PolyRef) const;
 
+	UE_DEPRECATED(5.4, "Use new override of this function with Array<FNavLinkId>* CustomLinks. This function has no effect.")
+	bool FindStraightPath(const FVector& StartLoc, const FVector& EndLoc, const TArray<NavNodeRef>& PathCorridor, TArray<FNavPathPoint>& PathPoints, TArray<uint32>* CustomLinks) const { return false; }
+
 	/** finds stringpulled path from given corridor */
-	bool FindStraightPath(const FVector& StartLoc, const FVector& EndLoc, const TArray<NavNodeRef>& PathCorridor, TArray<FNavPathPoint>& PathPoints, TArray<uint32>* CustomLinks = NULL) const;
+	bool FindStraightPath(const FVector& StartLoc, const FVector& EndLoc, const TArray<NavNodeRef>& PathCorridor, TArray<FNavPathPoint>& PathPoints, TArray<FNavLinkId>* CustomLinks = NULL) const;
 
 	/** Runs A* pathfinding on navmesh and collect data for every step */
 	int32 DebugPathfinding(const FPathFindingQuery& Query, TArray<FRecastDebugPathfindingData>& Steps);
diff --git a/Engine/Source/Runtime/NavigationSystem/Public/NavMesh/RecastQueryFilter.h b/Engine/Source/Runtime/NavigationSystem/Public/NavMesh/RecastQueryFilter.h
--- a/Engine/Source/Runtime/NavigationSystem/Public/NavMesh/RecastQueryFilter.h
+++ b/Engine/Source/Runtime/NavigationSystem/Public/NavMesh/RecastQueryFilter.h
@@ -59,7 +59,7 @@ public:
 struct NAVIGATIONSYSTEM_API FRecastSpeciaLinkFilter : public dtQuerySpecialLinkFilter
 {
 	FRecastSpeciaLinkFilter(UNavigationSystemV1* NavSystem, const UObject* Owner) : NavSys(NavSystem), SearchOwner(Owner), CachedOwnerOb(nullptr) {}
-	virtual bool isLinkAllowed(const int32 UserId) const override;
+	virtual bool isLinkAllowed(const uint64 UserId) const override;
 	virtual void initialize() override;
 
 	UNavigationSystemV1* NavSys;
diff --git a/Engine/Source/Runtime/NavigationSystem/Public/NavMesh/RecastVersion.h b/Engine/Source/Runtime/NavigationSystem/Public/NavMesh/RecastVersion.h
--- a/Engine/Source/Runtime/NavigationSystem/Public/NavMesh/RecastVersion.h
+++ b/Engine/Source/Runtime/NavigationSystem/Public/NavMesh/RecastVersion.h
@@ -2,7 +2,11 @@
 
 #pragma once
 
-/** Versioning */
+/** 
+ *  Versioning, Note the correct way of doing this is now to use the Custom Version see Ar.UsingCustomVersion(XXXBranchObjectVersion::GUID) this solves issues when versioning
+ *  is changed in different branches at the same time.
+ */
+
 #define NAVMESHVER_INITIAL						1
 #define NAVMESHVER_TILED_GENERATION				2
 #define NAVMESHVER_SEAMLESS_REBUILDING_1		3
diff --git a/Engine/Source/Runtime/NavigationSystem/Public/NavigationData.h b/Engine/Source/Runtime/NavigationSystem/Public/NavigationData.h
--- a/Engine/Source/Runtime/NavigationSystem/Public/NavigationData.h
+++ b/Engine/Source/Runtime/NavigationSystem/Public/NavigationData.h
@@ -78,6 +78,11 @@ struct NAVIGATIONSYSTEM_API FNavigationPath : public TSharedFromThis<FNavigation
 	virtual ~FNavigationPath()
 	{ }
 
+	PRAGMA_DISABLE_DEPRECATION_WARNINGS
+	FNavigationPath(const FNavigationPath&) = default;
+	FNavigationPath& operator=(const FNavigationPath& Other) = default;
+	PRAGMA_ENABLE_DEPRECATION_WARNINGS
+
 	FORCEINLINE bool IsValid() const
 	{
 		return bIsReady && PathPoints.Num() > 1 && bUpToDate;
@@ -234,7 +239,10 @@ struct NAVIGATIONSYSTEM_API FNavigationPath : public TSharedFromThis<FNavigation
 #endif // ENABLE_VISUAL_LOG
 
 	/** check if path contains specific custom nav link */
-	virtual bool ContainsCustomLink(uint32 UniqueLinkId) const;
+	UE_DEPRECATED(5.4, "Use version that takes FNavLinkId instead. This function only returns false.")
+	virtual bool ContainsCustomLink(uint32 UniqueLinkId) const final {	return false; }
+
+	virtual bool ContainsCustomLink(FNavLinkId UniqueLinkId) const;
 
 	/** check if path contains any custom nav link */
 	virtual bool ContainsAnyCustomLink() const;
diff --git a/Engine/Source/Runtime/NavigationSystem/Public/NavigationSystem.h b/Engine/Source/Runtime/NavigationSystem/Public/NavigationSystem.h
--- a/Engine/Source/Runtime/NavigationSystem/Public/NavigationSystem.h
+++ b/Engine/Source/Runtime/NavigationSystem/Public/NavigationSystem.h
@@ -909,8 +909,11 @@ public:
 	static void RequestCustomLinkRegistering(INavLinkCustomInterface& CustomLink, UObject* OwnerOb);
 	static void RequestCustomLinkUnregistering(INavLinkCustomInterface& CustomLink, UObject* ObjectOb);
 
+	UE_DEPRECATED(5.4, "LinkIds are now based on FNavLinkId call the version of this function that takes FNavLinkId. This function only returns nullptr.")
+	INavLinkCustomInterface* GetCustomLink(uint32 UniqueLinkId) const { return nullptr; };
+
 	/** find custom link by unique ID */
-	INavLinkCustomInterface* GetCustomLink(uint32 UniqueLinkId) const;
+	INavLinkCustomInterface* GetCustomLink(FNavLinkId UniqueLinkId) const;
 
 	/** updates custom link for all active navigation data instances */
 	void UpdateCustomLink(const INavLinkCustomInterface* CustomLink);
@@ -1131,8 +1134,11 @@ protected:
 	FNavigationOctreeController DefaultOctreeController;
 
 	/** Map of all custom navigation links, that are relevant for path following */
+	UE_DEPRECATED(5.4, "LinkIds are now based on FNavLinkId. Use CustomNavLinksMap instead, CustomLinksMap is no longer populated or used in the engine.")
 	TMap<uint32, FNavigationSystem::FCustomLinkOwnerInfo> CustomLinksMap;
 
+	TMap<FNavLinkId, FNavigationSystem::FCustomLinkOwnerInfo> CustomNavLinksMap;
+
 	FNavigationDirtyAreasController DefaultDirtyAreasController;
 
 	// async queries
diff --git a/Engine/Source/Runtime/Navmesh/Private/Detour/DetourNavMesh.cpp b/Engine/Source/Runtime/Navmesh/Private/Detour/DetourNavMesh.cpp
--- a/Engine/Source/Runtime/Navmesh/Private/Detour/DetourNavMesh.cpp
+++ b/Engine/Source/Runtime/Navmesh/Private/Detour/DetourNavMesh.cpp
@@ -2541,7 +2541,7 @@ void dtNavMesh::updateOffMeshSegmentConnectionByUserId(unsigned int userId, unsi
 #endif // WITH_NAVMESH_SEGMENT_LINKS
 //@UE END
 
-void dtNavMesh::updateOffMeshConnectionByUserId(unsigned int userId, unsigned char newArea, unsigned short newFlags)
+void dtNavMesh::updateOffMeshConnectionByUserId(unsigned long long int userId, unsigned char newArea, unsigned short newFlags)
 {
 	for (int it = 0; it < m_maxTiles; it++)
 	{
diff --git a/Engine/Source/Runtime/Navmesh/Public/Detour/DetourCommon.h b/Engine/Source/Runtime/Navmesh/Public/Detour/DetourCommon.h
--- a/Engine/Source/Runtime/Navmesh/Public/Detour/DetourCommon.h
+++ b/Engine/Source/Runtime/Navmesh/Public/Detour/DetourCommon.h
@@ -565,6 +565,24 @@ inline void dtSwapEndian(int* v)
 	dtSwapByte(x+0, x+3); dtSwapByte(x+1, x+2);
 }
 
+inline void dtSwapEndian(unsigned long long int* v)
+{
+	unsigned char* x = (unsigned char*)v;
+	dtSwapByte(x + 0, x + 7); 
+	dtSwapByte(x + 1, x + 6);
+	dtSwapByte(x + 2, x + 5);
+	dtSwapByte(x + 3, x + 4);
+}
+
+inline void dtSwapEndian(long long int* v)
+{
+	unsigned char* x = (unsigned char*)v;
+	dtSwapByte(x + 0, x + 7);
+	dtSwapByte(x + 1, x + 6);
+	dtSwapByte(x + 2, x + 5);
+	dtSwapByte(x + 3, x + 4);
+}
+
 inline void dtSwapEndian(float* v)
 {
 	unsigned char* x = (unsigned char*)v;
diff --git a/Engine/Source/Runtime/Navmesh/Public/Detour/DetourNavMesh.h b/Engine/Source/Runtime/Navmesh/Public/Detour/DetourNavMesh.h
--- a/Engine/Source/Runtime/Navmesh/Public/Detour/DetourNavMesh.h
+++ b/Engine/Source/Runtime/Navmesh/Public/Detour/DetourNavMesh.h
@@ -320,7 +320,7 @@ struct dtOffMeshConnection
 	//@UE END
 
 	/// The id of the offmesh connection. (User assigned when the navigation mesh is built.)
-	unsigned int userId;
+	unsigned long long int userId;
 
 	/// The polygon reference of the connection within the tile.
 	unsigned short poly;
@@ -647,7 +647,7 @@ public:
 	/// Updates area and flags for specified off-mesh connection: point type
 	///  @param[in] userId	User Id of connection
 	///	 @param[in] newArea	Area code to apply
-	void updateOffMeshConnectionByUserId(unsigned int userId, unsigned char newArea, unsigned short newFlags);
+	void updateOffMeshConnectionByUserId(unsigned long long int userId, unsigned char newArea, unsigned short newFlags);
 
 	//@UE BEGIN
 #if WITH_NAVMESH_SEGMENT_LINKS
diff --git a/Engine/Source/Runtime/Navmesh/Public/Detour/DetourNavMeshBuilder.h b/Engine/Source/Runtime/Navmesh/Public/Detour/DetourNavMeshBuilder.h
--- a/Engine/Source/Runtime/Navmesh/Public/Detour/DetourNavMeshBuilder.h
+++ b/Engine/Source/Runtime/Navmesh/Public/Detour/DetourNavMeshBuilder.h
@@ -38,7 +38,7 @@ struct dtOffMeshLinkCreateParams
 	/// Off-mesh connection height, less than 0 = use step height [Unit: wu]
 	dtReal snapHeight;
 	/// The user defined ids of the off-mesh connection.
-	unsigned int userID;
+	unsigned long long int userID;
 	/// User defined flags assigned to the polys of off-mesh connections
 	unsigned short polyFlag;
 	/// User defined area ids assigned to the off-mesh connections
diff --git a/Engine/Source/Runtime/Navmesh/Public/Detour/DetourNavMeshQuery.h b/Engine/Source/Runtime/Navmesh/Public/Detour/DetourNavMeshQuery.h
--- a/Engine/Source/Runtime/Navmesh/Public/Detour/DetourNavMeshQuery.h
+++ b/Engine/Source/Runtime/Navmesh/Public/Detour/DetourNavMeshQuery.h
@@ -52,7 +52,11 @@ struct NAVMESH_API dtQuerySpecialLinkFilter
 
 	/// Returns true if special link can be visited.  (I.e. Is traversable.)
 	///  @param[in]		UserId		Unique Id of link
-	virtual bool isLinkAllowed(const int UserId) const { return true; }
+	
+	virtual bool isLinkAllowed(const unsigned long long int UserId) const { return true; }
+
+	UE_DEPRECATED(5.4, "LinkIds are now based on a 64 bit uint. Use the version of this function that takes an unsigned long long int")
+	virtual bool isLinkAllowed(const int UserId) const final { return true; }
 
 	/// Called before accessing in A* loop (can be called multiple time for updateSlicedFindPath)
 	virtual void initialize() {}
