[Navmesh] Addition of invoker priorities

diff --git a/Engine/Source/Runtime/Engine/Classes/AI/Navigation/NavigationInvokerPriority.h b/Engine/Source/Runtime/Engine/Classes/AI/Navigation/NavigationInvokerPriority.h
new file mode 100644
--- /dev/null
+++ b/Engine/Source/Runtime/Engine/Classes/AI/Navigation/NavigationInvokerPriority.h
@@ -0,0 +1,16 @@
+﻿// Copyright Epic Games, Inc. All Rights Reserved.
+
+#pragma once
+
+#include "NavigationInvokerPriority.generated.h"
+
+UENUM()
+enum class ENavigationInvokerPriority : uint8
+{
+	VeryLow = 1,
+	Low = 2,
+	Default = 3,
+	High = 4,
+	VeryHigh = 5,
+	MAX = 6
+};
\ No newline at end of file
diff --git a/Engine/Source/Runtime/Engine/Classes/AI/Navigation/NavigationTypes.h b/Engine/Source/Runtime/Engine/Classes/AI/Navigation/NavigationTypes.h
--- a/Engine/Source/Runtime/Engine/Classes/AI/Navigation/NavigationTypes.h
+++ b/Engine/Source/Runtime/Engine/Classes/AI/Navigation/NavigationTypes.h
@@ -136,6 +136,14 @@ enum class ENavDataGatheringModeConfig : uint8
 	Lazy
 };
 
+UENUM()
+enum class ENavigationSortPendingTilesMethod : uint8
+{
+	SortWithSeedLocations,
+	SortByPriority,
+	None
+};
+
 //
 // Used to gather per instance transforms in a specific area
 //
diff --git a/Engine/Source/Runtime/NavigationSystem/Private/NavMesh/RecastNavMesh.cpp b/Engine/Source/Runtime/NavigationSystem/Private/NavMesh/RecastNavMesh.cpp
--- a/Engine/Source/Runtime/NavigationSystem/Private/NavMesh/RecastNavMesh.cpp
+++ b/Engine/Source/Runtime/NavigationSystem/Private/NavMesh/RecastNavMesh.cpp
@@ -3487,10 +3487,12 @@ void ARecastNavMesh::UpdateActiveTiles(const TArray<FNavigationInvokerRaw>& Invo
 
 	TArray<FIntPoint>& ActiveTiles = GetActiveTiles();
 	TArray<FIntPoint> OldActiveSet = ActiveTiles;
-	TArray<FIntPoint> TilesInMinDistance;
+	TArray<FNavMeshDirtyTileElement> TilesInMinDistance;
 	TArray<FIntPoint> TilesInMaxDistance;
+	TArray<FIntPoint> TileToAppend;
 	TilesInMinDistance.Reserve(ActiveTiles.Num());
 	TilesInMaxDistance.Reserve(ActiveTiles.Num());
+	TileToAppend.Reserve(ActiveTiles.Num());
 	ActiveTiles.Reset();
 
 	for (const FNavigationInvokerRaw& Invoker : InvokerLocations)
@@ -3520,14 +3522,25 @@ void ARecastNavMesh::UpdateActiveTiles(const TArray<FNavigationInvokerRaw>& Invo
 
 					if (DistanceSq < TileCenterDistanceToAddSq)
 					{
-						TilesInMinDistance.AddUnique(FIntPoint(X, Y));
+						// Add unique tile 
+						FNavMeshDirtyTileElement* FoundTile = TilesInMinDistance.FindByPredicate([X, Y](const FNavMeshDirtyTileElement& Tile){ return Tile.Coordinates == FIntPoint(X, Y);});
+						if (FoundTile)
+						{
+							// Update the priority if already existing
+							FoundTile->InvokerPriority = FMath::Max(FoundTile->InvokerPriority, Invoker.Priority);
+						}
+						else
+						{
+							TilesInMinDistance.Add(FNavMeshDirtyTileElement{FIntPoint(X,Y), Invoker.Priority});
+							TileToAppend.Add(FIntPoint(X,Y));
+						}
 					}
 				}
 			}
 		}
 	}
 
-	ActiveTiles.Append(TilesInMinDistance);
+	ActiveTiles.Append(TileToAppend);
 
 	TArray<FIntPoint> TilesToRemove;
 	TilesToRemove.Reserve(OldActiveSet.Num());
@@ -3546,19 +3559,25 @@ void ARecastNavMesh::UpdateActiveTiles(const TArray<FNavigationInvokerRaw>& Invo
 	}
 
 	// Find tiles to update
-	TArray<FIntPoint> TilesToUpdate;
+	TArray<FNavMeshDirtyTileElement> TilesToUpdate;
 	TilesToUpdate.Reserve(ActiveTiles.Num());
 	for (int32 Index = TilesInMinDistance.Num() - 1; Index >= 0; --Index)
 	{
 		// Check if it's a new tile (not in the active set)
-		const FIntPoint& Tile = TilesInMinDistance[Index];
-		if (OldActiveSet.Find(Tile) == INDEX_NONE)
+		const FNavMeshDirtyTileElement& Tile = TilesInMinDistance[Index];
+		if (OldActiveSet.Find(Tile.Coordinates) == INDEX_NONE)
 		{
 			TilesToUpdate.Add(Tile);
 		}
 	}
 
-	UE_VLOG(this, LogNavigation, Log, TEXT("Updating active tiles: %d to remove, %d to update"), TilesToRemove.Num(), TilesToUpdate.Num());
+	UE_SUPPRESS(LogNavigation, Log,
+	{
+		if (TilesToRemove.Num() != 0 || TilesToUpdate.Num() != 0)
+		{
+			UE_VLOG(this, LogNavigation, Log, TEXT("Updating active tiles: %d to remove, %d to update"), TilesToRemove.Num(), TilesToUpdate.Num());
+		}
+	});
 
 	RemoveTiles(TilesToRemove);
 	RebuildTile(TilesToUpdate);
@@ -3581,7 +3600,19 @@ void ARecastNavMesh::RemoveTiles(const TArray<FIntPoint>& Tiles)
 	}
 }
 
+// Deprecated
 void ARecastNavMesh::RebuildTile(const TArray<FIntPoint>& Tiles)
+{
+	TArray<FNavMeshDirtyTileElement> ActiveTiles;
+	ActiveTiles.Reserve(Tiles.Num());
+	for (const FIntPoint& Point : Tiles)
+	{
+		ActiveTiles.Add(FNavMeshDirtyTileElement{Point, ENavigationInvokerPriority::Default});
+	}
+	RebuildTile(ActiveTiles);
+}
+
+void ARecastNavMesh::RebuildTile(const TArray<FNavMeshDirtyTileElement>& Tiles)
 {
 	if (Tiles.Num() > 0)
 	{
@@ -3629,7 +3660,7 @@ void ARecastNavMesh::DirtyTilesInBounds(const FBox& Bounds)
 	if (OverlappingBounds.IsValid)
 	{
 		// Add tiles within the overlapping bounds
-		TArray<FIntPoint> Points;
+		TArray<FNavMeshDirtyTileElement> Tiles;
 		const FVector RcNavMeshOrigin = Unreal2RecastPoint(NavMeshOriginOffset);
 		const float TileSizeInWorldUnits = GetTileSizeUU();
 		const FRcTileBox TileBox(OverlappingBounds, RcNavMeshOrigin, TileSizeInWorldUnits);
@@ -3640,10 +3671,11 @@ void ARecastNavMesh::DirtyTilesInBounds(const FBox& Bounds)
 		{
 			for (int32 TileX = TileBox.XMin; TileX <= TileBox.XMax; ++TileX)
 			{
-				Points.Add(FIntPoint(TileX, TileY));
+				// For now, new dirtiness is made with default priority.
+				Tiles.Add(FNavMeshDirtyTileElement{FIntPoint(TileX, TileY), ENavigationInvokerPriority::Default});
 			}
 		}
-		RebuildTile(Points);
+		RebuildTile(Tiles);
 	}
 }
 
diff --git a/Engine/Source/Runtime/NavigationSystem/Private/NavMesh/RecastNavMeshGenerator.cpp b/Engine/Source/Runtime/NavigationSystem/Private/NavMesh/RecastNavMeshGenerator.cpp
--- a/Engine/Source/Runtime/NavigationSystem/Private/NavMesh/RecastNavMeshGenerator.cpp
+++ b/Engine/Source/Runtime/NavigationSystem/Private/NavMesh/RecastNavMeshGenerator.cpp
@@ -4609,7 +4609,9 @@ FRecastNavMeshGenerator::FRecastNavMeshGenerator(ARecastNavMesh& InDestNavMesh)
 	, DestNavMesh(&InDestNavMesh)
 	, bInitialized(false)
 	, bRestrictBuildingToActiveTiles(false)
+PRAGMA_DISABLE_DEPRECATION_WARNINGS
 	, bSortTilesWithSeedLocations(true)
+PRAGMA_ENABLE_DEPRECATION_WARNINGS	
 	, Version(0)
 {
 	INC_DWORD_STAT_BY(STAT_NavigationMemory, sizeof(*this));
@@ -5261,7 +5263,7 @@ void FRecastNavMeshGenerator::ResetTimeSlicedTileGeneratorSync()
 	SyncTimeSlicedData.AddGenTilesLayerIndex = 0;
 }
 
-//@TODO Investigate removing from RunningDirtyTiles here too (or atleast not using the results in any way)
+//@TODO Investigate removing from RunningDirtyTiles here too (or at least not using the results in any way)
 void FRecastNavMeshGenerator::RemoveTiles(const TArray<FIntPoint>& Tiles)
 {
 	for (const FIntPoint& TileXY : Tiles)
@@ -5285,7 +5287,18 @@ void FRecastNavMeshGenerator::RemoveTiles(const TArray<FIntPoint>& Tiles)
 	}
 }
 
+// Deprecated 
 void FRecastNavMeshGenerator::ReAddTiles(const TArray<FIntPoint>& Tiles)
+{
+	TArray<FNavMeshDirtyTileElement> ConvertedTiles;
+	for (const FIntPoint& Point : Tiles)
+	{
+		ConvertedTiles.Add(FNavMeshDirtyTileElement{Point, ENavigationInvokerPriority::Default});
+	}
+	ReAddTiles(ConvertedTiles);
+}
+	
+void FRecastNavMeshGenerator::ReAddTiles(const TArray<FNavMeshDirtyTileElement>& Tiles)
 {
 	static const FVector Expansion(1, 1, BIG_NUMBER);
 	// a little trick here - adding a dirty area so that navmesh building figures it out on its own
@@ -5297,10 +5310,11 @@ void FRecastNavMeshGenerator::ReAddTiles(const TArray<FIntPoint>& Tiles)
 	const double CurrentTimeSeconds = FPlatformTime::Seconds();
 
 	// @note we act on assumption all items in Tiles are unique
-	for (const FIntPoint& TileCoords : Tiles)
+	for (const FNavMeshDirtyTileElement& Tile : Tiles)
 	{
 		FPendingTileElement Element;
-		Element.Coord = TileCoords;
+		Element.Coord = Tile.Coordinates;
+		Element.InvokerPriority = Tile.InvokerPriority;
 		Element.bRebuildGeometry = true;
 		Element.CreationTime = CurrentTimeSeconds;
 		DirtyTiles.Add(Element);
@@ -5314,6 +5328,7 @@ void FRecastNavMeshGenerator::ReAddTiles(const TArray<FIntPoint>& Tiles)
 		FPendingTileElement* ExistingElement = DirtyTiles.Find(Element);
 		if (ExistingElement)
 		{
+			ExistingElement->InvokerPriority = FMath::Max(ExistingElement->InvokerPriority, Element.InvokerPriority);
 			ExistingElement->bRebuildGeometry |= Element.bRebuildGeometry;
 			ExistingElement->CreationTime = FMath::Min(Element.CreationTime, ExistingElement->CreationTime);
 			// Append area bounds to existing list 
@@ -6106,43 +6121,58 @@ bool FRecastNavMeshGenerator::ShouldDirtyTilesRequestedByObject(const UNavigatio
 
 void FRecastNavMeshGenerator::SortPendingBuildTiles()
 {
-	if (bSortTilesWithSeedLocations == false)
-	{
-		return;
-	}
-
-	UWorld* CurWorld = GetWorld();
-	if (CurWorld == nullptr)
+	if (SortPendingTilesMethod == ENavigationSortPendingTilesMethod::SortWithSeedLocations)
 	{
-		return;
-	}
+		UWorld* CurWorld = GetWorld();
+		if (CurWorld == nullptr)
+		{
+			return;
+		}
 
-	TArray<FVector2D> SeedLocations;
-	GetSeedLocations(*CurWorld, SeedLocations);
+		TArray<FVector2D> SeedLocations;
+		GetSeedLocations(*CurWorld, SeedLocations);
 
-	if (SeedLocations.Num() == 0)
-	{
-		// Use navmesh origin for sorting
-		SeedLocations.Add(FVector2D(TotalNavBounds.GetCenter()));
-	}
+		if (SeedLocations.Num() == 0)
+		{
+			// Use navmesh origin for sorting
+			SeedLocations.Add(FVector2D(TotalNavBounds.GetCenter()));
+		}
 
-	if (SeedLocations.Num() > 0)
-	{
-		const FVector::FReal TileSizeInWorldUnits = Config.GetTileSizeUU();
-		
-		// Calculate shortest distances between tiles and players
-		for (FPendingTileElement& Element : PendingDirtyTiles)
+		if (SeedLocations.Num() > 0)
 		{
-			const FBox TileBox = FRecastTileGenerator::CalculateTileBounds(Element.Coord.X, Element.Coord.Y, FVector::ZeroVector, TotalNavBounds, TileSizeInWorldUnits);
-			FVector2D TileCenter2D = FVector2D(TileBox.GetCenter());
-			for (FVector2D SeedLocation : SeedLocations)
+			const FVector::FReal TileSizeInWorldUnits = Config.GetTileSizeUU();
+		
+			// Calculate shortest distances between tiles and players
+			for (FPendingTileElement& Element : PendingDirtyTiles)
 			{
-				Element.SeedDistance = FMath::Min(Element.SeedDistance, FVector2D::DistSquared(TileCenter2D, SeedLocation));
+				const FBox TileBox = FRecastTileGenerator::CalculateTileBounds(Element.Coord.X, Element.Coord.Y, FVector::ZeroVector, TotalNavBounds, TileSizeInWorldUnits);
+				FVector2D TileCenter2D = FVector2D(TileBox.GetCenter());
+				for (FVector2D SeedLocation : SeedLocations)
+				{
+					Element.SeedDistance = FMath::Min(Element.SeedDistance, FVector2D::DistSquared(TileCenter2D, SeedLocation));
+				}
 			}
+
+			// Nearest tiles should be at the end of the list
+			PendingDirtyTiles.Sort();
 		}
+	}
+	else if (SortPendingTilesMethod == ENavigationSortPendingTilesMethod::SortByPriority)
+	{
+		// Highest priority should be at the end of the list
+		PendingDirtyTiles.Sort([](const FPendingTileElement& A, const FPendingTileElement& B){ return A.InvokerPriority < B.InvokerPriority; });
 
-		// nearest tiles should be at the end of the list
-		PendingDirtyTiles.Sort();
+		UE_SUPPRESS(LogNavigation, VeryVerbose,
+		{
+			for (int32 Index = 0; Index < PendingDirtyTiles.Num(); Index++)
+			{
+				const FPendingTileElement& Element = PendingDirtyTiles[Index];
+				const FVector::FReal TileSizeInWorldUnits = Config.GetTileSizeUU();
+				const FBox TileBox = FRecastTileGenerator::CalculateTileBounds(Element.Coord.X, Element.Coord.Y, FVector::ZeroVector, TotalNavBounds, TileSizeInWorldUnits);
+				const ARecastNavMesh* const OwnerNav = GetOwner();
+				UE_VLOG_BOX(OwnerNav, LogNavigation, VeryVerbose, TileBox, FColor::Silver, TEXT("Index: %i, Priority %i"), Index, Element.InvokerPriority);
+			}
+		});
 	}
 }
 
diff --git a/Engine/Source/Runtime/NavigationSystem/Private/NavigationInvokerComponent.cpp b/Engine/Source/Runtime/NavigationSystem/Private/NavigationInvokerComponent.cpp
--- a/Engine/Source/Runtime/NavigationSystem/Private/NavigationInvokerComponent.cpp
+++ b/Engine/Source/Runtime/NavigationSystem/Private/NavigationInvokerComponent.cpp
@@ -2,6 +2,7 @@
 
 #include "NavigationInvokerComponent.h"
 #include "NavigationSystem.h"
+#include "AI/Navigation/NavigationInvokerPriority.h"
 
 #include UE_INLINE_GENERATED_CPP_BY_NAME(NavigationInvokerComponent)
 
@@ -9,6 +10,7 @@ UNavigationInvokerComponent::UNavigationInvokerComponent(const FObjectInitialize
 	: Super(ObjectInitializer)
 	, TileGenerationRadius(3000)
 	, TileRemovalRadius(5000)
+	, Priority(ENavigationInvokerPriority::Default)
 {
 	bAutoActivate = true;
 	SupportedAgents.MarkInitialized();
@@ -21,7 +23,7 @@ void UNavigationInvokerComponent::Activate(bool bReset)
 	AActor* Owner = GetOwner();
 	if (Owner)
 	{
-		UNavigationSystemV1::RegisterNavigationInvoker(*Owner, TileGenerationRadius, TileRemovalRadius, SupportedAgents);
+		UNavigationSystemV1::RegisterNavigationInvoker(*Owner, TileGenerationRadius, TileRemovalRadius, SupportedAgents, Priority);
 	}
 }
 
@@ -49,7 +51,7 @@ void UNavigationInvokerComponent::RegisterWithNavigationSystem(UNavigationSystem
 		AActor* Owner = GetOwner();
 		if (Owner)
 		{
-			NavSys.RegisterInvoker(*Owner, TileGenerationRadius, TileRemovalRadius, SupportedAgents);
+			NavSys.RegisterInvoker(*Owner, TileGenerationRadius, TileRemovalRadius, SupportedAgents, Priority);
 		}
 	}
 }
diff --git a/Engine/Source/Runtime/NavigationSystem/Private/NavigationSystem.cpp b/Engine/Source/Runtime/NavigationSystem/Private/NavigationSystem.cpp
--- a/Engine/Source/Runtime/NavigationSystem/Private/NavigationSystem.cpp
+++ b/Engine/Source/Runtime/NavigationSystem/Private/NavigationSystem.cpp
@@ -10,6 +10,7 @@
 #include "GameFramework/Controller.h"
 #include "AI/Navigation/NavRelevantInterface.h"
 #include "AI/Navigation/NavigationDirtyElement.h"
+#include "AI/Navigation/NavigationInvokerPriority.h"
 #include "UObject/UObjectIterator.h"
 #include "EngineUtils.h"
 #include "Logging/MessageLog.h"
@@ -215,6 +216,15 @@ namespace NavigationDebugDrawing
 	const FVector PathNodeBoxExtent(16.f);
 }
 
+FNavigationInvokerRaw::FNavigationInvokerRaw(const FVector& InLocation, float Min, float Max, const FNavAgentSelector& InSupportedAgents, ENavigationInvokerPriority InPriority)
+: Location(InLocation)
+, RadiusMin(Min)
+, RadiusMax(Max)
+, SupportedAgents(InSupportedAgents)
+, Priority(InPriority)
+{
+}
+
 //----------------------------------------------------------------------//
 // FNavigationInvoker
 //----------------------------------------------------------------------//
@@ -222,15 +232,17 @@ FNavigationInvoker::FNavigationInvoker()
 : Actor(nullptr)
 , GenerationRadius(0)
 , RemovalRadius(0)
+, Priority(ENavigationInvokerPriority::Default)
 {
 	SupportedAgents.MarkInitialized();
 }
 
-FNavigationInvoker::FNavigationInvoker(AActor& InActor, float InGenerationRadius, float InRemovalRadius, const FNavAgentSelector& InSupportedAgents)
+FNavigationInvoker::FNavigationInvoker(AActor& InActor, float InGenerationRadius, float InRemovalRadius, const FNavAgentSelector& InSupportedAgents, ENavigationInvokerPriority InPriority)
 : Actor(&InActor)
 , GenerationRadius(InGenerationRadius)
 , RemovalRadius(InRemovalRadius)
 , SupportedAgents(InSupportedAgents)
+, Priority(InPriority)
 {
 	SupportedAgents.MarkInitialized();
 }
@@ -4803,12 +4815,12 @@ void UNavigationSystemV1::ResetMaxSimultaneousTileGenerationJobsCount()
 // Active tiles
 //----------------------------------------------------------------------//
 
-void UNavigationSystemV1::RegisterNavigationInvoker(AActor& Invoker, float TileGenerationRadius, float TileRemovalRadius, const FNavAgentSelector& Agents)
+void UNavigationSystemV1::RegisterNavigationInvoker(AActor& Invoker, float TileGenerationRadius, float TileRemovalRadius, const FNavAgentSelector& Agents, ENavigationInvokerPriority Priority)
 {
 	UNavigationSystemV1* NavSys = FNavigationSystem::GetCurrent<UNavigationSystemV1>(Invoker.GetWorld());
 	if (NavSys)
 	{
-		NavSys->RegisterInvoker(Invoker, TileGenerationRadius, TileRemovalRadius, Agents);
+		NavSys->RegisterInvoker(Invoker, TileGenerationRadius, TileRemovalRadius, Agents, Priority);
 	}
 }
 
@@ -4833,10 +4845,18 @@ void UNavigationSystemV1::SetGeometryGatheringMode(ENavDataGatheringModeConfig N
 // Deprecated
 void UNavigationSystemV1::RegisterInvoker(AActor& Invoker, float TileGenerationRadius, float TileRemovalRadius)
 {
+PRAGMA_DISABLE_DEPRECATION_WARNINGS
 	RegisterInvoker(Invoker, TileGenerationRadius, TileRemovalRadius, FNavAgentSelector());
+PRAGMA_ENABLE_DEPRECATION_WARNINGS	
 }
 
+// Deprecated
 void UNavigationSystemV1::RegisterInvoker(AActor& Invoker, float TileGenerationRadius, float TileRemovalRadius, const FNavAgentSelector& Agents)
+{
+	RegisterInvoker(Invoker, TileGenerationRadius, TileRemovalRadius, Agents, ENavigationInvokerPriority::Default);
+}
+
+void UNavigationSystemV1::RegisterInvoker(AActor& Invoker, float TileGenerationRadius, float TileRemovalRadius, const FNavAgentSelector& Agents, ENavigationInvokerPriority InPriority)
 {
 	UE_CVLOG(bGenerateNavigationOnlyAroundNavigationInvokers == false, this, LogNavInvokers, Warning
 		, TEXT("Trying to register %s as invoker, but NavigationSystem is not set up for invoker-centric generation. See GenerateNavigationOnlyAroundNavigationInvokers in NavigationSystem's properties")
@@ -4851,6 +4871,7 @@ void UNavigationSystemV1::RegisterInvoker(AActor& Invoker, float TileGenerationR
 	Data.RemovalRadius = TileRemovalRadius;
 	Data.SupportedAgents = Agents;
 	Data.SupportedAgents.MarkInitialized();
+	Data.Priority = InPriority;
 
 	UE_SUPPRESS(LogNavInvokers, Log,
 	{
@@ -4909,7 +4930,8 @@ void UNavigationSystemV1::UpdateInvokers()
 #endif //WITH_EDITOR
 					)
 				{
-					InvokerLocations.Add(FNavigationInvokerRaw(Actor->GetActorLocation(), ItemIterator->Value.GenerationRadius, ItemIterator->Value.RemovalRadius, ItemIterator->Value.SupportedAgents));
+					InvokerLocations.Add(FNavigationInvokerRaw(Actor->GetActorLocation(), ItemIterator->Value.GenerationRadius, ItemIterator->Value.RemovalRadius,
+						ItemIterator->Value.SupportedAgents, ItemIterator->Value.Priority));
 				}
 				else
 				{
@@ -4924,7 +4946,7 @@ void UNavigationSystemV1::UpdateInvokers()
 			for (const auto& InvokerData : InvokerLocations)
 			{
 				UE_VLOG_CYLINDER(this, LogNavInvokers, Log, InvokerData.Location, InvokerData.Location + FVector(0, 0, 20), InvokerData.RadiusMax, FColorList::Blue, TEXT(""));
-				UE_VLOG_CYLINDER(this, LogNavInvokers, Log, InvokerData.Location, InvokerData.Location + FVector(0, 0, 20), InvokerData.RadiusMin, FColorList::CadetBlue, TEXT(""));
+				UE_VLOG_CYLINDER(this, LogNavInvokers, Log, InvokerData.Location, InvokerData.Location + FVector(0, 0, 20), InvokerData.RadiusMin, FColorList::CadetBlue, TEXT("Priority %u"), InvokerData.Priority);
 			}
 #endif // ENABLE_VISUAL_LOG
 		}
@@ -4991,7 +5013,7 @@ void UNavigationSystemV1::RegisterNavigationInvoker(AActor* Invoker, float TileG
 	if (Invoker != nullptr)
 	{
 		// The FNavAgentSelector class is not yet exposed in BP so we use the default value to specify that we want to generate the navmesh for all agents
-		RegisterInvoker(*Invoker, TileGenerationRadius, TileRemovalRadius, FNavAgentSelector());
+		RegisterInvoker(*Invoker, TileGenerationRadius, TileRemovalRadius, FNavAgentSelector(), ENavigationInvokerPriority::Default);
 	}
 }
 
diff --git a/Engine/Source/Runtime/NavigationSystem/Public/NavMesh/RecastNavMesh.h b/Engine/Source/Runtime/NavigationSystem/Public/NavMesh/RecastNavMesh.h
--- a/Engine/Source/Runtime/NavigationSystem/Public/NavMesh/RecastNavMesh.h
+++ b/Engine/Source/Runtime/NavigationSystem/Public/NavMesh/RecastNavMesh.h
@@ -504,6 +504,15 @@ struct NAVIGATIONSYSTEM_API FRecastNavMeshTileGenerationDebug
 	uint32 bTileCacheDetailMesh : 1;
 };
 
+/**
+ *	Used to list tiles that needs rebuilding.
+ */
+struct FNavMeshDirtyTileElement
+{
+	FIntPoint Coordinates;
+	ENavigationInvokerPriority InvokerPriority;
+};
+
 /**
  *	Structure to handle nav mesh tile's raw data persistence and releasing
  */
@@ -1422,7 +1431,12 @@ PRAGMA_ENABLE_DEPRECATION_WARNINGS
 
 	virtual void UpdateActiveTiles(const TArray<FNavigationInvokerRaw>& InvokerLocations);
 	virtual void RemoveTiles(const TArray<FIntPoint>& Tiles);
+	
+	UE_DEPRECATED(5.3, "Use overload with FNavMeshDirtyTileElement instead.")
 	void RebuildTile(const TArray<FIntPoint>& Tiles);
+
+	void RebuildTile(const TArray<FNavMeshDirtyTileElement>& Tiles);
+	
 	void DirtyTilesInBounds(const FBox& Bounds);
 
 #if RECAST_INTERNAL_DEBUG_DATA
diff --git a/Engine/Source/Runtime/NavigationSystem/Public/NavMesh/RecastNavMeshGenerator.h b/Engine/Source/Runtime/NavigationSystem/Public/NavMesh/RecastNavMeshGenerator.h
--- a/Engine/Source/Runtime/NavigationSystem/Public/NavMesh/RecastNavMeshGenerator.h
+++ b/Engine/Source/Runtime/NavigationSystem/Public/NavMesh/RecastNavMeshGenerator.h
@@ -4,6 +4,7 @@
 #include "CoreMinimal.h"
 #include "Stats/Stats.h"
 #include "AI/Navigation/NavigationTypes.h"
+#include "AI/Navigation/NavigationInvokerPriority.h"
 #include "EngineDefines.h"
 #include "AI/NavigationModifier.h"
 #include "NavigationOctree.h"
@@ -560,6 +561,10 @@ struct FPendingTileElement
 	FVector::FReal SeedDistance;
 	/** time at which the element was first added to the queue */
 	double		CreationTime;
+
+	/** Priority from navigation invoker */
+	ENavigationInvokerPriority InvokerPriority = ENavigationInvokerPriority::Default;
+
 	/** Whether we need a full rebuild for this tile grid cell */
 	bool		bRebuildGeometry;
 	/** We need to store dirty area bounds to check which cached layers needs to be regenerated
@@ -864,8 +869,12 @@ public:
 	TArray<FNavTileRef> RemoveTileLayersAndGetUpdatedTiles(const int32 TileX, const int32 TileY, TMap<int32, dtPolyRef>* OldLayerTileIdMap = nullptr);
 
 	void RemoveTiles(const TArray<FIntPoint>& Tiles);
+
+	UE_DEPRECATED(5.3, "Use overload with FNavMeshDirtyTileElement instead.")
 	void ReAddTiles(const TArray<FIntPoint>& Tiles);
 
+	void ReAddTiles(const TArray<FNavMeshDirtyTileElement>& Tiles);
+
 	bool IsBuildingRestrictedToActiveTiles() const { return bRestrictBuildingToActiveTiles; }
 	bool IsInActiveSet(const FIntPoint& Tile) const;
 
@@ -874,6 +883,8 @@ public:
 	 *	@mote function does not validate the input parameter - it's on caller */
 	void SetMaxTileGeneratorTasks(int32 NewLimit) { MaxTileGeneratorTasks = NewLimit; }
 
+	void SetSortPendingTileMethod(const ENavigationSortPendingTilesMethod InMethod) { SortPendingTilesMethod = InMethod; }
+
 	static void CalcPolyRefBits(ARecastNavMesh* NavMeshOwner, int32& MaxTileBits, int32& MaxPolyBits);
 
 protected:
@@ -977,18 +988,23 @@ protected:
 
 	uint32 bRestrictBuildingToActiveTiles:1;
 
+	UE_DEPRECATED(5.3, "Use SortPendingTilesMethod instead.")
 	uint32 bSortTilesWithSeedLocations:1;
 
 	/** Runtime generator's version, increased every time all tile generators get invalidated
 	 *	like when navmesh size changes */
 	uint32 Version;
 
+	ENavigationSortPendingTilesMethod SortPendingTilesMethod = ENavigationSortPendingTilesMethod::SortWithSeedLocations;
+
 	/** Grouping all the member variables used by the time slicing code together for neatness */
 	struct FSyncTimeSlicedData
 	{
 		FSyncTimeSlicedData();
 
+		/** Accumulated time spent processing the tile (in seconds). */
 		double CurrentTileRegenDuration;
+		
 		/** if we are currently using time sliced regen or not - currently an experimental feature.
 		 *  do not manipulate this value directly instead call SetNextTimeSliceRegenActive()
 		 */
diff --git a/Engine/Source/Runtime/NavigationSystem/Public/NavigationInvokerComponent.h b/Engine/Source/Runtime/NavigationSystem/Public/NavigationInvokerComponent.h
--- a/Engine/Source/Runtime/NavigationSystem/Public/NavigationInvokerComponent.h
+++ b/Engine/Source/Runtime/NavigationSystem/Public/NavigationInvokerComponent.h
@@ -10,6 +10,7 @@
 #include "NavigationInvokerComponent.generated.h"
 
 class UNavigationSystemV1;
+enum class ENavigationInvokerPriority : uint8;
 
 UCLASS(ClassGroup = (Navigation), meta = (BlueprintSpawnableComponent))
 class NAVIGATIONSYSTEM_API UNavigationInvokerComponent : public UActorComponent
@@ -28,6 +29,10 @@ protected:
 	UPROPERTY(EditAnywhere, Category = Navigation)
 	FNavAgentSelector SupportedAgents;
 
+	/** Experimental invocation priority. It will modify the order in which invoked tiles are being built if SortPendingTilesMethod is set to SortByPriority. */
+	UPROPERTY(EditAnywhere, Category = Navigation)
+	ENavigationInvokerPriority Priority;
+
 public:
 	UNavigationInvokerComponent(const FObjectInitializer& ObjectInitializer = FObjectInitializer::Get());
 
diff --git a/Engine/Source/Runtime/NavigationSystem/Public/NavigationSystem.h b/Engine/Source/Runtime/NavigationSystem/Public/NavigationSystem.h
--- a/Engine/Source/Runtime/NavigationSystem/Public/NavigationSystem.h
+++ b/Engine/Source/Runtime/NavigationSystem/Public/NavigationSystem.h
@@ -11,6 +11,7 @@
 #include "NavFilters/NavigationQueryFilter.h"
 #include "AI/Navigation/NavigationTypes.h"
 #include "AI/Navigation/NavigationDirtyElement.h"
+#include "AI/Navigation/NavigationInvokerPriority.h"
 #include "NavigationSystemTypes.h"
 #include "NavigationData.h"
 #include "AI/NavigationSystemBase.h"
@@ -764,10 +765,16 @@ public:
 	//----------------------------------------------------------------------//
 	UE_DEPRECATED(5.2, "This function is deprecated. Please use the new RegisterInvoker method with the Agents parameter (FNavAgentSelector() can be used as default value to keep the same behavior)")
 	virtual void RegisterInvoker(AActor& Invoker, float TileGenerationRadius, float TileRemovalRadius);
+
+	UE_DEPRECATED(5.3, "This function is deprecated. Please use the new RegisterInvoker method with invoker priority.")
 	virtual void RegisterInvoker(AActor& Invoker, float TileGenerationRadius, float TileRemovalRadius, const FNavAgentSelector& Agents);
+
+	virtual void RegisterInvoker(AActor& Invoker, float TileGenerationRadius, float TileRemovalRadius, const FNavAgentSelector& Agents, ENavigationInvokerPriority InPriority);
+	
 	virtual void UnregisterInvoker(AActor& Invoker);
 
-	static void RegisterNavigationInvoker(AActor& Invoker, float TileGenerationRadius, float TileRemovalRadius, const FNavAgentSelector& Agents = FNavAgentSelector());
+	static void RegisterNavigationInvoker(AActor& Invoker, float TileGenerationRadius, float TileRemovalRadius, 
+		const FNavAgentSelector& Agents = FNavAgentSelector(), ENavigationInvokerPriority Priority = ENavigationInvokerPriority::Default);
 	static void UnregisterNavigationInvoker(AActor& Invoker);
 
 	const TArray<FNavigationInvokerRaw>& GetInvokerLocations() const { return InvokerLocations; }
diff --git a/Engine/Source/Runtime/NavigationSystem/Public/NavigationSystemTypes.h b/Engine/Source/Runtime/NavigationSystem/Public/NavigationSystemTypes.h
--- a/Engine/Source/Runtime/NavigationSystem/Public/NavigationSystemTypes.h
+++ b/Engine/Source/Runtime/NavigationSystem/Public/NavigationSystemTypes.h
@@ -13,6 +13,7 @@
 
 #define RECAST_INTERNAL_DEBUG_DATA (!UE_BUILD_SHIPPING)
 
+enum class ENavigationInvokerPriority : uint8;
 class UBodySetup;
 class UNavCollision;
 struct FKAggregateGeom;
@@ -97,10 +98,9 @@ struct FNavigationInvokerRaw
 	float RadiusMin;
 	float RadiusMax;
 	FNavAgentSelector SupportedAgents;
+	ENavigationInvokerPriority Priority;
 
-	FNavigationInvokerRaw(const FVector& InLocation, float Min, float Max, const FNavAgentSelector& InSupportedAgents)
-		: Location(InLocation), RadiusMin(Min), RadiusMax(Max), SupportedAgents(InSupportedAgents)
-	{}
+	FNavigationInvokerRaw(const FVector& InLocation, float Min, float Max, const FNavAgentSelector& InSupportedAgents, ENavigationInvokerPriority InPriority);
 };
 
 struct FNavigationInvoker
@@ -109,14 +109,19 @@ struct FNavigationInvoker
 
 	/** tiles GenerationRadius away or close will be generated if they're not already present */
 	float GenerationRadius;
+
 	/** tiles over RemovalRadius will get removed.
-	*	@Note needs to be >= GenerationRadius or will get clampped */
+	*	@Note needs to be >= GenerationRadius or will get clamped */
 	float RemovalRadius;
+
 	/** restrict navigation generation to specific agents */
 	FNavAgentSelector SupportedAgents;
 
+	/** invoker Priority used when dirtying tiles */
+	ENavigationInvokerPriority Priority;
+	
 	FNavigationInvoker();
-	FNavigationInvoker(AActor& InActor, float InGenerationRadius, float InRemovalRadius, const FNavAgentSelector& InSupportedAgents);
+	FNavigationInvoker(AActor& InActor, float InGenerationRadius, float InRemovalRadius, const FNavAgentSelector& InSupportedAgents, ENavigationInvokerPriority InPriority);
 };
 
 namespace NavigationHelper
