NavLinkId deprecation moved to 5.3 from 5.4

diff --git a/Engine/Source/Runtime/Engine/Classes/AI/Navigation/NavLinkDefinition.h b/Engine/Source/Runtime/Engine/Classes/AI/Navigation/NavLinkDefinition.h
--- a/Engine/Source/Runtime/Engine/Classes/AI/Navigation/NavLinkDefinition.h
+++ b/Engine/Source/Runtime/Engine/Classes/AI/Navigation/NavLinkDefinition.h
@@ -38,11 +38,11 @@ struct ENGINE_API FNavigationLinkBase
 	float MaxFallDownLength;
 
 	/** Needs to be 0 for recast data generator */
-	UE_DEPRECATED(5.4, "Use FNavLinkId::Invalid instead")
+	UE_DEPRECATED(5.3, "Use FNavLinkId::Invalid instead")
 	static constexpr uint32 InvalidUserId = 0;
 
 	/** ID passed to navigation data generator */
-	UE_DEPRECATED(5.4, "Use NavLinkId instead, this id is no longer used in the engine")
+	UE_DEPRECATED(5.3, "Use NavLinkId instead, this id is no longer used in the engine")
 	uint32 UserId;
 
 	FNavLinkId NavLinkId;
diff --git a/Engine/Source/Runtime/Engine/Classes/AI/Navigation/NavigationTypes.h b/Engine/Source/Runtime/Engine/Classes/AI/Navigation/NavigationTypes.h
--- a/Engine/Source/Runtime/Engine/Classes/AI/Navigation/NavigationTypes.h
+++ b/Engine/Source/Runtime/Engine/Classes/AI/Navigation/NavigationTypes.h
@@ -316,7 +316,7 @@ struct FNavPathPoint : public FNavLocation
 	/** extra node flags */
 	uint32 Flags;
 
-	UE_DEPRECATED(5.4, "Use CustomNavLinkId instead. This id is no longer used in the engine.")
+	UE_DEPRECATED(5.3, "Use CustomNavLinkId instead. This id is no longer used in the engine.")
 	uint32 CustomLinkId;
 
 	/** unique Id of custom navigation link starting at this point */
diff --git a/Engine/Source/Runtime/Engine/Public/AI/NavigationModifier.h b/Engine/Source/Runtime/Engine/Public/AI/NavigationModifier.h
--- a/Engine/Source/Runtime/Engine/Public/AI/NavigationModifier.h
+++ b/Engine/Source/Runtime/Engine/Public/AI/NavigationModifier.h
@@ -174,7 +174,7 @@ struct ENGINE_API FSimpleLinkNavModifier : public FNavigationModifier
 	TArray<FNavigationSegmentLink> SegmentLinks;
 	FTransform LocalToWorld;
 
-	UE_DEPRECATED(5.4, "Use NavLinkId instead, this id is no longer used in the engine.")
+	UE_DEPRECATED(5.3, "Use NavLinkId instead, this id is no longer used in the engine.")
 	int32 UserId;
 
 	FNavLinkId NavLinkId;
diff --git a/Engine/Source/Runtime/NavigationSystem/Public/NavLinkCustomComponent.h b/Engine/Source/Runtime/NavigationSystem/Public/NavLinkCustomComponent.h
--- a/Engine/Source/Runtime/NavigationSystem/Public/NavLinkCustomComponent.h
+++ b/Engine/Source/Runtime/NavigationSystem/Public/NavLinkCustomComponent.h
@@ -148,7 +148,7 @@ protected:
 
 protected:
 
-	UE_DEPRECATED(5.4, "LinkIds are now based on FNavLinkId. Use CustomLinkId instead. This Id is no longer used by the engine.")
+	UE_DEPRECATED(5.3, "LinkIds are now based on FNavLinkId. Use CustomLinkId instead. This Id is no longer used by the engine.")
 	UPROPERTY()
 	uint32 NavLinkUserId;
 
diff --git a/Engine/Source/Runtime/NavigationSystem/Public/NavLinkCustomInterface.h b/Engine/Source/Runtime/NavigationSystem/Public/NavLinkCustomInterface.h
--- a/Engine/Source/Runtime/NavigationSystem/Public/NavLinkCustomInterface.h
+++ b/Engine/Source/Runtime/NavigationSystem/Public/NavLinkCustomInterface.h
@@ -46,7 +46,7 @@ class NAVIGATIONSYSTEM_API INavLinkCustomInterface
 
 	virtual FNavLinkAuxiliaryId GetAuxiliaryId() const { return FNavLinkAuxiliaryId::Invalid; }
 
-	UE_DEPRECATED(5.4, "LinkIds are now based on a FNavLinkId. Call GetId() instead. This function only returns Invalid Id.")
+	UE_DEPRECATED(5.3, "LinkIds are now based on a FNavLinkId. Call GetId() instead. This function only returns Invalid Id.")
 	virtual uint32 GetLinkId() const final { return FNavLinkId::Invalid.GetId(); }
 
 	/** Get unique ID number for custom link
@@ -54,7 +54,7 @@ class NAVIGATIONSYSTEM_API INavLinkCustomInterface
 	 */
 	virtual FNavLinkId GetId() const { return FNavLinkId::Invalid; }
 
-	UE_DEPRECATED(5.4, "LinkIds are now based on a FNavLinkId. Call the version of this function that takes a FNavLinkId. This function now has no effect.")
+	UE_DEPRECATED(5.3, "LinkIds are now based on a FNavLinkId. Call the version of this function that takes a FNavLinkId. This function now has no effect.")
 	virtual void UpdateLinkId(uint32 NewUniqueId) final {}
 
 	/** Update unique ID number for custom link by navigation system. */
@@ -76,25 +76,25 @@ class NAVIGATIONSYSTEM_API INavLinkCustomInterface
 	/** Notify called when agent finishes using this link for movement */
 	virtual void OnLinkMoveFinished(class UObject* PathComp) {}
 
-	UE_DEPRECATED(5.4, "LinkIds are now based on FNavLinkId using FNavLinkId::GenerateUniqueId(). This function will still generate an incremental Id however it does not work well in all circumstances.")
+	UE_DEPRECATED(5.3, "LinkIds are now based on FNavLinkId using FNavLinkId::GenerateUniqueId(). This function will still generate an incremental Id however it does not work well in all circumstances.")
 	static uint32 GetUniqueId();
 
 	/** Helper function: bump unique ID numbers above given one */
-	UE_DEPRECATED(5.4, "LinkIds are now based on FNavLinkId. If your project is still using any of the old incremental Ids (saved in actors in levels or licensee code) then this function must be called still (typically by existing engine code), otherwise it is not necessary.")
+	UE_DEPRECATED(5.3, "LinkIds are now based on FNavLinkId. If your project is still using any of the old incremental Ids (saved in actors in levels or licensee code) then this function must be called still (typically by existing engine code), otherwise it is not necessary.")
 	static void UpdateUniqueId(FNavLinkId AlreadyUsedId);
 
-	UE_DEPRECATED(5.4, "LinkIds are now based on FNavLinkId. You may need to call the other version of this function that takes a FNavLinkId. This function has no effect.")
+	UE_DEPRECATED(5.3, "LinkIds are now based on FNavLinkId. You may need to call the other version of this function that takes a FNavLinkId. This function has no effect.")
 	static void UpdateUniqueId(uint32 AlreadyUsedId) {}
 
 	/** Helper function: create modifier for navigation data export */
 	static FNavigationLink GetModifier(const INavLinkCustomInterface* CustomNavLink);
 	
-	UE_DEPRECATED(5.4, "LinkIds are now based on a FNavLinkId Hash. If your project is still using any of the old incremental Ids then this function must be called still (typically by existing engine code), otherwise it is not necessary.")
+	UE_DEPRECATED(5.3, "LinkIds are now based on a FNavLinkId Hash. If your project is still using any of the old incremental Ids then this function must be called still (typically by existing engine code), otherwise it is not necessary.")
 	static void ResetUniqueId();
 
 	static void OnPreWorldInitialization(UWorld* World, const UWorld::InitializationValues IVS);
 
-	UE_DEPRECATED(5.4, "LinkIds are now based on FNavLinkId using FNavLinkId::GenerateUniqueId().")
+	UE_DEPRECATED(5.3, "LinkIds are now based on FNavLinkId using FNavLinkId::GenerateUniqueId().")
 	static uint32 NextUniqueId;
 };
 
diff --git a/Engine/Source/Runtime/NavigationSystem/Public/NavMesh/NavMeshPath.h b/Engine/Source/Runtime/NavigationSystem/Public/NavMesh/NavMeshPath.h
--- a/Engine/Source/Runtime/NavigationSystem/Public/NavMesh/NavMeshPath.h
+++ b/Engine/Source/Runtime/NavigationSystem/Public/NavMesh/NavMeshPath.h
@@ -148,7 +148,7 @@ public:
 	TArray<FVector::FReal> PathCorridorCost;
 
 	/** set of unique link Ids */
-	UE_DEPRECATED(5.4, "LinkIds are now based on FNavLinkId. Use CustomNavLinkIds instead. CustomLinkIds array is no longer populated or used in the engine")
+	UE_DEPRECATED(5.3, "LinkIds are now based on FNavLinkId. Use CustomNavLinkIds instead. CustomLinkIds array is no longer populated or used in the engine")
 	TArray<uint32> CustomLinkIds;
 
 	TArray<FNavLinkId> CustomNavLinkIds;
diff --git a/Engine/Source/Runtime/NavigationSystem/Public/NavMesh/PImplRecastNavMesh.h b/Engine/Source/Runtime/NavigationSystem/Public/NavMesh/PImplRecastNavMesh.h
--- a/Engine/Source/Runtime/NavigationSystem/Public/NavMesh/PImplRecastNavMesh.h
+++ b/Engine/Source/Runtime/NavigationSystem/Public/NavMesh/PImplRecastNavMesh.h
@@ -153,7 +153,7 @@ public:
 	bool FindStraightPath(const FVector& StartLoc, const FVector& EndLoc, const TArray<NavNodeRef>& PathCorridor, TArray<FNavPathPoint>& PathPoints, TArray<FNavLinkId>* CustomLinks = NULL) const;
 
 	/** finds stringpulled path from given corridor */
-	UE_DEPRECATED(5.4, "Please use FindStraightPath with the TArray<FNavPathPoint>* CustomLinks. This function has no effect.")
+	UE_DEPRECATED(5.3, "Please use FindStraightPath with the TArray<FNavPathPoint>* CustomLinks. This function has no effect.")
 	bool FindStraightPath(const FVector& StartLoc, const FVector& EndLoc, const TArray<NavNodeRef>& PathCorridor, TArray<FNavPathPoint>& PathPoints, TArray<uint32>* CustomLinks) const { return false; }
 
 	/** Filters nav polys in PolyRefs with Filter */
@@ -162,7 +162,7 @@ public:
 	/** Get all polys from tile */
 	bool GetPolysInTile(int32 TileIndex, TArray<FNavPoly>& Polys) const;
 
-	UE_DEPRECATED(5.4, "Please use the version of this function that takes a FNavLinkId. This function has no effect.")
+	UE_DEPRECATED(5.3, "Please use the version of this function that takes a FNavLinkId. This function has no effect.")
 	void UpdateNavigationLinkArea(int32 UserId, uint8 AreaType, uint16 PolyFlags) const {}
 
 	/** Updates area on polygons creating point-to-point connection with given UserId */
@@ -198,7 +198,7 @@ public:
 	/** Decode poly ID into FNavTileRef and poly index */
 	bool GetPolyTileRef(NavNodeRef PolyId, uint32& OutPolyIndex, FNavTileRef& OutTileRef) const;
 	/** Retrieves user ID for given offmesh link poly */
-	UE_DEPRECATED(5.4, "Please use GetNavLinkUserId() instead. This function only returns Invalid.")
+	UE_DEPRECATED(5.3, "Please use GetNavLinkUserId() instead. This function only returns Invalid.")
 	uint32 GetLinkUserId(NavNodeRef LinkPolyID) const
 	{
 		return FNavLinkId::Invalid.GetId();
diff --git a/Engine/Source/Runtime/NavigationSystem/Public/NavMesh/RecastNavMesh.h b/Engine/Source/Runtime/NavigationSystem/Public/NavMesh/RecastNavMesh.h
--- a/Engine/Source/Runtime/NavigationSystem/Public/NavMesh/RecastNavMesh.h
+++ b/Engine/Source/Runtime/NavigationSystem/Public/NavMesh/RecastNavMesh.h
@@ -1241,7 +1241,7 @@ public:
 
 	virtual void UpdateCustomLink(const INavLinkCustomInterface* CustomLink) override;
 
-	UE_DEPRECATED(5.4, "Use version of this function that takes a FNavLinkId. This function now has no effect.")
+	UE_DEPRECATED(5.3, "Use version of this function that takes a FNavLinkId. This function now has no effect.")
 	void UpdateNavigationLinkArea(int32 UserId, TSubclassOf<UNavArea> AreaClass) const {}
 
 	/** update area class and poly flags for all offmesh links with given UserId */
@@ -1272,7 +1272,7 @@ public:
 	/** dtNavMesh getter */
 	const dtNavMesh* GetRecastMesh() const;
 
-	UE_DEPRECATED(5.4, "Please use GetNavLinkUserId() instead. This function only returns Invalid.")
+	UE_DEPRECATED(5.3, "Please use GetNavLinkUserId() instead. This function only returns Invalid.")
 	int32 GetLinkUserId(NavNodeRef LinkPolyID) const
 	{
 		return FNavLinkId::Invalid.GetId();
@@ -1397,7 +1397,7 @@ public:
 	/** Check if poly is a custom link */
 	bool IsCustomLink(NavNodeRef PolyRef) const;
 
-	UE_DEPRECATED(5.4, "Use new override of this function with Array<FNavLinkId>* CustomLinks. This function has no effect.")
+	UE_DEPRECATED(5.3, "Use new override of this function with Array<FNavLinkId>* CustomLinks. This function has no effect.")
 	bool FindStraightPath(const FVector& StartLoc, const FVector& EndLoc, const TArray<NavNodeRef>& PathCorridor, TArray<FNavPathPoint>& PathPoints, TArray<uint32>* CustomLinks) const { return false; }
 
 	/** finds stringpulled path from given corridor */
diff --git a/Engine/Source/Runtime/NavigationSystem/Public/NavigationData.h b/Engine/Source/Runtime/NavigationSystem/Public/NavigationData.h
--- a/Engine/Source/Runtime/NavigationSystem/Public/NavigationData.h
+++ b/Engine/Source/Runtime/NavigationSystem/Public/NavigationData.h
@@ -239,7 +239,7 @@ struct NAVIGATIONSYSTEM_API FNavigationPath : public TSharedFromThis<FNavigation
 #endif // ENABLE_VISUAL_LOG
 
 	/** check if path contains specific custom nav link */
-	UE_DEPRECATED(5.4, "Use version that takes FNavLinkId instead. This function only returns false.")
+	UE_DEPRECATED(5.3, "Use version that takes FNavLinkId instead. This function only returns false.")
 	virtual bool ContainsCustomLink(uint32 UniqueLinkId) const final {	return false; }
 
 	virtual bool ContainsCustomLink(FNavLinkId UniqueLinkId) const;
diff --git a/Engine/Source/Runtime/NavigationSystem/Public/NavigationSystem.h b/Engine/Source/Runtime/NavigationSystem/Public/NavigationSystem.h
--- a/Engine/Source/Runtime/NavigationSystem/Public/NavigationSystem.h
+++ b/Engine/Source/Runtime/NavigationSystem/Public/NavigationSystem.h
@@ -916,7 +916,7 @@ public:
 	static void RequestCustomLinkRegistering(INavLinkCustomInterface& CustomLink, UObject* OwnerOb);
 	static void RequestCustomLinkUnregistering(INavLinkCustomInterface& CustomLink, UObject* ObjectOb);
 
-	UE_DEPRECATED(5.4, "LinkIds are now based on FNavLinkId call the version of this function that takes FNavLinkId. This function only returns nullptr.")
+	UE_DEPRECATED(5.3, "LinkIds are now based on FNavLinkId call the version of this function that takes FNavLinkId. This function only returns nullptr.")
 	INavLinkCustomInterface* GetCustomLink(uint32 UniqueLinkId) const { return nullptr; };
 
 	/** find custom link by unique ID */
@@ -1141,7 +1141,7 @@ protected:
 	FNavigationOctreeController DefaultOctreeController;
 
 	/** Map of all custom navigation links, that are relevant for path following */
-	UE_DEPRECATED(5.4, "LinkIds are now based on FNavLinkId. Use CustomNavLinksMap instead, CustomLinksMap is no longer populated or used in the engine.")
+	UE_DEPRECATED(5.3, "LinkIds are now based on FNavLinkId. Use CustomNavLinksMap instead, CustomLinksMap is no longer populated or used in the engine.")
 	TMap<uint32, FNavigationSystem::FCustomLinkOwnerInfo> CustomLinksMap;
 
 	TMap<FNavLinkId, FNavigationSystem::FCustomLinkOwnerInfo> CustomNavLinksMap;
diff --git a/Engine/Source/Runtime/Navmesh/Public/Detour/DetourNavMeshQuery.h b/Engine/Source/Runtime/Navmesh/Public/Detour/DetourNavMeshQuery.h
--- a/Engine/Source/Runtime/Navmesh/Public/Detour/DetourNavMeshQuery.h
+++ b/Engine/Source/Runtime/Navmesh/Public/Detour/DetourNavMeshQuery.h
@@ -55,7 +55,7 @@ struct NAVMESH_API dtQuerySpecialLinkFilter
 	
 	virtual bool isLinkAllowed(const unsigned long long int UserId) const { return true; }
 
-	UE_DEPRECATED(5.4, "LinkIds are now based on a 64 bit uint. Use the version of this function that takes an unsigned long long int")
+	UE_DEPRECATED(5.3, "LinkIds are now based on a 64 bit uint. Use the version of this function that takes an unsigned long long int")
 	virtual bool isLinkAllowed(const int UserId) const final { return true; }
 
 	/// Called before accessing in A* loop (can be called multiple time for updateSlicedFindPath)
