Fix for static analysis warning.

diff --git a/Engine/Source/Runtime/NavigationSystem/Private/NavigationSystem.cpp b/Engine/Source/Runtime/NavigationSystem/Private/NavigationSystem.cpp
--- a/Engine/Source/Runtime/NavigationSystem/Private/NavigationSystem.cpp
+++ b/Engine/Source/Runtime/NavigationSystem/Private/NavigationSystem.cpp
@@ -2624,7 +2624,7 @@ void UNavigationSystemV1::RegisterCustomLink(INavLinkCustomInterface& CustomLink
 		if (NewId != OldId)
 		{
 			CustomLink.UpdateLinkId(NewId);
-			UE_LOG(LogNavLink, VeryVerbose, TEXT("%hs new navlink ID %llu."), __FUNCTION__, CustomLink.GetId());
+			UE_LOG(LogNavLink, VeryVerbose, TEXT("%hs new navlink ID %llu."), __FUNCTION__, CustomLink.GetId().GetId());
 
 			const FBox LinkBounds = ComputeCustomLinkBounds(CustomLink);
 			if (LinkBounds.IsValid)
@@ -2634,7 +2634,7 @@ void UNavigationSystemV1::RegisterCustomLink(INavLinkCustomInterface& CustomLink
 		}
 	}
 
-	UE_CLOG(bGenerateNewId && CustomNavLinksMap.Contains(CustomLink.GetId()), LogNavLink, Warning, TEXT("%hs New navlink ID %llu is clashing with existing ID."), __FUNCTION__, CustomLink.GetId());
+	UE_CLOG(bGenerateNewId && CustomNavLinksMap.Contains(CustomLink.GetId()), LogNavLink, Warning, TEXT("%hs New navlink ID %llu is clashing with existing ID."), __FUNCTION__, CustomLink.GetId().GetId());
 	CustomNavLinksMap.Add(CustomLink.GetId(), FNavigationSystem::FCustomLinkOwnerInfo(&CustomLink));
 }
 
