Bump the priority of dirty tiles near invokers when using sorting by invokers priority.

diff --git a/Engine/Source/Runtime/Engine/Classes/AI/Navigation/NavigationInvokerPriority.h b/Engine/Source/Runtime/Engine/Classes/AI/Navigation/NavigationInvokerPriority.h
--- a/Engine/Source/Runtime/Engine/Classes/AI/Navigation/NavigationInvokerPriority.h
+++ b/Engine/Source/Runtime/Engine/Classes/AI/Navigation/NavigationInvokerPriority.h
@@ -7,6 +7,7 @@
 UENUM()
 enum class ENavigationInvokerPriority : uint8
 {
+	// Enum values are used (keep in ascending order)
 	VeryLow = 1,
 	Low = 2,
 	Default = 3,
diff --git a/Engine/Source/Runtime/NavigationSystem/Private/NavMesh/RecastNavMesh.cpp b/Engine/Source/Runtime/NavigationSystem/Private/NavMesh/RecastNavMesh.cpp
--- a/Engine/Source/Runtime/NavigationSystem/Private/NavMesh/RecastNavMesh.cpp
+++ b/Engine/Source/Runtime/NavigationSystem/Private/NavMesh/RecastNavMesh.cpp
@@ -3531,7 +3531,7 @@ void ARecastNavMesh::UpdateActiveTiles(const TArray<FNavigationInvokerRaw>& Invo
 						}
 						else
 						{
-							TilesInMinDistance.Add(FNavMeshDirtyTileElement{FIntPoint(X,Y), Invoker.Priority});
+							TilesInMinDistance.Add(FNavMeshDirtyTileElement{FIntPoint(X,Y), DistanceSq, Invoker.Priority});
 							TileToAppend.Add(FIntPoint(X,Y));
 						}
 					}
@@ -3607,7 +3607,7 @@ void ARecastNavMesh::RebuildTile(const TArray<FIntPoint>& Tiles)
 	ActiveTiles.Reserve(Tiles.Num());
 	for (const FIntPoint& Point : Tiles)
 	{
-		ActiveTiles.Add(FNavMeshDirtyTileElement{Point, ENavigationInvokerPriority::Default});
+		ActiveTiles.Add(FNavMeshDirtyTileElement{Point, TNumericLimits<FVector::FReal>::Max(), ENavigationInvokerPriority::Default});
 	}
 	RebuildTile(ActiveTiles);
 }
@@ -3672,7 +3672,7 @@ void ARecastNavMesh::DirtyTilesInBounds(const FBox& Bounds)
 			for (int32 TileX = TileBox.XMin; TileX <= TileBox.XMax; ++TileX)
 			{
 				// For now, new dirtiness is made with default priority.
-				Tiles.Add(FNavMeshDirtyTileElement{FIntPoint(TileX, TileY), ENavigationInvokerPriority::Default});
+				Tiles.Add(FNavMeshDirtyTileElement{FIntPoint(TileX, TileY), TNumericLimits<FVector::FReal>::Max(), ENavigationInvokerPriority::Default});
 			}
 		}
 		RebuildTile(Tiles);
diff --git a/Engine/Source/Runtime/NavigationSystem/Private/NavMesh/RecastNavMeshGenerator.cpp b/Engine/Source/Runtime/NavigationSystem/Private/NavMesh/RecastNavMeshGenerator.cpp
--- a/Engine/Source/Runtime/NavigationSystem/Private/NavMesh/RecastNavMeshGenerator.cpp
+++ b/Engine/Source/Runtime/NavigationSystem/Private/NavMesh/RecastNavMeshGenerator.cpp
@@ -5293,7 +5293,7 @@ void FRecastNavMeshGenerator::ReAddTiles(const TArray<FIntPoint>& Tiles)
 	TArray<FNavMeshDirtyTileElement> ConvertedTiles;
 	for (const FIntPoint& Point : Tiles)
 	{
-		ConvertedTiles.Add(FNavMeshDirtyTileElement{Point, ENavigationInvokerPriority::Default});
+		ConvertedTiles.Add(FNavMeshDirtyTileElement{Point, TNumericLimits<FVector::FReal>::Max(), ENavigationInvokerPriority::Default});
 	}
 	ReAddTiles(ConvertedTiles);
 }
@@ -5308,13 +5308,29 @@ void FRecastNavMeshGenerator::ReAddTiles(const TArray<FNavMeshDirtyTileElement>&
 	TSet<FPendingTileElement> DirtyTiles;
 
 	const double CurrentTimeSeconds = FPlatformTime::Seconds();
+	const FVector::FReal NearDistanceSquared = FMath::Square(DestNavMesh->InvokerTilePriorityBumpDistanceThresholdInTileUnits*GetConfig().GetTileSizeUU());
+	const uint8 PriorityIncrease = DestNavMesh->InvokerTilePriorityBumpIncrease;
 
 	// @note we act on assumption all items in Tiles are unique
 	for (const FNavMeshDirtyTileElement& Tile : Tiles)
 	{
 		FPendingTileElement Element;
 		Element.Coord = Tile.Coordinates;
-		Element.InvokerPriority = Tile.InvokerPriority;
+#if !UE_BUILD_SHIPPING		
+		Element.DebugInvokerDistanceSquared = Tile.InvokerDistanceSquared;
+		Element.DebugInvokerPriority = Tile.InvokerPriority;
+#endif // !UE_BUILD_SHIPPING			
+
+		// Bump sorting priority for tiles near invokers
+		if (Tile.InvokerDistanceSquared < NearDistanceSquared)
+		{
+			Element.SortingPriority = (ENavigationInvokerPriority)FMath::Min((int)Tile.InvokerPriority+PriorityIncrease, (int)ENavigationInvokerPriority::MAX-1);	
+		}
+		else
+		{
+			Element.SortingPriority = Tile.InvokerPriority;
+		}
+
 		Element.bRebuildGeometry = true;
 		Element.CreationTime = CurrentTimeSeconds;
 		DirtyTiles.Add(Element);
@@ -5328,7 +5344,11 @@ void FRecastNavMeshGenerator::ReAddTiles(const TArray<FNavMeshDirtyTileElement>&
 		FPendingTileElement* ExistingElement = DirtyTiles.Find(Element);
 		if (ExistingElement)
 		{
-			ExistingElement->InvokerPriority = FMath::Max(ExistingElement->InvokerPriority, Element.InvokerPriority);
+#if !UE_BUILD_SHIPPING			
+			ExistingElement->DebugInvokerDistanceSquared = FMath::Min(ExistingElement->DebugInvokerDistanceSquared, Element.DebugInvokerDistanceSquared);
+			ExistingElement->DebugInvokerPriority = FMath::Max(ExistingElement->DebugInvokerPriority, Element.DebugInvokerPriority);
+#endif // !UE_BUILD_SHIPPING			
+			ExistingElement->SortingPriority = FMath::Max(ExistingElement->SortingPriority, Element.SortingPriority);
 			ExistingElement->bRebuildGeometry |= Element.bRebuildGeometry;
 			ExistingElement->CreationTime = FMath::Min(Element.CreationTime, ExistingElement->CreationTime);
 			// Append area bounds to existing list 
@@ -6160,7 +6180,7 @@ void FRecastNavMeshGenerator::SortPendingBuildTiles()
 	else if (SortPendingTilesMethod == ENavigationSortPendingTilesMethod::SortByPriority)
 	{
 		// Highest priority should be at the end of the list
-		PendingDirtyTiles.Sort([](const FPendingTileElement& A, const FPendingTileElement& B){ return A.InvokerPriority < B.InvokerPriority; });
+		PendingDirtyTiles.Sort([](const FPendingTileElement& A, const FPendingTileElement& B){ return A.SortingPriority < B.SortingPriority; });
 
 		UE_SUPPRESS(LogNavigation, VeryVerbose,
 		{
@@ -6170,7 +6190,7 @@ void FRecastNavMeshGenerator::SortPendingBuildTiles()
 				const FVector::FReal TileSizeInWorldUnits = Config.GetTileSizeUU();
 				const FBox TileBox = FRecastTileGenerator::CalculateTileBounds(Element.Coord.X, Element.Coord.Y, FVector::ZeroVector, TotalNavBounds, TileSizeInWorldUnits);
 				const ARecastNavMesh* const OwnerNav = GetOwner();
-				UE_VLOG_BOX(OwnerNav, LogNavigation, VeryVerbose, TileBox, FColor::Silver, TEXT("Index: %i, Priority %i"), Index, Element.InvokerPriority);
+				UE_VLOG_BOX(OwnerNav, LogNavigation, VeryVerbose, TileBox, FColor::Silver, TEXT("Index: %i, Priority %i"), Index, Element.SortingPriority);
 			}
 		});
 	}
diff --git a/Engine/Source/Runtime/NavigationSystem/Private/NavigationSystem.cpp b/Engine/Source/Runtime/NavigationSystem/Private/NavigationSystem.cpp
--- a/Engine/Source/Runtime/NavigationSystem/Private/NavigationSystem.cpp
+++ b/Engine/Source/Runtime/NavigationSystem/Private/NavigationSystem.cpp
@@ -4889,8 +4889,8 @@ void UNavigationSystemV1::RegisterInvoker(AActor& Invoker, float TileGenerationR
 			}
 		}
 
-		const FString RegisterText = FString::Printf(TEXT("Register invoker r: %.0f, r area: %.0f m2, removal r: %.0f (%s %s) "),
-			TileGenerationRadius, UE_PI*FMath::Square(TileGenerationRadius/100.f), TileRemovalRadius, *Invoker.GetName(), *InvokerNavData);
+		const FString RegisterText = FString::Printf(TEXT("Register invoker r: %.0f, r area: %.0f m2, removal r: %.0f, priority: %d, (%s %s) "),
+			TileGenerationRadius, UE_PI*FMath::Square(TileGenerationRadius/100.f), TileRemovalRadius, InPriority, *Invoker.GetName(), *InvokerNavData);
 		UE_LOG(LogNavInvokers, Log, TEXT("%s"), *RegisterText);
 
 		UE_VLOG_CYLINDER(this, LogNavInvokers, Log, Invoker.GetActorLocation(), Invoker.GetActorLocation() + FVector(0, 0, 20), TileGenerationRadius, FColorList::LimeGreen, TEXT("%s"), *RegisterText);
diff --git a/Engine/Source/Runtime/NavigationSystem/Public/NavMesh/RecastNavMesh.h b/Engine/Source/Runtime/NavigationSystem/Public/NavMesh/RecastNavMesh.h
--- a/Engine/Source/Runtime/NavigationSystem/Public/NavMesh/RecastNavMesh.h
+++ b/Engine/Source/Runtime/NavigationSystem/Public/NavMesh/RecastNavMesh.h
@@ -510,6 +510,7 @@ struct NAVIGATIONSYSTEM_API FRecastNavMeshTileGenerationDebug
 struct FNavMeshDirtyTileElement
 {
 	FIntPoint Coordinates;
+	FVector::FReal InvokerDistanceSquared;
 	ENavigationInvokerPriority InvokerPriority;
 };
 
@@ -908,6 +909,14 @@ class NAVIGATIONSYSTEM_API ARecastNavMesh : public ANavigationData
 	UPROPERTY(EditAnywhere, Category = TimeSlicing, config, AdvancedDisplay)
 	double TimeSliceLongDurationDebug = 0.002;
 
+	/** If >= 1, when sorting pending tiles by priority, tiles near invokers (within the distance threshold) will have their priority increased. */
+	UPROPERTY(EditAnywhere, Category = Generation, config, AdvancedDisplay)
+	uint32 InvokerTilePriorityBumpDistanceThresholdInTileUnits = 1;
+
+	/** Priority increase steps for tiles that are withing near distance. */
+	UPROPERTY(EditAnywhere, Category = Generation, config, AdvancedDisplay)
+	uint8 InvokerTilePriorityBumpIncrease = 1;
+
 protected:
 #if WITH_EDITORONLY_DATA
 	/** World partitioned navmesh are only allowed in partitioned worlds. */
diff --git a/Engine/Source/Runtime/NavigationSystem/Public/NavMesh/RecastNavMeshGenerator.h b/Engine/Source/Runtime/NavigationSystem/Public/NavMesh/RecastNavMeshGenerator.h
--- a/Engine/Source/Runtime/NavigationSystem/Public/NavMesh/RecastNavMeshGenerator.h
+++ b/Engine/Source/Runtime/NavigationSystem/Public/NavMesh/RecastNavMeshGenerator.h
@@ -562,8 +562,16 @@ struct FPendingTileElement
 	/** time at which the element was first added to the queue */
 	double		CreationTime;
 
+#if !UE_BUILD_SHIPPING	
+	/** Squared distance to invoker source */
+	FVector::FReal DebugInvokerDistanceSquared = TNumericLimits<FVector::FReal>::Max();
+
 	/** Priority from navigation invoker */
-	ENavigationInvokerPriority InvokerPriority = ENavigationInvokerPriority::Default;
+	ENavigationInvokerPriority DebugInvokerPriority = ENavigationInvokerPriority::Default;
+#endif // !UE_BUILD_SHIPPING
+
+	/** Priority used for sorting */
+	ENavigationInvokerPriority SortingPriority = ENavigationInvokerPriority::Default;
 
 	/** Whether we need a full rebuild for this tile grid cell */
 	bool		bRebuildGeometry;
