Fix memory leak in FNavigationOctreeElement operator= and added move operators.

diff --git a/Engine/Source/Runtime/NavigationSystem/Public/NavigationOctree.h b/Engine/Source/Runtime/NavigationSystem/Public/NavigationOctree.h
--- a/Engine/Source/Runtime/NavigationSystem/Public/NavigationOctree.h
+++ b/Engine/Source/Runtime/NavigationSystem/Public/NavigationOctree.h
@@ -37,9 +37,30 @@ public:
 		, OwnerUniqueId(Other.OwnerUniqueId)
 	{}
 
+	FNavigationOctreeElement(FNavigationOctreeElement&& Other) noexcept
+	: Bounds(MoveTemp(Other.Bounds))
+	, Data(MoveTemp(Other.Data))
+	, OwnerUniqueId(MoveTemp(Other.OwnerUniqueId))
+	{
+	}
+
+	FNavigationOctreeElement& operator=(FNavigationOctreeElement&& Other)
+	{
+		if (this != &Other)
+		{
+			this->~FNavigationOctreeElement();
+			new(this) FNavigationOctreeElement(Forward<FNavigationOctreeElement>(Other));			
+		}
+		return *this;
+	}
+	
 	FNavigationOctreeElement& operator=(const FNavigationOctreeElement& Other)
 	{
-		new(this) FNavigationOctreeElement(Other);
+		if (this != &Other)
+		{
+			this->~FNavigationOctreeElement();
+			new(this) FNavigationOctreeElement(Other);
+		}
 		return *this;
 	}
 
