[Navmesh] Addition of tile generation history and stats when LogNavigationHistory is set to Log.

diff --git a/Engine/Source/Runtime/Engine/Classes/AI/NavigationSystemBase.h b/Engine/Source/Runtime/Engine/Classes/AI/NavigationSystemBase.h
--- a/Engine/Source/Runtime/Engine/Classes/AI/NavigationSystemBase.h
+++ b/Engine/Source/Runtime/Engine/Classes/AI/NavigationSystemBase.h
@@ -22,6 +22,7 @@ class UNavAreaBase;
 
 ENGINE_API DECLARE_LOG_CATEGORY_EXTERN(LogNavigation, Warning, All);
 ENGINE_API DECLARE_LOG_CATEGORY_EXTERN(LogNavigationDataBuild, Log, All);
+ENGINE_API DECLARE_LOG_CATEGORY_EXTERN(LogNavigationHistory, Warning, All);
 ENGINE_API DECLARE_LOG_CATEGORY_EXTERN(LogNavInvokers, Warning, All);
 ENGINE_API DECLARE_LOG_CATEGORY_EXTERN(LogNavLink, Warning, All);
 ENGINE_API DECLARE_LOG_CATEGORY_EXTERN(LogAStar, Display, All);
diff --git a/Engine/Source/Runtime/Engine/Private/AI/NavigationSystemBase.cpp b/Engine/Source/Runtime/Engine/Private/AI/NavigationSystemBase.cpp
--- a/Engine/Source/Runtime/Engine/Private/AI/NavigationSystemBase.cpp
+++ b/Engine/Source/Runtime/Engine/Private/AI/NavigationSystemBase.cpp
@@ -13,6 +13,7 @@
 
 DEFINE_LOG_CATEGORY(LogNavigation);
 DEFINE_LOG_CATEGORY(LogNavigationDataBuild);
+DEFINE_LOG_CATEGORY(LogNavigationHistory);
 DEFINE_LOG_CATEGORY(LogNavInvokers);
 DEFINE_LOG_CATEGORY(LogNavLink);
 DEFINE_LOG_CATEGORY(LogAStar);
diff --git a/Engine/Source/Runtime/NavigationSystem/Private/NavMesh/RecastNavMeshGenerator.cpp b/Engine/Source/Runtime/NavigationSystem/Private/NavMesh/RecastNavMeshGenerator.cpp
--- a/Engine/Source/Runtime/NavigationSystem/Private/NavMesh/RecastNavMeshGenerator.cpp
+++ b/Engine/Source/Runtime/NavigationSystem/Private/NavMesh/RecastNavMeshGenerator.cpp
@@ -6503,6 +6503,10 @@ TArray<FNavTileRef> FRecastNavMeshGenerator::ProcessTileTasksSyncTimeSlicedAndGe
 
 				SyncTimeSlicedData.CurrentTileRegenDuration = 0.;
 
+#if !UE_BUILD_SHIPPING				
+				SyncTimeSlicedData.TileRegenStartFrame = GFrameCounter;
+#endif // !UE_BUILD_SHIPPING 
+
 				if (SyncTimeSlicedData.TileGeneratorSync->HasDataToBuild())
 				{
 					SyncTimeSlicedData.ProcessTileTasksSyncState = EProcessTileTasksSyncTimeSlicedState::DoWork;
@@ -6592,12 +6596,17 @@ TArray<FNavTileRef> FRecastNavMeshGenerator::ProcessTileTasksSyncTimeSlicedAndGe
 				SyncTimeSlicedData.ProcessTileTasksSyncState = EProcessTileTasksSyncTimeSlicedState::Init;
 
 				SyncTimeSlicedData.CurrentTileRegenDuration += (FPlatformTime::Seconds() - TimeStartProcessingTileThisFrame);
+#if !UE_BUILD_SHIPPING					
+				const double TempRegenDuration = SyncTimeSlicedData.CurrentTileRegenDuration;
+#endif					
 
 				SyncTimeSlicedData.TimeSliceManager->PushTileRegenTime(SyncTimeSlicedData.CurrentTileRegenDuration);
 
 				SyncTimeSlicedData.CurrentTileRegenDuration = 0.;
 
 #if !UE_BUILD_SHIPPING					
+				SyncTimeSlicedData.TileRegenEndFrame = GFrameCounter;
+					
 				const UNavigationSystemV1* NavSys = FNavigationSystem::GetCurrent<UNavigationSystemV1>(GetWorld());
 				if (NavSys)
 				{
@@ -6605,6 +6614,18 @@ TArray<FNavTileRef> FRecastNavMeshGenerator::ProcessTileTasksSyncTimeSlicedAndGe
 					const double WaitTime = (FPlatformTime::Seconds() - TileGenerator->TileCreationTime);
 					const int32 NavDataIndex = NavSys->NavDataSet.Find(DestNavMesh);	
 					SyncTimeSlicedData.TimeSliceManager->PushTileWaitTime(NavDataIndex, WaitTime);
+
+					UE_SUPPRESS(LogNavigationHistory, Log,
+					{
+						FTileHistoryData HistoryData;
+						HistoryData.TileX = TileGenerator->TileX;
+						HistoryData.TileY = TileGenerator->TileY;
+						HistoryData.TileRegenTime = (float)TempRegenDuration;
+						HistoryData.TileWaitTime = (float)WaitTime;
+						HistoryData.StartRegenFrame = SyncTimeSlicedData.TileRegenStartFrame;
+						HistoryData.EndRegenFrame = SyncTimeSlicedData.TileRegenEndFrame;
+						SyncTimeSlicedData.TimeSliceManager->PushTileHistoryData(NavDataIndex, HistoryData);
+					});
 				}
 #endif // !UE_BUILD_SHIPPING					
 				
diff --git a/Engine/Source/Runtime/NavigationSystem/Private/NavigationSystem.cpp b/Engine/Source/Runtime/NavigationSystem/Private/NavigationSystem.cpp
--- a/Engine/Source/Runtime/NavigationSystem/Private/NavigationSystem.cpp
+++ b/Engine/Source/Runtime/NavigationSystem/Private/NavigationSystem.cpp
@@ -325,6 +325,15 @@ bool FNavRegenTimeSlicer::TestTimeSliceFinished() const
 	return bTimeSliceFinishedCached;
 }
 
+void FNavRegenTimeSliceManager::ResetTileWaitTimeArrays(const TArray<TObjectPtr<ANavigationData>>& NavDataSet)
+{
+	TileWaitTimes.SetNum(NavDataSet.Num());
+	for (TArray<double>& Array : TileWaitTimes)
+	{
+		Array.Empty();
+	}
+}
+
 void FNavRegenTimeSliceManager::PushTileWaitTime(const int32 NavDataIndex, const double NewTime)
 {
 	if (TileWaitTimes.IsValidIndex(NavDataIndex))
@@ -333,6 +342,25 @@ void FNavRegenTimeSliceManager::PushTileWaitTime(const int32 NavDataIndex, const
 	}
 }
 
+#if !UE_BUILD_SHIPPING
+void FNavRegenTimeSliceManager::ResetTileHistoryData(const TArray<TObjectPtr<ANavigationData>>& NavDataSet)
+{
+	TileHistoryData.SetNum(NavDataSet.Num());
+	for (TArray<FTileHistoryData>& HistoryData : TileHistoryData)
+	{
+		HistoryData.Empty();
+	}
+}
+
+void FNavRegenTimeSliceManager::PushTileHistoryData(const int32 NavDataIndex, const FTileHistoryData& TileData)
+{
+	if (TileHistoryData.IsValidIndex(NavDataIndex))
+	{
+		TileHistoryData[NavDataIndex].Add(TileData);	
+	}
+}
+#endif // UE_BUILD_SHIPPING
+
 double FNavRegenTimeSliceManager::GetAverageTileWaitTime(const int32 NavDataIndex) const
 {
 	if (!TileWaitTimes.IsValidIndex(NavDataIndex))
@@ -461,6 +489,41 @@ void FNavRegenTimeSliceManager::SetMaxDesiredTileRegenDuration(float NewMaxDesir
 	UE_LOG(LogNavigationDataBuild, Verbose, TEXT("Navigation System: MaxDesiredTileRegenDuration = %f"), MaxDesiredTileRegenDuration);
 }
 
+#if !UE_BUILD_SHIPPING
+void FNavRegenTimeSliceManager::LogTileStatistics(const TArray<TObjectPtr<ANavigationData>>& NavDataSet) const
+{
+	UE_SUPPRESS(LogNavigationHistory, Log,
+	{
+		// Log median tile processing time every 60 frames.
+		const bool bLog = GFrameCounter % 60 == 0;
+		for (int32 NavDataIndex = 0; bLog && NavDataIndex < NavDataSet.Num(); ++NavDataIndex)
+		{
+			if (TileHistoryData.IsValidIndex(NavDataIndex))
+			{
+				TArray<FTileHistoryData> HistoryData = TileHistoryData[NavDataIndex];
+				if (HistoryData.Num() > 0)
+				{
+					const int32 MedianIndex = HistoryData.Num()/2;
+					const int32 HighIndex = int(HistoryData.Num()*0.9);
+
+					HistoryData.Sort([](const FTileHistoryData& A, const FTileHistoryData& B){ return A.TileRegenTime < B.TileRegenTime; });
+					const double MedianRegenTimeMs = HistoryData[MedianIndex].TileRegenTime * 1000.f;
+					const double HighRegenTimeMs = HistoryData[HighIndex].TileRegenTime * 1000.f;
+					const int64 MedianRegenFrames = HistoryData[MedianIndex].EndRegenFrame - HistoryData[MedianIndex].StartRegenFrame;
+
+					HistoryData.Sort([](const FTileHistoryData& A, const FTileHistoryData& B){ return A.TileWaitTime < B.TileWaitTime; });
+					const double MedianWaitTimeMs = HistoryData[MedianIndex].TileWaitTime * 1000.f;
+					const double HighWaitTimeMs = HistoryData[HighIndex].TileWaitTime * 1000.f;
+					
+					UE_LOG(LogNavigation, Warning, TEXT("%-30s Median tile stats: regen time: %2.2f ms, regen frames %lld, wait time: %4.f ms (high regen time: %2.2f ms, high wait time: %4.f ms."),
+						*GetNameSafe(NavDataSet[NavDataIndex]), MedianRegenTimeMs, MedianRegenFrames, MedianWaitTimeMs, HighRegenTimeMs, HighWaitTimeMs);
+				}
+			}
+		}
+	});
+}
+#endif // !UE_BUILD_SHIPPING
+
 //----------------------------------------------------------------------//
 // UNavigationSystemV1                                                                
 //----------------------------------------------------------------------//
@@ -1302,6 +1365,10 @@ void UNavigationSystemV1::Tick(float DeltaSeconds)
 				{
 					NumTimeSlicedRemainingBuildTasks += NumTasks;
 				}
+
+#if !UE_BUILD_SHIPPING
+				NavRegenTimeSliceManager.LogTileStatistics(NavDataSet);
+#endif // !UE_BUILD_SHIPPING
 				
 				if (NumTimeSlicedRemainingBuildTasks > 0)
 				{
@@ -2539,7 +2606,11 @@ UNavigationSystemV1::ERegistrationResult UNavigationSystemV1::RegisterNavData(AN
 		Result = RegistrationFailed_AgentNotValid;
 	}
 
-	NavRegenTimeSliceManager.InitTileWaitTimeArrays(NavDataSet);
+	NavRegenTimeSliceManager.ResetTileWaitTimeArrays(NavDataSet);
+
+#if !UE_BUILD_SHIPPING
+	NavRegenTimeSliceManager.ResetTileHistoryData(NavDataSet);
+#endif // UE_BUILD_SHIPPING
 
 	// @todo else might consider modifying this NavData to implement navigation for one of the supported agents
 	// care needs to be taken to not make it implement navigation for agent who's real implementation has 
@@ -2570,6 +2641,12 @@ void UNavigationSystemV1::UnregisterNavData(ANavigationData* NavData)
 	NavDataRegistrationQueue.Remove(NavData);
 	NavData->OnUnregistered();
 
+	NavRegenTimeSliceManager.ResetTileWaitTimeArrays(NavDataSet);
+
+#if !UE_BUILD_SHIPPING
+	NavRegenTimeSliceManager.ResetTileHistoryData(NavDataSet);
+#endif // UE_BUILD_SHIPPING
+	
 	if (CrowdManager != nullptr)
 	{
 		CrowdManager->OnNavDataUnregistered(*NavData);
diff --git a/Engine/Source/Runtime/NavigationSystem/Public/NavMesh/RecastNavMeshGenerator.h b/Engine/Source/Runtime/NavigationSystem/Public/NavMesh/RecastNavMeshGenerator.h
--- a/Engine/Source/Runtime/NavigationSystem/Public/NavMesh/RecastNavMeshGenerator.h
+++ b/Engine/Source/Runtime/NavigationSystem/Public/NavMesh/RecastNavMeshGenerator.h
@@ -1012,6 +1012,14 @@ protected:
 
 		/** Accumulated time spent processing the tile (in seconds). */
 		double CurrentTileRegenDuration;
+
+#if !UE_BUILD_SHIPPING		
+		/** Frame when the tile start being processed. */
+		int64 TileRegenStartFrame = 0;
+
+		/** Frame when the tile is done being processed. */
+		int64 TileRegenEndFrame = 0;
+#endif // !UE_BUILD_SHIPPING		
 		
 		/** if we are currently using time sliced regen or not - currently an experimental feature.
 		 *  do not manipulate this value directly instead call SetNextTimeSliceRegenActive()
diff --git a/Engine/Source/Runtime/NavigationSystem/Public/NavigationSystem.h b/Engine/Source/Runtime/NavigationSystem/Public/NavigationSystem.h
--- a/Engine/Source/Runtime/NavigationSystem/Public/NavigationSystem.h
+++ b/Engine/Source/Runtime/NavigationSystem/Public/NavigationSystem.h
@@ -183,6 +183,27 @@ protected:
 #endif
 };
 
+#if !UE_BUILD_SHIPPING
+struct FTileHistoryData
+{
+	/** Tile coordinates **/
+	int32 TileX = 0;
+	int32 TileY = 0;
+	
+	/** Accumulated time spent processing the tile (seconds). */
+	float TileRegenTime = 0.f;
+
+	/** Time between the tile was requested and when it's finished building (seconds). */
+	float TileWaitTime = 0.f;
+
+	/** Frame when tile processing started. */
+	int64 StartRegenFrame = 0;
+
+	/** Frame when tile processing ended. */
+	int64 EndRegenFrame = 0;
+};
+#endif // !UE_BUILD_SHIPPING
+
 class NAVIGATIONSYSTEM_API FNavRegenTimeSliceManager
 {
 public:
@@ -194,8 +215,14 @@ public:
 
 	double GetAverageDeltaTime() const { return MovingWindowDeltaTime.GetAverage(); }
 
-	void InitTileWaitTimeArrays(const TArray<TObjectPtr<ANavigationData>>& NavDataSet) { TileWaitTimes.SetNum(NavDataSet.Num()); }
+	void ResetTileWaitTimeArrays(const TArray<TObjectPtr<ANavigationData>>& NavDataSet);
 	void PushTileWaitTime(const int32 NavDataIndex, const double NewTime);
+
+#if !UE_BUILD_SHIPPING
+	void ResetTileHistoryData(const TArray<TObjectPtr<ANavigationData>>& NavDataSet);
+	void PushTileHistoryData(const int32 NavDataIndex, const FTileHistoryData& TileData);
+#endif // UE_BUILD_SHIPPING
+
 	double GetAverageTileWaitTime(const int32 NavDataIndex) const;
 	void ResetTileWaitTime(const int32 NavDataIndex);
 
@@ -217,6 +244,10 @@ public:
 	FNavRegenTimeSlicer& GetTimeSlicer() { return TimeSlicer; }
 	const FNavRegenTimeSlicer& GetTimeSlicer() const { return TimeSlicer; }
 
+#if !UE_BUILD_SHIPPING	
+	void LogTileStatistics(const TArray<TObjectPtr<ANavigationData>>& NavDataSet) const;
+#endif // !UE_BUILD_SHIPPING
+	
 protected:
 	FNavRegenTimeSlicer TimeSlicer;
 
@@ -229,6 +260,11 @@ protected:
 	/** Average tile wait time per NavDataIndex */
 	TArray<TArray<double>> TileWaitTimes;
 
+#if !UE_BUILD_SHIPPING
+	/** Tile processing time per NavDataIndex */
+	TArray<TArray<FTileHistoryData>> TileHistoryData;
+#endif // UE_BUILD_SHIPPING
+
 	/** If there are enough tiles to process this in the Min Time Slice Duration */
 	double MinTimeSliceDuration;
 
