[Engine]

diff --git a/Engine/Source/Runtime/MovieSceneTracks/Public/Systems/MovieSceneMotionVectorSimulationSystem.h b/Engine/Source/Runtime/MovieSceneTracks/Public/Systems/MovieSceneMotionVectorSimulationSystem.h
--- a/Engine/Source/Runtime/MovieSceneTracks/Public/Systems/MovieSceneMotionVectorSimulationSystem.h
+++ b/Engine/Source/Runtime/MovieSceneTracks/Public/Systems/MovieSceneMotionVectorSimulationSystem.h
@@ -24,46 +24,46 @@ FFrameTime GetSimulatedMotionVectorTime(const FMovieSceneContext& Context);
 } // namespace MovieScene
 } // namespace UE
 
-UCLASS()
-class MOVIESCENETRACKS_API UMovieSceneMotionVectorSimulationSystem : public UMovieSceneEntitySystem
+UCLASS(MinimalAPI)
+class UMovieSceneMotionVectorSimulationSystem : public UMovieSceneEntitySystem
 {
 public:
 
 	GENERATED_BODY()
 
-	UMovieSceneMotionVectorSimulationSystem(const FObjectInitializer& ObjInit);
+	MOVIESCENETRACKS_API UMovieSceneMotionVectorSimulationSystem(const FObjectInitializer& ObjInit);
 
-	void EnableThisFrame();
-	void SimulateAllTransforms();
+	MOVIESCENETRACKS_API void EnableThisFrame();
+	MOVIESCENETRACKS_API void SimulateAllTransforms();
 	bool IsSimulationEnabled() const
 	{
 		return bSimulationEnabled;
 	}
 
-	void PreserveSimulatedMotion(bool bShouldPreserveTransforms);
+	MOVIESCENETRACKS_API void PreserveSimulatedMotion(bool bShouldPreserveTransforms);
 
-	void AddSimulatedTransform(USceneComponent* Component, const FTransform& SimulatedTransform, FName SocketName);
+	MOVIESCENETRACKS_API void AddSimulatedTransform(USceneComponent* Component, const FTransform& SimulatedTransform, FName SocketName);
 
 private:
 
-	virtual bool IsRelevantImpl(UMovieSceneEntitySystemLinker* InLinker) const override;
-	virtual void OnRun(FSystemTaskPrerequisites& InPrerequisites, FSystemSubsequentTasks& Subsequents) override final;
+	MOVIESCENETRACKS_API virtual bool IsRelevantImpl(UMovieSceneEntitySystemLinker* InLinker) const override;
+	MOVIESCENETRACKS_API virtual void OnRun(FSystemTaskPrerequisites& InPrerequisites, FSystemSubsequentTasks& Subsequents) override final;
 
-	void OnPostEvaluation();
+	MOVIESCENETRACKS_API void OnPostEvaluation();
 
-	void ComputeSimulatedMotion();
+	MOVIESCENETRACKS_API void ComputeSimulatedMotion();
 
 private:
 
-	void PropagateMotionToComponents();
+	MOVIESCENETRACKS_API void PropagateMotionToComponents();
 
-	FTransform GetTransform(USceneComponent* Component);
+	MOVIESCENETRACKS_API FTransform GetTransform(USceneComponent* Component);
 
-	FTransform GetSocketTransform(USceneComponent* Component, FName SocketName);
+	MOVIESCENETRACKS_API FTransform GetSocketTransform(USceneComponent* Component, FName SocketName);
 
-	bool HavePreviousTransformForParent(USceneComponent* InComponent) const;
+	MOVIESCENETRACKS_API bool HavePreviousTransformForParent(USceneComponent* InComponent) const;
 
-	void ApplySimulatedTransforms(USceneComponent* InComponent, const FTransform& InPreviousTransform);
+	MOVIESCENETRACKS_API void ApplySimulatedTransforms(USceneComponent* InComponent, const FTransform& InPreviousTransform);
 
 private:
 
diff --git a/Engine/Source/Runtime/MovieSceneTracks/Public/Systems/MovieScenePiecewiseBoolBlenderSystem.h b/Engine/Source/Runtime/MovieSceneTracks/Public/Systems/MovieScenePiecewiseBoolBlenderSystem.h
--- a/Engine/Source/Runtime/MovieSceneTracks/Public/Systems/MovieScenePiecewiseBoolBlenderSystem.h
+++ b/Engine/Source/Runtime/MovieSceneTracks/Public/Systems/MovieScenePiecewiseBoolBlenderSystem.h
@@ -14,16 +14,16 @@
 class UObject;
 
 
-UCLASS()
-class MOVIESCENETRACKS_API UMovieScenePiecewiseBoolBlenderSystem : public UMovieSceneBlenderSystem
+UCLASS(MinimalAPI)
+class UMovieScenePiecewiseBoolBlenderSystem : public UMovieSceneBlenderSystem
 {
 public:
 
 	GENERATED_BODY()
 
-	UMovieScenePiecewiseBoolBlenderSystem(const FObjectInitializer& ObjInit);
+	MOVIESCENETRACKS_API UMovieScenePiecewiseBoolBlenderSystem(const FObjectInitializer& ObjInit);
 
-	virtual void OnRun(FSystemTaskPrerequisites& InPrerequisites, FSystemSubsequentTasks& Subsequents) override;
+	MOVIESCENETRACKS_API virtual void OnRun(FSystemTaskPrerequisites& InPrerequisites, FSystemSubsequentTasks& Subsequents) override;
 
 private:
 
diff --git a/Engine/Source/Runtime/MovieSceneTracks/Public/Systems/MovieScenePiecewiseByteBlenderSystem.h b/Engine/Source/Runtime/MovieSceneTracks/Public/Systems/MovieScenePiecewiseByteBlenderSystem.h
--- a/Engine/Source/Runtime/MovieSceneTracks/Public/Systems/MovieScenePiecewiseByteBlenderSystem.h
+++ b/Engine/Source/Runtime/MovieSceneTracks/Public/Systems/MovieScenePiecewiseByteBlenderSystem.h
@@ -15,16 +15,16 @@
 class UObject;
 
 
-UCLASS()
-class MOVIESCENETRACKS_API UMovieScenePiecewiseByteBlenderSystem : public UMovieSceneBlenderSystem
+UCLASS(MinimalAPI)
+class UMovieScenePiecewiseByteBlenderSystem : public UMovieSceneBlenderSystem
 {
 public:
 
 	GENERATED_BODY()
 
-	UMovieScenePiecewiseByteBlenderSystem(const FObjectInitializer& ObjInit);
+	MOVIESCENETRACKS_API UMovieScenePiecewiseByteBlenderSystem(const FObjectInitializer& ObjInit);
 
-	virtual void OnRun(FSystemTaskPrerequisites& InPrerequisites, FSystemSubsequentTasks& Subsequents) override;
+	MOVIESCENETRACKS_API virtual void OnRun(FSystemTaskPrerequisites& InPrerequisites, FSystemSubsequentTasks& Subsequents) override;
 
 private:
 
diff --git a/Engine/Source/Runtime/MovieSceneTracks/Public/Systems/MovieScenePiecewiseDoubleBlenderSystem.h b/Engine/Source/Runtime/MovieSceneTracks/Public/Systems/MovieScenePiecewiseDoubleBlenderSystem.h
--- a/Engine/Source/Runtime/MovieSceneTracks/Public/Systems/MovieScenePiecewiseDoubleBlenderSystem.h
+++ b/Engine/Source/Runtime/MovieSceneTracks/Public/Systems/MovieScenePiecewiseDoubleBlenderSystem.h
@@ -84,21 +84,21 @@ struct FAccumulationBuffers
 } // namespace MovieScene
 } // namespace UE
 
-UCLASS(DisplayName="Weighted per-channel", meta=(Tooltip="Blends each channel of this object's transform as separate scalar components. Useful for blending to/from over-rotated objects (ie, 0 < rotation > 360."))
-class MOVIESCENETRACKS_API UMovieScenePiecewiseDoubleBlenderSystem : public UMovieSceneBlenderSystem, public IMovieSceneValueDecomposer
+UCLASS(DisplayName="Weighted per-channel", meta=(Tooltip="Blends each channel of this object's transform as separate scalar components. Useful for blending to/from over-rotated objects (ie, 0 < rotation > 360."), MinimalAPI)
+class UMovieScenePiecewiseDoubleBlenderSystem : public UMovieSceneBlenderSystem, public IMovieSceneValueDecomposer
 {
 public:
 	GENERATED_BODY()
 
-	UMovieScenePiecewiseDoubleBlenderSystem(const FObjectInitializer& ObjInit);
+	MOVIESCENETRACKS_API UMovieScenePiecewiseDoubleBlenderSystem(const FObjectInitializer& ObjInit);
 
 	using FMovieSceneEntityID  = UE::MovieScene::FMovieSceneEntityID;
 	using FComponentTypeID     = UE::MovieScene::FComponentTypeID;
 
-	virtual void OnSchedulePersistentTasks(UE::MovieScene::IEntitySystemScheduler* TaskScheduler) override;
-	virtual void OnRun(FSystemTaskPrerequisites& InPrerequisites, FSystemSubsequentTasks& Subsequents) override;
+	MOVIESCENETRACKS_API virtual void OnSchedulePersistentTasks(UE::MovieScene::IEntitySystemScheduler* TaskScheduler) override;
+	MOVIESCENETRACKS_API virtual void OnRun(FSystemTaskPrerequisites& InPrerequisites, FSystemSubsequentTasks& Subsequents) override;
 
-	virtual FGraphEventRef DispatchDecomposeTask(const UE::MovieScene::FValueDecompositionParams& Params, UE::MovieScene::FAlignedDecomposedValue* Output) override;
+	MOVIESCENETRACKS_API virtual FGraphEventRef DispatchDecomposeTask(const UE::MovieScene::FValueDecompositionParams& Params, UE::MovieScene::FAlignedDecomposedValue* Output) override;
 
 private:
 
diff --git a/Engine/Source/Runtime/MovieSceneTracks/Public/Systems/MovieScenePiecewiseEnumBlenderSystem.h b/Engine/Source/Runtime/MovieSceneTracks/Public/Systems/MovieScenePiecewiseEnumBlenderSystem.h
--- a/Engine/Source/Runtime/MovieSceneTracks/Public/Systems/MovieScenePiecewiseEnumBlenderSystem.h
+++ b/Engine/Source/Runtime/MovieSceneTracks/Public/Systems/MovieScenePiecewiseEnumBlenderSystem.h
@@ -15,16 +15,16 @@
 class UObject;
 
 
-UCLASS()
-class MOVIESCENETRACKS_API UMovieScenePiecewiseEnumBlenderSystem : public UMovieSceneBlenderSystem
+UCLASS(MinimalAPI)
+class UMovieScenePiecewiseEnumBlenderSystem : public UMovieSceneBlenderSystem
 {
 public:
 
 	GENERATED_BODY()
 
-	UMovieScenePiecewiseEnumBlenderSystem(const FObjectInitializer& ObjInit);
+	MOVIESCENETRACKS_API UMovieScenePiecewiseEnumBlenderSystem(const FObjectInitializer& ObjInit);
 
-	virtual void OnRun(FSystemTaskPrerequisites& InPrerequisites, FSystemSubsequentTasks& Subsequents) override;
+	MOVIESCENETRACKS_API virtual void OnRun(FSystemTaskPrerequisites& InPrerequisites, FSystemSubsequentTasks& Subsequents) override;
 
 private:
 
diff --git a/Engine/Source/Runtime/MovieSceneTracks/Public/Systems/MovieScenePiecewiseIntegerBlenderSystem.h b/Engine/Source/Runtime/MovieSceneTracks/Public/Systems/MovieScenePiecewiseIntegerBlenderSystem.h
--- a/Engine/Source/Runtime/MovieSceneTracks/Public/Systems/MovieScenePiecewiseIntegerBlenderSystem.h
+++ b/Engine/Source/Runtime/MovieSceneTracks/Public/Systems/MovieScenePiecewiseIntegerBlenderSystem.h
@@ -39,20 +39,20 @@ namespace MovieScene
 } // namespace UE
 
 
-UCLASS()
-class MOVIESCENETRACKS_API UMovieScenePiecewiseIntegerBlenderSystem : public UMovieSceneBlenderSystem
+UCLASS(MinimalAPI)
+class UMovieScenePiecewiseIntegerBlenderSystem : public UMovieSceneBlenderSystem
 {
 public:
 	GENERATED_BODY()
 
-	UMovieScenePiecewiseIntegerBlenderSystem(const FObjectInitializer& ObjInit);
+	MOVIESCENETRACKS_API UMovieScenePiecewiseIntegerBlenderSystem(const FObjectInitializer& ObjInit);
 
 	using FMovieSceneEntityID  = UE::MovieScene::FMovieSceneEntityID;
 	using FComponentTypeID     = UE::MovieScene::FComponentTypeID;
 
-	virtual void OnLink() override;
-	virtual void OnSchedulePersistentTasks(UE::MovieScene::IEntitySystemScheduler* TaskScheduler) override;
-	virtual void OnRun(FSystemTaskPrerequisites& InPrerequisites, FSystemSubsequentTasks& Subsequents) override;
+	MOVIESCENETRACKS_API virtual void OnLink() override;
+	MOVIESCENETRACKS_API virtual void OnSchedulePersistentTasks(UE::MovieScene::IEntitySystemScheduler* TaskScheduler) override;
+	MOVIESCENETRACKS_API virtual void OnRun(FSystemTaskPrerequisites& InPrerequisites, FSystemSubsequentTasks& Subsequents) override;
 
 private:
 
diff --git a/Engine/Source/Runtime/MovieSceneTracks/Public/Systems/MovieScenePropertyInstantiator.h b/Engine/Source/Runtime/MovieSceneTracks/Public/Systems/MovieScenePropertyInstantiator.h
--- a/Engine/Source/Runtime/MovieSceneTracks/Public/Systems/MovieScenePropertyInstantiator.h
+++ b/Engine/Source/Runtime/MovieSceneTracks/Public/Systems/MovieScenePropertyInstantiator.h
@@ -33,8 +33,8 @@ class FEntityManager;
 class UMovieSceneBlenderSystem;
 
 /** Class responsible for resolving all property types registered with FBuiltInComponentTypes::PropertyRegistry */
-UCLASS()
-class MOVIESCENETRACKS_API UMovieScenePropertyInstantiatorSystem
+UCLASS(MinimalAPI)
+class UMovieScenePropertyInstantiatorSystem
 	: public UMovieSceneEntityInstantiatorSystem
 {
 public:
@@ -43,12 +43,12 @@ public:
 
 	GENERATED_BODY()
 
-	UMovieScenePropertyInstantiatorSystem(const FObjectInitializer& ObjInit);
+	MOVIESCENETRACKS_API UMovieScenePropertyInstantiatorSystem(const FObjectInitializer& ObjInit);
 
 	/**
 	 * Retrieve the stats for a specific property
 	 */
-	UE::MovieScene::FPropertyStats GetStatsForProperty(UE::MovieScene::FCompositePropertyTypeID PropertyID) const;
+	MOVIESCENETRACKS_API UE::MovieScene::FPropertyStats GetStatsForProperty(UE::MovieScene::FCompositePropertyTypeID PropertyID) const;
 
 
 	/**
@@ -96,10 +96,10 @@ public:
 
 private:
 
-	virtual void OnRun(FSystemTaskPrerequisites& InPrerequisites, FSystemSubsequentTasks& Subsequents) override;
-	virtual void OnLink() override;
-	virtual void OnUnlink() override;
-	virtual void OnCleanTaggedGarbage() override;
+	MOVIESCENETRACKS_API virtual void OnRun(FSystemTaskPrerequisites& InPrerequisites, FSystemSubsequentTasks& Subsequents) override;
+	MOVIESCENETRACKS_API virtual void OnLink() override;
+	MOVIESCENETRACKS_API virtual void OnUnlink() override;
+	MOVIESCENETRACKS_API virtual void OnCleanTaggedGarbage() override;
 
 private:
 
@@ -226,20 +226,20 @@ private:
 			TArrayView<const UE::MovieScene::FPropertyCompositeDefinition> Composites,
 			UE::MovieScene::FComponentMask& OutComponentType) const;
 	};
-	void DiscoverInvalidatedProperties(TBitArray<>& OutInvalidatedProperties);
-	void UpgradeFloatToDoubleProperties(const TBitArray<>& InvalidatedProperties);
-	void ProcessInvalidatedProperties(const TBitArray<>& InvalidatedProperties);
-	void UpdatePropertyInfo(const FPropertyParameters& Params);
-	void InitializeFastPath(const FPropertyParameters& Params);
-	void InitializeBlendPath(const FPropertyParameters& Params);
+	MOVIESCENETRACKS_API void DiscoverInvalidatedProperties(TBitArray<>& OutInvalidatedProperties);
+	MOVIESCENETRACKS_API void UpgradeFloatToDoubleProperties(const TBitArray<>& InvalidatedProperties);
+	MOVIESCENETRACKS_API void ProcessInvalidatedProperties(const TBitArray<>& InvalidatedProperties);
+	MOVIESCENETRACKS_API void UpdatePropertyInfo(const FPropertyParameters& Params);
+	MOVIESCENETRACKS_API void InitializeFastPath(const FPropertyParameters& Params);
+	MOVIESCENETRACKS_API void InitializeBlendPath(const FPropertyParameters& Params);
 
-	FSetupBlenderSystemResult SetupBlenderSystem(const FPropertyParameters& Params);
+	MOVIESCENETRACKS_API FSetupBlenderSystemResult SetupBlenderSystem(const FPropertyParameters& Params);
 
-	int32 ResolveProperty(UE::MovieScene::FCustomAccessorView CustomAccessors, UObject* Object, const FMovieScenePropertyBinding& PropertyBinding, int32 PropertyDefinitionIndex);
+	MOVIESCENETRACKS_API int32 ResolveProperty(UE::MovieScene::FCustomAccessorView CustomAccessors, UObject* Object, const FMovieScenePropertyBinding& PropertyBinding, int32 PropertyDefinitionIndex);
 
-	UE::MovieScene::FPropertyRecomposerPropertyInfo FindPropertyFromSource(FMovieSceneEntityID EntityID, UObject* Object) const;
+	MOVIESCENETRACKS_API UE::MovieScene::FPropertyRecomposerPropertyInfo FindPropertyFromSource(FMovieSceneEntityID EntityID, UObject* Object) const;
 
-	void InitializePropertyMetaData(FSystemTaskPrerequisites& InPrerequisites, FSystemSubsequentTasks& Subsequents);
+	MOVIESCENETRACKS_API void InitializePropertyMetaData(FSystemTaskPrerequisites& InPrerequisites, FSystemSubsequentTasks& Subsequents);
 
 private:
 
diff --git a/Engine/Source/Runtime/MovieSceneTracks/Public/Systems/MovieScenePropertySystem.h b/Engine/Source/Runtime/MovieSceneTracks/Public/Systems/MovieScenePropertySystem.h
--- a/Engine/Source/Runtime/MovieSceneTracks/Public/Systems/MovieScenePropertySystem.h
+++ b/Engine/Source/Runtime/MovieSceneTracks/Public/Systems/MovieScenePropertySystem.h
@@ -12,15 +12,15 @@ class UMovieScenePropertyInstantiatorSystem;
 
 
 /** Abstract base class for any property system that deals with a property registered with FBuiltInComponentTypes::PropertyRegistry */
-UCLASS(Abstract)
-class MOVIESCENETRACKS_API UMovieScenePropertySystem
+UCLASS(Abstract, MinimalAPI)
+class UMovieScenePropertySystem
 	: public UMovieSceneEntitySystem
 	, public IMovieScenePreAnimatedStateSystemInterface
 {
 public:
 	GENERATED_BODY()
 
-	UMovieScenePropertySystem(const FObjectInitializer& ObjInit);
+	MOVIESCENETRACKS_API UMovieScenePropertySystem(const FObjectInitializer& ObjInit);
 
 	/**
 	 * Must be called on construction of derived classes to initialize the members necessary for this system to animate its property
@@ -36,12 +36,12 @@ public:
 
 protected:
 
-	virtual void OnLink() override;
-	virtual void OnSchedulePersistentTasks(UE::MovieScene::IEntitySystemScheduler* TaskScheduler) override;
-	virtual void OnRun(FSystemTaskPrerequisites& InPrerequisites, FSystemSubsequentTasks& Subsequents) override;
+	MOVIESCENETRACKS_API virtual void OnLink() override;
+	MOVIESCENETRACKS_API virtual void OnSchedulePersistentTasks(UE::MovieScene::IEntitySystemScheduler* TaskScheduler) override;
+	MOVIESCENETRACKS_API virtual void OnRun(FSystemTaskPrerequisites& InPrerequisites, FSystemSubsequentTasks& Subsequents) override;
 
-	virtual void SavePreAnimatedState(const FPreAnimationParameters& InParameters) override;
-	virtual void RestorePreAnimatedState(const FPreAnimationParameters& InParameters) override;
+	MOVIESCENETRACKS_API virtual void SavePreAnimatedState(const FPreAnimationParameters& InParameters) override;
+	MOVIESCENETRACKS_API virtual void RestorePreAnimatedState(const FPreAnimationParameters& InParameters) override;
 
 	/** Pointer to the property instantiator system for retrieving property stats */
 	UPROPERTY()
diff --git a/Engine/Source/Runtime/MovieSceneTracks/Public/Systems/MovieSceneQuaternionBlenderSystem.h b/Engine/Source/Runtime/MovieSceneTracks/Public/Systems/MovieSceneQuaternionBlenderSystem.h
--- a/Engine/Source/Runtime/MovieSceneTracks/Public/Systems/MovieSceneQuaternionBlenderSystem.h
+++ b/Engine/Source/Runtime/MovieSceneTracks/Public/Systems/MovieSceneQuaternionBlenderSystem.h
@@ -27,18 +27,18 @@ namespace UE::MovieScene
 	};
 }
 
-UCLASS(DisplayName="Shortest Rotation (quaternion slerp)", meta=(Tooltip="Uses a quaternion spherical interpolation (slerp) to blend between transforms using the shortest rotation path. Does not support over-rotation."))
-class MOVIESCENETRACKS_API UMovieSceneQuaternionBlenderSystem : public UMovieSceneBlenderSystem, public IMovieSceneValueDecomposer
+UCLASS(DisplayName="Shortest Rotation (quaternion slerp)", meta=(Tooltip="Uses a quaternion spherical interpolation (slerp) to blend between transforms using the shortest rotation path. Does not support over-rotation."), MinimalAPI)
+class UMovieSceneQuaternionBlenderSystem : public UMovieSceneBlenderSystem, public IMovieSceneValueDecomposer
 {
 public:
 	GENERATED_BODY()
 
-	UMovieSceneQuaternionBlenderSystem(const FObjectInitializer& ObjInit);
+	MOVIESCENETRACKS_API UMovieSceneQuaternionBlenderSystem(const FObjectInitializer& ObjInit);
 
-	virtual void OnSchedulePersistentTasks(UE::MovieScene::IEntitySystemScheduler* TaskScheduler) override;
-	virtual void OnRun(FSystemTaskPrerequisites& InPrerequisites, FSystemSubsequentTasks& Subsequents) override;
+	MOVIESCENETRACKS_API virtual void OnSchedulePersistentTasks(UE::MovieScene::IEntitySystemScheduler* TaskScheduler) override;
+	MOVIESCENETRACKS_API virtual void OnRun(FSystemTaskPrerequisites& InPrerequisites, FSystemSubsequentTasks& Subsequents) override;
 
-	virtual FGraphEventRef DispatchDecomposeTask(const UE::MovieScene::FValueDecompositionParams& Params, UE::MovieScene::FAlignedDecomposedValue* Output) override;
+	MOVIESCENETRACKS_API virtual FGraphEventRef DispatchDecomposeTask(const UE::MovieScene::FValueDecompositionParams& Params, UE::MovieScene::FAlignedDecomposedValue* Output) override;
 
 private:
 
diff --git a/Engine/Source/Runtime/MovieSceneTracks/Public/Systems/MovieSceneQuaternionInterpolationRotationSystem.h b/Engine/Source/Runtime/MovieSceneTracks/Public/Systems/MovieSceneQuaternionInterpolationRotationSystem.h
--- a/Engine/Source/Runtime/MovieSceneTracks/Public/Systems/MovieSceneQuaternionInterpolationRotationSystem.h
+++ b/Engine/Source/Runtime/MovieSceneTracks/Public/Systems/MovieSceneQuaternionInterpolationRotationSystem.h
@@ -6,18 +6,18 @@
 #include "MovieSceneQuaternionInterpolationRotationSystem.generated.h"
 
 
-UCLASS()
-class MOVIESCENETRACKS_API UMovieSceneQuaternionInterpolationRotationSystem : public UMovieSceneEntitySystem
+UCLASS(MinimalAPI)
+class UMovieSceneQuaternionInterpolationRotationSystem : public UMovieSceneEntitySystem
 {
 public:
 
 	GENERATED_BODY()
 
-	UMovieSceneQuaternionInterpolationRotationSystem(const FObjectInitializer& ObjInit);
+	MOVIESCENETRACKS_API UMovieSceneQuaternionInterpolationRotationSystem(const FObjectInitializer& ObjInit);
 
 private:
-	virtual bool IsRelevantImpl(UMovieSceneEntitySystemLinker* InLinker) const override;
-	virtual void OnSchedulePersistentTasks(UE::MovieScene::IEntitySystemScheduler* TaskScheduler) override;
-	virtual void OnRun(FSystemTaskPrerequisites& InPrerequisites, FSystemSubsequentTasks& Subsequents) override;
+	MOVIESCENETRACKS_API virtual bool IsRelevantImpl(UMovieSceneEntitySystemLinker* InLinker) const override;
+	MOVIESCENETRACKS_API virtual void OnSchedulePersistentTasks(UE::MovieScene::IEntitySystemScheduler* TaskScheduler) override;
+	MOVIESCENETRACKS_API virtual void OnRun(FSystemTaskPrerequisites& InPrerequisites, FSystemSubsequentTasks& Subsequents) override;
 };
 
diff --git a/Engine/Source/Runtime/MovieSceneTracks/Public/Systems/MovieSceneStringPropertySystem.h b/Engine/Source/Runtime/MovieSceneTracks/Public/Systems/MovieSceneStringPropertySystem.h
--- a/Engine/Source/Runtime/MovieSceneTracks/Public/Systems/MovieSceneStringPropertySystem.h
+++ b/Engine/Source/Runtime/MovieSceneTracks/Public/Systems/MovieSceneStringPropertySystem.h
@@ -7,12 +7,12 @@
 
 
 
-UCLASS()
-class MOVIESCENETRACKS_API UMovieSceneStringPropertySystem : public UMovieScenePropertySystem
+UCLASS(MinimalAPI)
+class UMovieSceneStringPropertySystem : public UMovieScenePropertySystem
 {
 	GENERATED_BODY()
 
 	UMovieSceneStringPropertySystem(const FObjectInitializer& ObjInit);
 
-	virtual void OnRun(FSystemTaskPrerequisites& InPrerequisites, FSystemSubsequentTasks& Subsequents) override;
+	MOVIESCENETRACKS_API virtual void OnRun(FSystemTaskPrerequisites& InPrerequisites, FSystemSubsequentTasks& Subsequents) override;
 };
diff --git a/Engine/Source/Runtime/MovieSceneTracks/Public/Systems/MovieSceneVectorPropertySystem.h b/Engine/Source/Runtime/MovieSceneTracks/Public/Systems/MovieSceneVectorPropertySystem.h
--- a/Engine/Source/Runtime/MovieSceneTracks/Public/Systems/MovieSceneVectorPropertySystem.h
+++ b/Engine/Source/Runtime/MovieSceneTracks/Public/Systems/MovieSceneVectorPropertySystem.h
@@ -6,25 +6,25 @@
 #include "Systems/MovieScenePropertySystem.h"
 #include "MovieSceneVectorPropertySystem.generated.h"
 
-UCLASS()
-class MOVIESCENETRACKS_API UMovieSceneFloatVectorPropertySystem : public UMovieScenePropertySystem
+UCLASS(MinimalAPI)
+class UMovieSceneFloatVectorPropertySystem : public UMovieScenePropertySystem
 {
 	GENERATED_BODY()
 
 private:
 	UMovieSceneFloatVectorPropertySystem(const FObjectInitializer& ObjInit);
 
-	virtual void OnRun(FSystemTaskPrerequisites& InPrerequisites, FSystemSubsequentTasks& Subsequents) override;
+	MOVIESCENETRACKS_API virtual void OnRun(FSystemTaskPrerequisites& InPrerequisites, FSystemSubsequentTasks& Subsequents) override;
 };
 
-UCLASS()
-class MOVIESCENETRACKS_API UMovieSceneDoubleVectorPropertySystem : public UMovieScenePropertySystem
+UCLASS(MinimalAPI)
+class UMovieSceneDoubleVectorPropertySystem : public UMovieScenePropertySystem
 {
 	GENERATED_BODY()
 
 private:
 	UMovieSceneDoubleVectorPropertySystem(const FObjectInitializer& ObjInit);
 
-	virtual void OnRun(FSystemTaskPrerequisites& InPrerequisites, FSystemSubsequentTasks& Subsequents) override;
+	MOVIESCENETRACKS_API virtual void OnRun(FSystemTaskPrerequisites& InPrerequisites, FSystemSubsequentTasks& Subsequents) override;
 };
 
diff --git a/Engine/Source/Runtime/MovieSceneTracks/Public/Systems/WeightAndEasingEvaluatorSystem.h b/Engine/Source/Runtime/MovieSceneTracks/Public/Systems/WeightAndEasingEvaluatorSystem.h
--- a/Engine/Source/Runtime/MovieSceneTracks/Public/Systems/WeightAndEasingEvaluatorSystem.h
+++ b/Engine/Source/Runtime/MovieSceneTracks/Public/Systems/WeightAndEasingEvaluatorSystem.h
@@ -104,8 +104,8 @@ struct FHierarchicalEasingChannelBuffer
 /**
  * System that creates hierarchical easing channels for any newly introduced HierarchicalEasingProvider components
  */
-UCLASS()
-class MOVIESCENETRACKS_API UMovieSceneHierarchicalEasingInstantiatorSystem : public UMovieSceneEntityInstantiatorSystem
+UCLASS(MinimalAPI)
+class UMovieSceneHierarchicalEasingInstantiatorSystem : public UMovieSceneEntityInstantiatorSystem
 {
 public:
 
@@ -118,16 +118,16 @@ public:
 
 	GENERATED_BODY()
 
-	UMovieSceneHierarchicalEasingInstantiatorSystem(const FObjectInitializer& ObjInit);
+	MOVIESCENETRACKS_API UMovieSceneHierarchicalEasingInstantiatorSystem(const FObjectInitializer& ObjInit);
 
-	void FinalizeBlendTargets();
+	MOVIESCENETRACKS_API void FinalizeBlendTargets();
 
 private:
 
-	virtual bool IsRelevantImpl(UMovieSceneEntitySystemLinker* InLinker) const;
-	virtual void OnLink() override;
-	virtual void OnUnlink() override;
-	virtual void OnRun(FSystemTaskPrerequisites& InPrerequisites, FSystemSubsequentTasks& Subsequents) override;
+	MOVIESCENETRACKS_API virtual bool IsRelevantImpl(UMovieSceneEntitySystemLinker* InLinker) const;
+	MOVIESCENETRACKS_API virtual void OnLink() override;
+	MOVIESCENETRACKS_API virtual void OnUnlink() override;
+	MOVIESCENETRACKS_API virtual void OnRun(FSystemTaskPrerequisites& InPrerequisites, FSystemSubsequentTasks& Subsequents) override;
 
 private:
 
@@ -170,27 +170,27 @@ private:
 /**
  * System that combines manual weights and easings and propagates them to entities with matching EasingChannelID components
  */
-UCLASS()
-class MOVIESCENETRACKS_API UWeightAndEasingEvaluatorSystem : public UMovieSceneEntitySystem
+UCLASS(MinimalAPI)
+class UWeightAndEasingEvaluatorSystem : public UMovieSceneEntitySystem
 {
 public:
 
 	GENERATED_BODY()
 
-	UWeightAndEasingEvaluatorSystem(const FObjectInitializer& ObjInit);
+	MOVIESCENETRACKS_API UWeightAndEasingEvaluatorSystem(const FObjectInitializer& ObjInit);
 
 	/**
 	 * Called from the instantiator system once channel IDs have been calculated
 	 */
-	UE::MovieScene::FHierarchicalEasingChannelBuffer& GetComputationBuffer();
+	MOVIESCENETRACKS_API UE::MovieScene::FHierarchicalEasingChannelBuffer& GetComputationBuffer();
 
 private:
 
-	virtual void OnSchedulePersistentTasks(UE::MovieScene::IEntitySystemScheduler* TaskScheduler) override;
-	virtual void OnRun(FSystemTaskPrerequisites& InPrerequisites, FSystemSubsequentTasks& Subsequents) override;
-	virtual bool IsRelevantImpl(UMovieSceneEntitySystemLinker* InLinker) const override;
-	virtual void OnLink() override final;
-	virtual void OnUnlink() override final;
+	MOVIESCENETRACKS_API virtual void OnSchedulePersistentTasks(UE::MovieScene::IEntitySystemScheduler* TaskScheduler) override;
+	MOVIESCENETRACKS_API virtual void OnRun(FSystemTaskPrerequisites& InPrerequisites, FSystemSubsequentTasks& Subsequents) override;
+	MOVIESCENETRACKS_API virtual bool IsRelevantImpl(UMovieSceneEntitySystemLinker* InLinker) const override;
+	MOVIESCENETRACKS_API virtual void OnLink() override final;
+	MOVIESCENETRACKS_API virtual void OnUnlink() override final;
 
 private:
 
diff --git a/Engine/Source/Runtime/MovieSceneTracks/Public/Tests/MovieSceneTestDataBuilders.h b/Engine/Source/Runtime/MovieSceneTracks/Public/Tests/MovieSceneTestDataBuilders.h
--- a/Engine/Source/Runtime/MovieSceneTracks/Public/Tests/MovieSceneTestDataBuilders.h
+++ b/Engine/Source/Runtime/MovieSceneTracks/Public/Tests/MovieSceneTestDataBuilders.h
@@ -24,8 +24,8 @@
  * Bound objects are specified manually on the sequence and will be simply returned when
  * bindings are resolved.
  */
-UCLASS()
-class MOVIESCENETRACKS_API UMovieSceneTestSequence : public UMovieSceneSequence
+UCLASS(MinimalAPI)
+class UMovieSceneTestSequence : public UMovieSceneSequence
 {
 	GENERATED_BODY()
 
@@ -36,7 +36,7 @@ public:
 
 public:
 	/** Initialize this test sequence */
-	void Initialize();
+	MOVIESCENETRACKS_API void Initialize();
 
 	/**
 	 * Add an object binding to the sequnce
@@ -44,13 +44,13 @@ public:
 	 * @param InObject The object that will be returned when the binding is resolved
 	 * @return The ID of the new object binding
 	 */
-	FGuid AddObjectBinding(TObjectPtr<UObject> InObject);
+	MOVIESCENETRACKS_API FGuid AddObjectBinding(TObjectPtr<UObject> InObject);
 
 public:
 	/** UMovieSceneSequence interface */
 	virtual void BindPossessableObject(const FGuid& ObjectId, UObject& PossessedObject, UObject* Context) override {}
 	virtual bool CanPossessObject(UObject& Object, UObject* InPlaybackContext) const override { return true; }
-	virtual void LocateBoundObjects(const FGuid& ObjectId, UObject* Context, TArray<UObject*, TInlineAllocator<1>>& OutObjects) const override;
+	MOVIESCENETRACKS_API virtual void LocateBoundObjects(const FGuid& ObjectId, UObject* Context, TArray<UObject*, TInlineAllocator<1>>& OutObjects) const override;
 	virtual UMovieScene* GetMovieScene() const override { return MovieScene; }
 	virtual UObject* GetParentObject(UObject* Object) const override { return nullptr; }
 	virtual void UnbindPossessableObjects(const FGuid& ObjectId) override {}
diff --git a/Engine/Source/Runtime/MovieSceneTracks/Public/Tracks/IMovieSceneTransformOrigin.h b/Engine/Source/Runtime/MovieSceneTracks/Public/Tracks/IMovieSceneTransformOrigin.h
--- a/Engine/Source/Runtime/MovieSceneTracks/Public/Tracks/IMovieSceneTransformOrigin.h
+++ b/Engine/Source/Runtime/MovieSceneTracks/Public/Tracks/IMovieSceneTransformOrigin.h
@@ -8,15 +8,15 @@
 #include "IMovieSceneTransformOrigin.generated.h"
 
 
-UINTERFACE(Category="Sequencer", Blueprintable, meta=(DisplayName = "Transform Origin"))
-class MOVIESCENETRACKS_API UMovieSceneTransformOrigin : public UInterface
+UINTERFACE(Category="Sequencer", Blueprintable, meta=(DisplayName = "Transform Origin"), MinimalAPI)
+class UMovieSceneTransformOrigin : public UInterface
 {
 public:
 	GENERATED_BODY()
 };
 
 /** Interface that should be implemented to provide transform tracks with an origin transform. Scale is ignored. */
-class MOVIESCENETRACKS_API IMovieSceneTransformOrigin
+class IMovieSceneTransformOrigin
 {
 public:
 	GENERATED_BODY()
@@ -31,6 +31,6 @@ protected:
 
 	/** Get the transform from which all absolute component transform sections should be relative. Scale is ignored. */
 	UFUNCTION(BlueprintCallable, BlueprintImplementableEvent, Category = "Sequencer|Section", DisplayName = "GetTransformOrigin", meta=(CallInEditor="true"))
-	FTransform BP_GetTransformOrigin() const;
+	MOVIESCENETRACKS_API FTransform BP_GetTransformOrigin() const;
 	virtual FTransform NativeGetTransformOrigin() const { return BP_GetTransformOrigin(); }
 };
diff --git a/Engine/Source/Runtime/MovieSceneTracks/Public/Tracks/MovieSceneAudioTrack.h b/Engine/Source/Runtime/MovieSceneTracks/Public/Tracks/MovieSceneAudioTrack.h
--- a/Engine/Source/Runtime/MovieSceneTracks/Public/Tracks/MovieSceneAudioTrack.h
+++ b/Engine/Source/Runtime/MovieSceneTracks/Public/Tracks/MovieSceneAudioTrack.h
@@ -19,8 +19,8 @@ namespace AudioTrackConstants
 /**
  * Handles manipulation of audio.
  */
-UCLASS()
-class MOVIESCENETRACKS_API UMovieSceneAudioTrack
+UCLASS(MinimalAPI)
+class UMovieSceneAudioTrack
 	: public UMovieSceneNameableTrack
 {
 	GENERATED_UCLASS_BODY()
@@ -28,7 +28,7 @@ class MOVIESCENETRACKS_API UMovieSceneAudioTrack
 public:
 
 	/** Adds a new sound cue to the audio */
-	virtual UMovieSceneSection* AddNewSoundOnRow(USoundBase* Sound, FFrameNumber Time, int32 RowIndex);
+	MOVIESCENETRACKS_API virtual UMovieSceneSection* AddNewSoundOnRow(USoundBase* Sound, FFrameNumber Time, int32 RowIndex);
 
 	/** Adds a new sound cue on the next available/non-overlapping row */
 	virtual UMovieSceneSection* AddNewSound(USoundBase* Sound, FFrameNumber Time) { return AddNewSoundOnRow(Sound, Time, INDEX_NONE); }
@@ -40,25 +40,25 @@ public:
 	}
 
 	UE_DEPRECATED(5.2, "IsAMasterTrack is deprecated. Please use Cast<UMovieScene>(GetOuter())->ContainsTrack instead")
-	bool IsAMasterTrack() const;
+	MOVIESCENETRACKS_API bool IsAMasterTrack() const;
 
 public:
 
 	// UMovieSceneTrack interface
 
-	virtual bool SupportsType(TSubclassOf<UMovieSceneSection> SectionClass) const override;
-	virtual void RemoveAllAnimationData() override;
-	virtual bool HasSection(const UMovieSceneSection& Section) const override;
-	virtual void AddSection(UMovieSceneSection& Section) override;
-	virtual void RemoveSection(UMovieSceneSection& Section) override;
-	virtual void RemoveSectionAt(int32 SectionIndex) override;
-	virtual bool IsEmpty() const override;
-	virtual const TArray<UMovieSceneSection*>& GetAllSections() const override;
-	virtual bool SupportsMultipleRows() const override;
-	virtual UMovieSceneSection* CreateNewSection() override;
+	MOVIESCENETRACKS_API virtual bool SupportsType(TSubclassOf<UMovieSceneSection> SectionClass) const override;
+	MOVIESCENETRACKS_API virtual void RemoveAllAnimationData() override;
+	MOVIESCENETRACKS_API virtual bool HasSection(const UMovieSceneSection& Section) const override;
+	MOVIESCENETRACKS_API virtual void AddSection(UMovieSceneSection& Section) override;
+	MOVIESCENETRACKS_API virtual void RemoveSection(UMovieSceneSection& Section) override;
+	MOVIESCENETRACKS_API virtual void RemoveSectionAt(int32 SectionIndex) override;
+	MOVIESCENETRACKS_API virtual bool IsEmpty() const override;
+	MOVIESCENETRACKS_API virtual const TArray<UMovieSceneSection*>& GetAllSections() const override;
+	MOVIESCENETRACKS_API virtual bool SupportsMultipleRows() const override;
+	MOVIESCENETRACKS_API virtual UMovieSceneSection* CreateNewSection() override;
 
 	// ~UObject interface
-	virtual void PostRename(UObject* OldOuter, const FName OldName) override;
+	MOVIESCENETRACKS_API virtual void PostRename(UObject* OldOuter, const FName OldName) override;
 
 private:
 
diff --git a/Engine/Source/Runtime/MovieSceneTracks/Public/Tracks/MovieSceneByteTrack.h b/Engine/Source/Runtime/MovieSceneTracks/Public/Tracks/MovieSceneByteTrack.h
--- a/Engine/Source/Runtime/MovieSceneTracks/Public/Tracks/MovieSceneByteTrack.h
+++ b/Engine/Source/Runtime/MovieSceneTracks/Public/Tracks/MovieSceneByteTrack.h
@@ -12,20 +12,20 @@ class UEnum;
 /**
  * Handles manipulation of byte properties in a movie scene
  */
-UCLASS()
-class MOVIESCENETRACKS_API UMovieSceneByteTrack : public UMovieScenePropertyTrack
+UCLASS(MinimalAPI)
+class UMovieSceneByteTrack : public UMovieScenePropertyTrack
 {
 	GENERATED_UCLASS_BODY()
 
 public:
 	/** UMovieSceneTrack interface */
-	virtual void PostLoad() override;
-	virtual bool SupportsType(TSubclassOf<UMovieSceneSection> SectionClass) const override;
-	virtual UMovieSceneSection* CreateNewSection() override;
+	MOVIESCENETRACKS_API virtual void PostLoad() override;
+	MOVIESCENETRACKS_API virtual bool SupportsType(TSubclassOf<UMovieSceneSection> SectionClass) const override;
+	MOVIESCENETRACKS_API virtual UMovieSceneSection* CreateNewSection() override;
 
-	void SetEnum(UEnum* Enum);
+	MOVIESCENETRACKS_API void SetEnum(UEnum* Enum);
 
-	class UEnum* GetEnum() const;
+	MOVIESCENETRACKS_API class UEnum* GetEnum() const;
 
 protected:
 	UPROPERTY()
diff --git a/Engine/Source/Runtime/MovieSceneTracks/Public/Tracks/MovieSceneCameraShakeSourceTriggerTrack.h b/Engine/Source/Runtime/MovieSceneTracks/Public/Tracks/MovieSceneCameraShakeSourceTriggerTrack.h
--- a/Engine/Source/Runtime/MovieSceneTracks/Public/Tracks/MovieSceneCameraShakeSourceTriggerTrack.h
+++ b/Engine/Source/Runtime/MovieSceneTracks/Public/Tracks/MovieSceneCameraShakeSourceTriggerTrack.h
@@ -20,8 +20,8 @@
 class UObject;
 struct FMovieSceneEvaluationTrack;
 
-UCLASS()
-class MOVIESCENETRACKS_API UMovieSceneCameraShakeSourceTriggerTrack
+UCLASS(MinimalAPI)
+class UMovieSceneCameraShakeSourceTriggerTrack
 	: public UMovieSceneTrack
 	, public IMovieSceneTrackTemplateProducer
 {
@@ -29,24 +29,24 @@ public:
 
 	GENERATED_BODY()
 
-	UMovieSceneCameraShakeSourceTriggerTrack(const FObjectInitializer& Obj);
+	MOVIESCENETRACKS_API UMovieSceneCameraShakeSourceTriggerTrack(const FObjectInitializer& Obj);
 
 	// UMovieSceneTrack interface
-	virtual bool SupportsType(TSubclassOf<UMovieSceneSection> SectionClass) const override;
-	virtual UMovieSceneSection* CreateNewSection() override;
-	virtual bool HasSection(const UMovieSceneSection& Section) const override;
-	virtual void AddSection(UMovieSceneSection& Section) override;
-	virtual void RemoveSection(UMovieSceneSection& Section) override;
-	virtual void RemoveSectionAt(int32 SectionIndex) override;
-	virtual void RemoveAllAnimationData() override;
-	virtual bool IsEmpty() const override;
-	virtual const TArray<UMovieSceneSection*>& GetAllSections() const override;
-	virtual void PostCompile(FMovieSceneEvaluationTrack& Track, const FMovieSceneTrackCompilerArgs& Args) const override;
-
-	virtual FMovieSceneEvalTemplatePtr CreateTemplateForSection(const UMovieSceneSection& InSection) const override;
+	MOVIESCENETRACKS_API virtual bool SupportsType(TSubclassOf<UMovieSceneSection> SectionClass) const override;
+	MOVIESCENETRACKS_API virtual UMovieSceneSection* CreateNewSection() override;
+	MOVIESCENETRACKS_API virtual bool HasSection(const UMovieSceneSection& Section) const override;
+	MOVIESCENETRACKS_API virtual void AddSection(UMovieSceneSection& Section) override;
+	MOVIESCENETRACKS_API virtual void RemoveSection(UMovieSceneSection& Section) override;
+	MOVIESCENETRACKS_API virtual void RemoveSectionAt(int32 SectionIndex) override;
+	MOVIESCENETRACKS_API virtual void RemoveAllAnimationData() override;
+	MOVIESCENETRACKS_API virtual bool IsEmpty() const override;
+	MOVIESCENETRACKS_API virtual const TArray<UMovieSceneSection*>& GetAllSections() const override;
+	MOVIESCENETRACKS_API virtual void PostCompile(FMovieSceneEvaluationTrack& Track, const FMovieSceneTrackCompilerArgs& Args) const override;
+
+	MOVIESCENETRACKS_API virtual FMovieSceneEvalTemplatePtr CreateTemplateForSection(const UMovieSceneSection& InSection) const override;
 
 #if WITH_EDITORONLY_DATA
-	virtual FText GetDisplayName() const override;
+	MOVIESCENETRACKS_API virtual FText GetDisplayName() const override;
 #endif
 
 protected:
diff --git a/Engine/Source/Runtime/MovieSceneTracks/Public/Tracks/MovieSceneCinematicShotTrack.h b/Engine/Source/Runtime/MovieSceneTracks/Public/Tracks/MovieSceneCinematicShotTrack.h
--- a/Engine/Source/Runtime/MovieSceneTracks/Public/Tracks/MovieSceneCinematicShotTrack.h
+++ b/Engine/Source/Runtime/MovieSceneTracks/Public/Tracks/MovieSceneCinematicShotTrack.h
@@ -26,38 +26,38 @@ template <typename> struct TMovieSceneEvaluationTree;
 /**
 * A track that holds consecutive sub sequences.
 */
-UCLASS()
-class MOVIESCENETRACKS_API UMovieSceneCinematicShotTrack
+UCLASS(MinimalAPI)
+class UMovieSceneCinematicShotTrack
 	: public UMovieSceneSubTrack
 {
 	GENERATED_BODY()
 
 public:
 
-	UMovieSceneCinematicShotTrack(const FObjectInitializer& ObjectInitializer);
+	MOVIESCENETRACKS_API UMovieSceneCinematicShotTrack(const FObjectInitializer& ObjectInitializer);
 
-	void SortSections();
+	MOVIESCENETRACKS_API void SortSections();
 
 	// UMovieSceneSubTrack interface
 	virtual UMovieSceneSubSection* AddSequence(UMovieSceneSequence* Sequence, FFrameNumber StartTime, int32 Duration) { return AddSequenceOnRow(Sequence, StartTime, Duration, INDEX_NONE); }
 
-	virtual UMovieSceneSubSection* AddSequenceOnRow(UMovieSceneSequence* Sequence, FFrameNumber StartTime, int32 Duration, int32 RowIndex);
+	MOVIESCENETRACKS_API virtual UMovieSceneSubSection* AddSequenceOnRow(UMovieSceneSequence* Sequence, FFrameNumber StartTime, int32 Duration, int32 RowIndex);
 	
 	// UMovieSceneTrack interface
-	virtual void AddSection(UMovieSceneSection& Section) override;
-	virtual bool SupportsType(TSubclassOf<UMovieSceneSection> SectionClass) const override;
-	virtual UMovieSceneSection* CreateNewSection() override;
-	virtual void RemoveSection(UMovieSceneSection& Section) override;
-	virtual void RemoveSectionAt(int32 SectionIndex) override;
-	virtual bool SupportsMultipleRows() const override;
-	virtual bool PopulateEvaluationTree(TMovieSceneEvaluationTree<FMovieSceneTrackEvaluationData>& OutData) const override;
-	virtual int8 GetEvaluationFieldVersion() const override;
+	MOVIESCENETRACKS_API virtual void AddSection(UMovieSceneSection& Section) override;
+	MOVIESCENETRACKS_API virtual bool SupportsType(TSubclassOf<UMovieSceneSection> SectionClass) const override;
+	MOVIESCENETRACKS_API virtual UMovieSceneSection* CreateNewSection() override;
+	MOVIESCENETRACKS_API virtual void RemoveSection(UMovieSceneSection& Section) override;
+	MOVIESCENETRACKS_API virtual void RemoveSectionAt(int32 SectionIndex) override;
+	MOVIESCENETRACKS_API virtual bool SupportsMultipleRows() const override;
+	MOVIESCENETRACKS_API virtual bool PopulateEvaluationTree(TMovieSceneEvaluationTree<FMovieSceneTrackEvaluationData>& OutData) const override;
+	MOVIESCENETRACKS_API virtual int8 GetEvaluationFieldVersion() const override;
 	
 #if WITH_EDITOR
-	virtual EMovieSceneSectionMovedResult OnSectionMoved(UMovieSceneSection& Section, const FMovieSceneSectionMovedParams& Params) override;
+	MOVIESCENETRACKS_API virtual EMovieSceneSectionMovedResult OnSectionMoved(UMovieSceneSection& Section, const FMovieSceneSectionMovedParams& Params) override;
 #endif
 
 #if WITH_EDITORONLY_DATA
-	virtual FText GetDefaultDisplayName() const override;
+	MOVIESCENETRACKS_API virtual FText GetDefaultDisplayName() const override;
 #endif
 };
diff --git a/Engine/Source/Runtime/MovieSceneTracks/Public/Tracks/MovieSceneEnumTrack.h b/Engine/Source/Runtime/MovieSceneTracks/Public/Tracks/MovieSceneEnumTrack.h
--- a/Engine/Source/Runtime/MovieSceneTracks/Public/Tracks/MovieSceneEnumTrack.h
+++ b/Engine/Source/Runtime/MovieSceneTracks/Public/Tracks/MovieSceneEnumTrack.h
@@ -12,20 +12,20 @@ class UEnum;
 /**
  * Handles manipulation of byte properties in a movie scene
  */
-UCLASS()
-class MOVIESCENETRACKS_API UMovieSceneEnumTrack : public UMovieScenePropertyTrack
+UCLASS(MinimalAPI)
+class UMovieSceneEnumTrack : public UMovieScenePropertyTrack
 {
 	GENERATED_UCLASS_BODY()
 
 public:
 	/** UMovieSceneTrack interface */
-	virtual void PostLoad() override;
-	virtual bool SupportsType(TSubclassOf<UMovieSceneSection> SectionClass) const override;
-	virtual UMovieSceneSection* CreateNewSection() override;
+	MOVIESCENETRACKS_API virtual void PostLoad() override;
+	MOVIESCENETRACKS_API virtual bool SupportsType(TSubclassOf<UMovieSceneSection> SectionClass) const override;
+	MOVIESCENETRACKS_API virtual UMovieSceneSection* CreateNewSection() override;
 
-	void SetEnum(UEnum* Enum);
+	MOVIESCENETRACKS_API void SetEnum(UEnum* Enum);
 
-	class UEnum* GetEnum() const;
+	MOVIESCENETRACKS_API class UEnum* GetEnum() const;
 
 protected:
 	UPROPERTY()
diff --git a/Engine/Source/Runtime/MovieSceneTracks/Public/Tracks/MovieSceneIntegerTrack.h b/Engine/Source/Runtime/MovieSceneTracks/Public/Tracks/MovieSceneIntegerTrack.h
--- a/Engine/Source/Runtime/MovieSceneTracks/Public/Tracks/MovieSceneIntegerTrack.h
+++ b/Engine/Source/Runtime/MovieSceneTracks/Public/Tracks/MovieSceneIntegerTrack.h
@@ -11,13 +11,13 @@
 /**
  * Handles manipulation of integer properties in a movie scene
  */
-UCLASS()
-class MOVIESCENETRACKS_API UMovieSceneIntegerTrack : public UMovieScenePropertyTrack
+UCLASS(MinimalAPI)
+class UMovieSceneIntegerTrack : public UMovieScenePropertyTrack
 {
 	GENERATED_UCLASS_BODY()
 
 public:
 	/** UMovieSceneTrack interface */
-	virtual bool SupportsType(TSubclassOf<UMovieSceneSection> SectionClass) const override;
-	virtual UMovieSceneSection* CreateNewSection() override;
+	MOVIESCENETRACKS_API virtual bool SupportsType(TSubclassOf<UMovieSceneSection> SectionClass) const override;
+	MOVIESCENETRACKS_API virtual UMovieSceneSection* CreateNewSection() override;
 };
diff --git a/Engine/Source/Runtime/MovieSceneTracks/Public/Tracks/MovieSceneMaterialParameterCollectionTrack.h b/Engine/Source/Runtime/MovieSceneTracks/Public/Tracks/MovieSceneMaterialParameterCollectionTrack.h
--- a/Engine/Source/Runtime/MovieSceneTracks/Public/Tracks/MovieSceneMaterialParameterCollectionTrack.h
+++ b/Engine/Source/Runtime/MovieSceneTracks/Public/Tracks/MovieSceneMaterialParameterCollectionTrack.h
@@ -13,8 +13,8 @@ class UMaterialParameterCollection;
 /**
  * Handles manipulation of material parameter collections in a movie scene.
  */
-UCLASS()
-class MOVIESCENETRACKS_API UMovieSceneMaterialParameterCollectionTrack
+UCLASS(MinimalAPI)
+class UMovieSceneMaterialParameterCollectionTrack
 	: public UMovieSceneMaterialTrack
 	, public IMovieSceneEntityProvider
 	, public IMovieSceneParameterSectionExtender
@@ -27,19 +27,19 @@ public:
 	UPROPERTY(EditAnywhere, Category=General, DisplayName="Material Parameter Collection")
 	TObjectPtr<UMaterialParameterCollection> MPC;
 
-	UMovieSceneMaterialParameterCollectionTrack(const FObjectInitializer& ObjectInitializer);
+	MOVIESCENETRACKS_API UMovieSceneMaterialParameterCollectionTrack(const FObjectInitializer& ObjectInitializer);
 
-	virtual bool SupportsType(TSubclassOf<UMovieSceneSection> SectionClass) const override;
-	virtual UMovieSceneSection* CreateNewSection() override;
+	MOVIESCENETRACKS_API virtual bool SupportsType(TSubclassOf<UMovieSceneSection> SectionClass) const override;
+	MOVIESCENETRACKS_API virtual UMovieSceneSection* CreateNewSection() override;
 
 	/*~ IMovieSceneEntityProvider */
-	virtual void ImportEntityImpl(UMovieSceneEntitySystemLinker* EntityLinker, const FEntityImportParams& Params, FImportedEntity* OutImportedEntity) override;
-	virtual bool PopulateEvaluationFieldImpl(const TRange<FFrameNumber>& EffectiveRange, const FMovieSceneEvaluationFieldEntityMetaData& InMetaData, FMovieSceneEntityComponentFieldBuilder* OutFieldBuilder) override;
+	MOVIESCENETRACKS_API virtual void ImportEntityImpl(UMovieSceneEntitySystemLinker* EntityLinker, const FEntityImportParams& Params, FImportedEntity* OutImportedEntity) override;
+	MOVIESCENETRACKS_API virtual bool PopulateEvaluationFieldImpl(const TRange<FFrameNumber>& EffectiveRange, const FMovieSceneEvaluationFieldEntityMetaData& InMetaData, FMovieSceneEntityComponentFieldBuilder* OutFieldBuilder) override;
 
 	/*~ IMovieSceneParameterSectionExtender */
-	virtual void ExtendEntityImpl(UMovieSceneParameterSection* Section, UMovieSceneEntitySystemLinker* EntityLinker, const UE::MovieScene::FEntityImportParams& Params, UE::MovieScene::FImportedEntity* OutImportedEntity) override;
+	MOVIESCENETRACKS_API virtual void ExtendEntityImpl(UMovieSceneParameterSection* Section, UMovieSceneEntitySystemLinker* EntityLinker, const UE::MovieScene::FEntityImportParams& Params, UE::MovieScene::FImportedEntity* OutImportedEntity) override;
 
 #if WITH_EDITORONLY_DATA
-	virtual FText GetDefaultDisplayName() const override;
+	MOVIESCENETRACKS_API virtual FText GetDefaultDisplayName() const override;
 #endif
 };
diff --git a/Engine/Source/Runtime/MovieSceneTracks/Public/Tracks/MovieSceneMaterialTrack.h b/Engine/Source/Runtime/MovieSceneTracks/Public/Tracks/MovieSceneMaterialTrack.h
--- a/Engine/Source/Runtime/MovieSceneTracks/Public/Tracks/MovieSceneMaterialTrack.h
+++ b/Engine/Source/Runtime/MovieSceneTracks/Public/Tracks/MovieSceneMaterialTrack.h
@@ -12,30 +12,30 @@
 /**
  * Handles manipulation of material parameters in a movie scene.
  */
-UCLASS()
-class MOVIESCENETRACKS_API UMovieSceneMaterialTrack
+UCLASS(MinimalAPI)
+class UMovieSceneMaterialTrack
 	: public UMovieSceneNameableTrack
 {
 	GENERATED_BODY()
 
 public:
 
-	UMovieSceneMaterialTrack(const FObjectInitializer& ObjectInitializer);
+	MOVIESCENETRACKS_API UMovieSceneMaterialTrack(const FObjectInitializer& ObjectInitializer);
 
 	// UMovieSceneTrack interface
 
-	virtual bool SupportsType(TSubclassOf<UMovieSceneSection> SectionClass) const override;
-	virtual UMovieSceneSection* CreateNewSection() override;
-	virtual void RemoveAllAnimationData() override;
-	virtual bool HasSection(const UMovieSceneSection& Section) const override;
-	virtual void AddSection(UMovieSceneSection& Section) override;
-	virtual void RemoveSection(UMovieSceneSection& Section) override;
-	virtual void RemoveSectionAt(int32 SectionIndex) override;
-	virtual bool IsEmpty() const override;
-	virtual const TArray<UMovieSceneSection*>& GetAllSections() const override;	
-	virtual bool SupportsMultipleRows() const override;
-	virtual void SetSectionToKey(UMovieSceneSection* Section) override;
-	virtual UMovieSceneSection* GetSectionToKey() const override;
+	MOVIESCENETRACKS_API virtual bool SupportsType(TSubclassOf<UMovieSceneSection> SectionClass) const override;
+	MOVIESCENETRACKS_API virtual UMovieSceneSection* CreateNewSection() override;
+	MOVIESCENETRACKS_API virtual void RemoveAllAnimationData() override;
+	MOVIESCENETRACKS_API virtual bool HasSection(const UMovieSceneSection& Section) const override;
+	MOVIESCENETRACKS_API virtual void AddSection(UMovieSceneSection& Section) override;
+	MOVIESCENETRACKS_API virtual void RemoveSection(UMovieSceneSection& Section) override;
+	MOVIESCENETRACKS_API virtual void RemoveSectionAt(int32 SectionIndex) override;
+	MOVIESCENETRACKS_API virtual bool IsEmpty() const override;
+	MOVIESCENETRACKS_API virtual const TArray<UMovieSceneSection*>& GetAllSections() const override;	
+	MOVIESCENETRACKS_API virtual bool SupportsMultipleRows() const override;
+	MOVIESCENETRACKS_API virtual void SetSectionToKey(UMovieSceneSection* Section) override;
+	MOVIESCENETRACKS_API virtual UMovieSceneSection* GetSectionToKey() const override;
 
 public:
 
@@ -45,7 +45,7 @@ public:
 	 * @param Time The time to add the new key.
 	 * @param The value for the new key.
 	 */
-	void AddScalarParameterKey(FName ParameterName, FFrameNumber Position, float Value);
+	MOVIESCENETRACKS_API void AddScalarParameterKey(FName ParameterName, FFrameNumber Position, float Value);
 
 	/**
 	 * Adds a scalar parameter key to the track. 
@@ -54,7 +54,7 @@ public:
 	 * @param RowIndex The preferred row index on which to look for sections.
 	 * @param The value for the new key.
 	 */
-	void AddScalarParameterKey(FName ParameterName, FFrameNumber Position, int32 RowIndex, float Value);
+	MOVIESCENETRACKS_API void AddScalarParameterKey(FName ParameterName, FFrameNumber Position, int32 RowIndex, float Value);
 
 	/**
 	 * Adds a color parameter key to the track.
@@ -62,7 +62,7 @@ public:
 	 * @param Time The time to add the new key.
 	 * @param The value for the new key.
 	 */
-	void AddColorParameterKey(FName ParameterName, FFrameNumber Position, FLinearColor Value);
+	MOVIESCENETRACKS_API void AddColorParameterKey(FName ParameterName, FFrameNumber Position, FLinearColor Value);
 
 	/**
 	 * Adds a color parameter key to the track.
@@ -71,7 +71,7 @@ public:
 	 * @param RowIndex The preferred row index on which to look for sections.
 	 * @param The value for the new key.
 	 */
-	void AddColorParameterKey(FName ParameterName, FFrameNumber Position, int32 RowIndex, FLinearColor Value);
+	MOVIESCENETRACKS_API void AddColorParameterKey(FName ParameterName, FFrameNumber Position, int32 RowIndex, FLinearColor Value);
 
 private:
 
diff --git a/Engine/Source/Runtime/MovieSceneTracks/Public/Tracks/MovieScenePropertyTrack.h b/Engine/Source/Runtime/MovieSceneTracks/Public/Tracks/MovieScenePropertyTrack.h
--- a/Engine/Source/Runtime/MovieSceneTracks/Public/Tracks/MovieScenePropertyTrack.h
+++ b/Engine/Source/Runtime/MovieSceneTracks/Public/Tracks/MovieScenePropertyTrack.h
@@ -19,8 +19,8 @@
 /**
  * Base class for tracks that animate an object property
  */
-UCLASS(abstract)
-class MOVIESCENETRACKS_API UMovieScenePropertyTrack
+UCLASS(abstract, MinimalAPI)
+class UMovieScenePropertyTrack
 	: public UMovieSceneNameableTrack
 {
 	GENERATED_UCLASS_BODY()
@@ -29,23 +29,23 @@ public:
 
 	// UMovieSceneTrack interface
 
-	virtual void RemoveAllAnimationData() override;
-	virtual bool HasSection(const UMovieSceneSection& Section) const override;
-	virtual void AddSection(UMovieSceneSection& Section) override;
-	virtual void RemoveSection(UMovieSceneSection& Section) override;
-	virtual void RemoveSectionAt(int32 SectionIndex) override;
-	virtual bool IsEmpty() const override;
-	virtual const TArray<UMovieSceneSection*>& GetAllSections() const override;
+	MOVIESCENETRACKS_API virtual void RemoveAllAnimationData() override;
+	MOVIESCENETRACKS_API virtual bool HasSection(const UMovieSceneSection& Section) const override;
+	MOVIESCENETRACKS_API virtual void AddSection(UMovieSceneSection& Section) override;
+	MOVIESCENETRACKS_API virtual void RemoveSection(UMovieSceneSection& Section) override;
+	MOVIESCENETRACKS_API virtual void RemoveSectionAt(int32 SectionIndex) override;
+	MOVIESCENETRACKS_API virtual bool IsEmpty() const override;
+	MOVIESCENETRACKS_API virtual const TArray<UMovieSceneSection*>& GetAllSections() const override;
 
 #if WITH_EDITORONLY_DATA
-	virtual FText GetDefaultDisplayName() const override;
-	virtual FText GetDisplayNameToolTipText() const override;
+	MOVIESCENETRACKS_API virtual FText GetDefaultDisplayName() const override;
+	MOVIESCENETRACKS_API virtual FText GetDisplayNameToolTipText() const override;
 	virtual bool CanRename() const override { return false; }
-	virtual FName GetTrackName() const override;
+	MOVIESCENETRACKS_API virtual FName GetTrackName() const override;
 #endif
 
-	virtual void PostLoad() override;
-	virtual void Serialize(FArchive& Ar) override;
+	MOVIESCENETRACKS_API virtual void PostLoad() override;
+	MOVIESCENETRACKS_API virtual void Serialize(FArchive& Ar) override;
 
 public:
 
@@ -54,7 +54,7 @@ public:
 	 *
 	 * @param InPropertyName The property being animated
 	 */
-	void SetPropertyNameAndPath(FName InPropertyName, const FString& InPropertyPath);
+	MOVIESCENETRACKS_API void SetPropertyNameAndPath(FName InPropertyName, const FString& InPropertyPath);
 
 	/** @return the name of the property being animated by this track */
 	FName GetPropertyName() const { return PropertyBinding.PropertyName; }
@@ -77,7 +77,7 @@ public:
 	*@param Time  The Time relative to the owning movie scene where the section should be
 	*@Return All sections at that time
 	*/
-	TArray<UMovieSceneSection*, TInlineAllocator<4>> FindAllSections(FFrameNumber Time);
+	MOVIESCENETRACKS_API TArray<UMovieSceneSection*, TInlineAllocator<4>> FindAllSections(FFrameNumber Time);
 
 	/**
 	 * Finds a section at the current time.
@@ -85,7 +85,7 @@ public:
 	 * @param Time The time relative to the owning movie scene where the section should be
 	 * @return The found section.
 	 */
-	class UMovieSceneSection* FindSection(FFrameNumber Time);
+	MOVIESCENETRACKS_API class UMovieSceneSection* FindSection(FFrameNumber Time);
 
 	/**
 	 * Finds a section at the current time or extends an existing one
@@ -94,7 +94,7 @@ public:
 	 * @param OutWeight The weight of the section if found
 	 * @return The found section.
 	 */
-	class UMovieSceneSection* FindOrExtendSection(FFrameNumber Time, float& OutWeight);
+	MOVIESCENETRACKS_API class UMovieSceneSection* FindOrExtendSection(FFrameNumber Time, float& OutWeight);
 
 	/**
 	 * Finds a section at the current time.
@@ -103,20 +103,20 @@ public:
 	 * @param bSectionAdded Whether a section was added or not
 	 * @return The found section, or the new section.
 	 */
-	class UMovieSceneSection* FindOrAddSection(FFrameNumber Time, bool& bSectionAdded);
+	MOVIESCENETRACKS_API class UMovieSceneSection* FindOrAddSection(FFrameNumber Time, bool& bSectionAdded);
 
 	/**
 	 * Set the section we want to key and recieve globally changed values.
 	 *
 	 * @param Section The section that changes.
 	 */
-	virtual void SetSectionToKey(UMovieSceneSection* Section) override;
+	MOVIESCENETRACKS_API virtual void SetSectionToKey(UMovieSceneSection* Section) override;
 
 	/**
 	 * Finds a section we want to key and recieve globally changed values.
 	 * @return The Section that changes.
 	 */
-	virtual UMovieSceneSection* GetSectionToKey() const override;
+	MOVIESCENETRACKS_API virtual UMovieSceneSection* GetSectionToKey() const override;
 
 #if WITH_EDITORONLY_DATA
 public:
@@ -151,18 +151,18 @@ protected:
 };
 
 
-struct MOVIESCENETRACKS_API FMovieScenePropertyTrackEntityImportHelper
+struct FMovieScenePropertyTrackEntityImportHelper
 {
-	static const int32 SectionPropertyValueImportingID;
-	static const int32 SectionEditConditionToggleImportingID;
+	static MOVIESCENETRACKS_API const int32 SectionPropertyValueImportingID;
+	static MOVIESCENETRACKS_API const int32 SectionEditConditionToggleImportingID;
 
-	static void PopulateEvaluationField(UMovieSceneSection& Section, const TRange<FFrameNumber>& EffectiveRange, const FMovieSceneEvaluationFieldEntityMetaData& InMetaData, FMovieSceneEntityComponentFieldBuilder* OutFieldBuilder);
+	static MOVIESCENETRACKS_API void PopulateEvaluationField(UMovieSceneSection& Section, const TRange<FFrameNumber>& EffectiveRange, const FMovieSceneEvaluationFieldEntityMetaData& InMetaData, FMovieSceneEntityComponentFieldBuilder* OutFieldBuilder);
 
-	static bool IsPropertyValueID(const UE::MovieScene::FEntityImportParams& Params);
-	static bool IsEditConditionToggleID(const UE::MovieScene::FEntityImportParams& Params);
-	static void ImportEditConditionToggleEntity(const UE::MovieScene::FEntityImportParams& Params, UE::MovieScene::FImportedEntity* OutImportedEntity);
+	static MOVIESCENETRACKS_API bool IsPropertyValueID(const UE::MovieScene::FEntityImportParams& Params);
+	static MOVIESCENETRACKS_API bool IsEditConditionToggleID(const UE::MovieScene::FEntityImportParams& Params);
+	static MOVIESCENETRACKS_API void ImportEditConditionToggleEntity(const UE::MovieScene::FEntityImportParams& Params, UE::MovieScene::FImportedEntity* OutImportedEntity);
 
-	static FName SanitizeBoolPropertyName(FName InPropertyName);
+	static MOVIESCENETRACKS_API FName SanitizeBoolPropertyName(FName InPropertyName);
 };
 
 
@@ -263,7 +263,7 @@ protected:
 	FComponentTypeID PropertyTag;
 };
 
-struct MOVIESCENETRACKS_API FPropertyTrackEntityImportHelper : TPropertyTrackEntityImportHelperImpl<>
+struct FPropertyTrackEntityImportHelper : TPropertyTrackEntityImportHelperImpl<>
 {
 	template<typename PropertyTraits>
 	FPropertyTrackEntityImportHelper(const TPropertyComponents<PropertyTraits>& PropertyComponents)
@@ -426,7 +426,7 @@ protected:
 	FComponentTypeID PropertyTag;
 };
 
-struct MOVIESCENETRACKS_API FPropertyTrackWithOverridableChannelsEntityImportHelper : TPropertyTrackWithOverridableChannelsEntityImportHelperImpl<>
+struct FPropertyTrackWithOverridableChannelsEntityImportHelper : TPropertyTrackWithOverridableChannelsEntityImportHelperImpl<>
 {
 	template<typename PropertyTraits>
 	FPropertyTrackWithOverridableChannelsEntityImportHelper(const TPropertyComponents<PropertyTraits>& PropertyComponents, TScriptInterface<IMovieSceneChannelOverrideProvider> InRegistryProvider)
diff --git a/Engine/Source/Runtime/MovieSceneTracks/Public/Tracks/MovieSceneSkeletalAnimationTrack.h b/Engine/Source/Runtime/MovieSceneTracks/Public/Tracks/MovieSceneSkeletalAnimationTrack.h
--- a/Engine/Source/Runtime/MovieSceneTracks/Public/Tracks/MovieSceneSkeletalAnimationTrack.h
+++ b/Engine/Source/Runtime/MovieSceneTracks/Public/Tracks/MovieSceneSkeletalAnimationTrack.h
@@ -49,7 +49,7 @@ public:
 	MOVIESCENETRACKS_API virtual UMovieSceneSection* AddNewAnimationOnRow(FFrameNumber KeyTime, class UAnimSequenceBase* AnimSequence, int32 RowIndex);
 
 	/** Adds a new animation to this track on the next available/non-overlapping row */
-	MOVIESCENETRACKS_API virtual UMovieSceneSection* AddNewAnimation(FFrameNumber KeyTime, class UAnimSequenceBase* AnimSequence) { return AddNewAnimationOnRow(KeyTime, AnimSequence, INDEX_NONE); }
+	virtual UMovieSceneSection* AddNewAnimation(FFrameNumber KeyTime, class UAnimSequenceBase* AnimSequence) { return AddNewAnimationOnRow(KeyTime, AnimSequence, INDEX_NONE); }
 
 	/** Gets the animation sections at a certain time */
 	MOVIESCENETRACKS_API TArray<UMovieSceneSection*> GetAnimSectionsAtTime(FFrameNumber Time);
diff --git a/Engine/Source/Runtime/NavigationSystem/Public/AbstractNavData.h b/Engine/Source/Runtime/NavigationSystem/Public/AbstractNavData.h
--- a/Engine/Source/Runtime/NavigationSystem/Public/AbstractNavData.h
+++ b/Engine/Source/Runtime/NavigationSystem/Public/AbstractNavData.h
@@ -23,16 +23,16 @@ class UClass;
 class UObject;
 struct FPathFindingQuery;
 
-struct NAVIGATIONSYSTEM_API FAbstractNavigationPath : public FNavigationPath
+struct FAbstractNavigationPath : public FNavigationPath
 {
 	typedef FNavigationPath Super;
 
-	FAbstractNavigationPath();
+	NAVIGATIONSYSTEM_API FAbstractNavigationPath();
 
-	static const FNavPathType Type;
+	static NAVIGATIONSYSTEM_API const FNavPathType Type;
 };
 
-class NAVIGATIONSYSTEM_API FAbstractQueryFilter : public INavigationQueryFilterInterface
+class FAbstractQueryFilter : public INavigationQueryFilterInterface
 {
 public:
 	virtual void Reset() override {}
@@ -50,18 +50,18 @@ public:
 	virtual void SetExcludeFlags(uint16 Flags) override {}
 	virtual uint16 GetExcludeFlags() const override { return 0; }
 	virtual FVector GetAdjustedEndLocation(const FVector& EndLocation) const override { return EndLocation; }
-	virtual INavigationQueryFilterInterface* CreateCopy() const override;
+	NAVIGATIONSYSTEM_API virtual INavigationQueryFilterInterface* CreateCopy() const override;
 };
 
-UCLASS()
-class NAVIGATIONSYSTEM_API AAbstractNavData : public ANavigationData
+UCLASS(MinimalAPI)
+class AAbstractNavData : public ANavigationData
 {
 	GENERATED_BODY()
 
 public:
-	AAbstractNavData(const FObjectInitializer& ObjectInitializer = FObjectInitializer::Get());
+	NAVIGATIONSYSTEM_API AAbstractNavData(const FObjectInitializer& ObjectInitializer = FObjectInitializer::Get());
 
-	virtual void PostLoad() override;
+	NAVIGATIONSYSTEM_API virtual void PostLoad() override;
 
 #if WITH_EDITOR
 	// Begin AActor overrides
@@ -90,7 +90,7 @@ public:
 	virtual bool IsNodeRefValid(NavNodeRef NodeRef) const override { return true; };
 	// End ANavigationData overrides
 
-	static FPathFindingResult FindPathAbstract(const FNavAgentProperties& AgentProperties, const FPathFindingQuery& Query);
-	static bool TestPathAbstract(const FNavAgentProperties& AgentProperties, const FPathFindingQuery& Query, int32* NumVisitedNodes);
-	static bool RaycastAbstract(const ANavigationData* NavDataInstance, const FVector& RayStart, const FVector& RayEnd, FVector& HitLocation, FSharedConstNavQueryFilter QueryFilter, const UObject* Querier);
+	static NAVIGATIONSYSTEM_API FPathFindingResult FindPathAbstract(const FNavAgentProperties& AgentProperties, const FPathFindingQuery& Query);
+	static NAVIGATIONSYSTEM_API bool TestPathAbstract(const FNavAgentProperties& AgentProperties, const FPathFindingQuery& Query, int32* NumVisitedNodes);
+	static NAVIGATIONSYSTEM_API bool RaycastAbstract(const ANavigationData* NavDataInstance, const FVector& RayStart, const FVector& RayEnd, FVector& HitLocation, FSharedConstNavQueryFilter QueryFilter, const UObject* Querier);
 };
diff --git a/Engine/Source/Runtime/NavigationSystem/Public/CrowdManagerBase.h b/Engine/Source/Runtime/NavigationSystem/Public/CrowdManagerBase.h
--- a/Engine/Source/Runtime/NavigationSystem/Public/CrowdManagerBase.h
+++ b/Engine/Source/Runtime/NavigationSystem/Public/CrowdManagerBase.h
@@ -9,19 +9,19 @@ class ANavigationData;
 /** Base class for Crowd Managers. If you want to create a custom crowd manager
  *	implement a class extending this one and set UNavigationSystemV1::CrowdManagerClass
  *	to point at your class */
-UCLASS(Abstract, Transient)
-class NAVIGATIONSYSTEM_API UCrowdManagerBase : public UObject
+UCLASS(Abstract, Transient, MinimalAPI)
+class UCrowdManagerBase : public UObject
 {
 	GENERATED_BODY()
 public:
-	virtual void Tick(float DeltaTime) PURE_VIRTUAL(UCrowdManagerBase::Tick, );
+	NAVIGATIONSYSTEM_API virtual void Tick(float DeltaTime) PURE_VIRTUAL(UCrowdManagerBase::Tick, );
 
 	/** Called by the nav system when a new navigation data instance is registered. */
-	virtual void OnNavDataRegistered(ANavigationData& NavDataInstance) PURE_VIRTUAL(UCrowdManagerBase::OnNavDataRegistered, );
+	NAVIGATIONSYSTEM_API virtual void OnNavDataRegistered(ANavigationData& NavDataInstance) PURE_VIRTUAL(UCrowdManagerBase::OnNavDataRegistered, );
 
 	/** Called by the nav system when a navigation data instance is removed. */
-	virtual void OnNavDataUnregistered(ANavigationData& NavDataInstance) PURE_VIRTUAL(UCrowdManagerBase::OnNavDataUnregistered, );
+	NAVIGATIONSYSTEM_API virtual void OnNavDataUnregistered(ANavigationData& NavDataInstance) PURE_VIRTUAL(UCrowdManagerBase::OnNavDataUnregistered, );
 
-	virtual void CleanUp(float DeltaTime) PURE_VIRTUAL(UCrowdManagerBase::CleanUp, );
+	NAVIGATIONSYSTEM_API virtual void CleanUp(float DeltaTime) PURE_VIRTUAL(UCrowdManagerBase::CleanUp, );
 };
 
diff --git a/Engine/Source/Runtime/NavigationSystem/Public/NavAreas/NavArea.h b/Engine/Source/Runtime/NavigationSystem/Public/NavAreas/NavArea.h
--- a/Engine/Source/Runtime/NavigationSystem/Public/NavAreas/NavArea.h
+++ b/Engine/Source/Runtime/NavigationSystem/Public/NavAreas/NavArea.h
@@ -11,8 +11,8 @@
 #include "NavArea.generated.h"
 
 /** Class containing definition of a navigation area */
-UCLASS(DefaultToInstanced, abstract, Config=Engine, Blueprintable)
-class NAVIGATIONSYSTEM_API UNavArea : public UNavAreaBase
+UCLASS(DefaultToInstanced, abstract, Config=Engine, Blueprintable, MinimalAPI)
+class UNavArea : public UNavAreaBase
 {
 	GENERATED_BODY()
 
@@ -80,14 +80,14 @@ public:
 	};
 #endif
 
-	UNavArea(const FObjectInitializer& ObjectInitializer = FObjectInitializer::Get());
+	NAVIGATIONSYSTEM_API UNavArea(const FObjectInitializer& ObjectInitializer = FObjectInitializer::Get());
 
-	virtual void FinishDestroy() override;
-	virtual void PostLoad() override;
-	virtual void PostInitProperties() override;
-	virtual void Serialize(FArchive& Ar) override;
+	NAVIGATIONSYSTEM_API virtual void FinishDestroy() override;
+	NAVIGATIONSYSTEM_API virtual void PostLoad() override;
+	NAVIGATIONSYSTEM_API virtual void PostInitProperties() override;
+	NAVIGATIONSYSTEM_API virtual void Serialize(FArchive& Ar) override;
 #if WITH_EDITOR
-	virtual void PostEditChangeProperty(FPropertyChangedEvent& PropertyChangedEvent);
+	NAVIGATIONSYSTEM_API virtual void PostEditChangeProperty(FPropertyChangedEvent& PropertyChangedEvent);
 #endif // WITH_EDITOR
 
 	FORCEINLINE uint16 GetAreaFlags() const { return AreaFlags; }
@@ -102,15 +102,15 @@ public:
 	virtual float GetFixedAreaEnteringCost() { return FixedAreaEnteringCost; }
 
 	/** Retrieved color declared for AreaDefinitionClass */
-	static FColor GetColor(UClass* AreaDefinitionClass);
+	static NAVIGATIONSYSTEM_API FColor GetColor(UClass* AreaDefinitionClass);
 
 	/** copy properties from other area */
-	virtual void CopyFrom(TSubclassOf<UNavArea> AreaClass);
+	NAVIGATIONSYSTEM_API virtual void CopyFrom(TSubclassOf<UNavArea> AreaClass);
 
 protected:
 
 	/** these flags will be applied to navigation data along with AreaID */
 	uint16 AreaFlags;
 	
-	void RegisterArea();
+	NAVIGATIONSYSTEM_API void RegisterArea();
 };
diff --git a/Engine/Source/Runtime/NavigationSystem/Public/NavAreas/NavAreaMeta.h b/Engine/Source/Runtime/NavigationSystem/Public/NavAreas/NavAreaMeta.h
--- a/Engine/Source/Runtime/NavigationSystem/Public/NavAreas/NavAreaMeta.h
+++ b/Engine/Source/Runtime/NavigationSystem/Public/NavAreas/NavAreaMeta.h
@@ -16,11 +16,11 @@ class UObject;
 /** A convenience class for an area that has IsMetaArea() == true.
  *	Do not use this class when determining whether an area class is "meta". 
  *	Call IsMetaArea instead. */
-UCLASS(Abstract)
-class NAVIGATIONSYSTEM_API UNavAreaMeta : public UNavArea
+UCLASS(Abstract, MinimalAPI)
+class UNavAreaMeta : public UNavArea
 {
 	GENERATED_BODY()
 
 public:
-	UNavAreaMeta(const FObjectInitializer& ObjectInitializer);
+	NAVIGATIONSYSTEM_API UNavAreaMeta(const FObjectInitializer& ObjectInitializer);
 };
diff --git a/Engine/Source/Runtime/NavigationSystem/Public/NavAreas/NavAreaMeta_SwitchByAgent.h b/Engine/Source/Runtime/NavigationSystem/Public/NavAreas/NavAreaMeta_SwitchByAgent.h
--- a/Engine/Source/Runtime/NavigationSystem/Public/NavAreas/NavAreaMeta_SwitchByAgent.h
+++ b/Engine/Source/Runtime/NavigationSystem/Public/NavAreas/NavAreaMeta_SwitchByAgent.h
@@ -11,8 +11,8 @@
 class AActor;
 
 /** Class containing definition of a navigation area */
-UCLASS(Abstract)
-class NAVIGATIONSYSTEM_API UNavAreaMeta_SwitchByAgent : public UNavAreaMeta
+UCLASS(Abstract, MinimalAPI)
+class UNavAreaMeta_SwitchByAgent : public UNavAreaMeta
 {
 	GENERATED_UCLASS_BODY()
 
@@ -64,10 +64,10 @@ class NAVIGATIONSYSTEM_API UNavAreaMeta_SwitchByAgent : public UNavAreaMeta
 	UPROPERTY(EditAnywhere, Category=AgentTypes)
 	TSubclassOf<UNavArea> Agent15Area;
 
-	virtual TSubclassOf<UNavAreaBase> PickAreaClassForAgent(const AActor& Actor, const FNavAgentProperties& NavAgent) const override;
+	NAVIGATIONSYSTEM_API virtual TSubclassOf<UNavAreaBase> PickAreaClassForAgent(const AActor& Actor, const FNavAgentProperties& NavAgent) const override;
 
 #if WITH_EDITOR
 	/** setup AgentXArea properties */
-	virtual void UpdateAgentConfig();
+	NAVIGATIONSYSTEM_API virtual void UpdateAgentConfig();
 #endif
 };
diff --git a/Engine/Source/Runtime/NavigationSystem/Public/NavAreas/NavArea_Default.h b/Engine/Source/Runtime/NavigationSystem/Public/NavAreas/NavArea_Default.h
--- a/Engine/Source/Runtime/NavigationSystem/Public/NavAreas/NavArea_Default.h
+++ b/Engine/Source/Runtime/NavigationSystem/Public/NavAreas/NavArea_Default.h
@@ -12,8 +12,8 @@
 class UObject;
 
 /** Regular navigation area, applied to entire navigation data by default */
-UCLASS(Config=Engine)
-class NAVIGATIONSYSTEM_API UNavArea_Default : public UNavArea
+UCLASS(Config=Engine, MinimalAPI)
+class UNavArea_Default : public UNavArea
 {
 	GENERATED_UCLASS_BODY()
 };
diff --git a/Engine/Source/Runtime/NavigationSystem/Public/NavAreas/NavArea_LowHeight.h b/Engine/Source/Runtime/NavigationSystem/Public/NavAreas/NavArea_LowHeight.h
--- a/Engine/Source/Runtime/NavigationSystem/Public/NavAreas/NavArea_LowHeight.h
+++ b/Engine/Source/Runtime/NavigationSystem/Public/NavAreas/NavArea_LowHeight.h
@@ -13,8 +13,8 @@
 class UObject;
 
 /** Special area that can be generated in spaces with insufficient free height above. Cannot be traversed by anyone. */
-UCLASS(Config = Engine)
-class NAVIGATIONSYSTEM_API UNavArea_LowHeight : public UNavArea
+UCLASS(Config = Engine, MinimalAPI)
+class UNavArea_LowHeight : public UNavArea
 {
 	GENERATED_UCLASS_BODY()
 public:
diff --git a/Engine/Source/Runtime/NavigationSystem/Public/NavAreas/NavArea_Null.h b/Engine/Source/Runtime/NavigationSystem/Public/NavAreas/NavArea_Null.h
--- a/Engine/Source/Runtime/NavigationSystem/Public/NavAreas/NavArea_Null.h
+++ b/Engine/Source/Runtime/NavigationSystem/Public/NavAreas/NavArea_Null.h
@@ -12,8 +12,8 @@
 class UObject;
 
 /** In general represents an empty area, that cannot be traversed by anyone. Ever.*/
-UCLASS(Config=Engine)
-class NAVIGATIONSYSTEM_API UNavArea_Null : public UNavArea
+UCLASS(Config=Engine, MinimalAPI)
+class UNavArea_Null : public UNavArea
 {
 	GENERATED_UCLASS_BODY()
 };
diff --git a/Engine/Source/Runtime/NavigationSystem/Public/NavAreas/NavArea_Obstacle.h b/Engine/Source/Runtime/NavigationSystem/Public/NavAreas/NavArea_Obstacle.h
--- a/Engine/Source/Runtime/NavigationSystem/Public/NavAreas/NavArea_Obstacle.h
+++ b/Engine/Source/Runtime/NavigationSystem/Public/NavAreas/NavArea_Obstacle.h
@@ -12,8 +12,8 @@
 class UObject;
 
 /** In general represents a high cost area, that shouldn't be traversed by anyone unless no other path exist.*/
-UCLASS(Config = Engine)
-class NAVIGATIONSYSTEM_API UNavArea_Obstacle : public UNavArea
+UCLASS(Config = Engine, MinimalAPI)
+class UNavArea_Obstacle : public UNavArea
 {
 	GENERATED_UCLASS_BODY()
 };
diff --git a/Engine/Source/Runtime/NavigationSystem/Public/NavCollision.h b/Engine/Source/Runtime/NavigationSystem/Public/NavCollision.h
--- a/Engine/Source/Runtime/NavigationSystem/Public/NavCollision.h
+++ b/Engine/Source/Runtime/NavigationSystem/Public/NavCollision.h
@@ -43,8 +43,8 @@ struct FNavCollisionBox
 	FVector Extent = FVector::ZeroVector;
 };
 
-UCLASS(config=Engine)
-class NAVIGATIONSYSTEM_API UNavCollision : public UNavCollisionBase
+UCLASS(config=Engine, MinimalAPI)
+class UNavCollision : public UNavCollisionBase
 {
 	GENERATED_UCLASS_BODY()
 
@@ -82,44 +82,44 @@ class NAVIGATIONSYSTEM_API UNavCollision : public UNavCollisionBase
 	FFormatContainer CookedFormatData;
 
 	//~ Begin UObject Interface.
-	virtual void PostInitProperties() override;
-	virtual void Serialize(FArchive& Ar) override;
-	virtual void PostLoad() override;
-	virtual void GetResourceSizeEx(FResourceSizeEx& CumulativeResourceSize) override;
-	virtual bool NeedsLoadForTargetPlatform(const class ITargetPlatform* TargetPlatform) const override;
+	NAVIGATIONSYSTEM_API virtual void PostInitProperties() override;
+	NAVIGATIONSYSTEM_API virtual void Serialize(FArchive& Ar) override;
+	NAVIGATIONSYSTEM_API virtual void PostLoad() override;
+	NAVIGATIONSYSTEM_API virtual void GetResourceSizeEx(FResourceSizeEx& CumulativeResourceSize) override;
+	NAVIGATIONSYSTEM_API virtual bool NeedsLoadForTargetPlatform(const class ITargetPlatform* TargetPlatform) const override;
 	virtual bool NeedsLoadForClient() const override { return bCreateOnClient; }
 	//~ End UObject Interface.
 
-	FGuid GetGuid() const;
+	NAVIGATIONSYSTEM_API FGuid GetGuid() const;
 
 	/** Tries to read data from DDC, and if that fails gathers navigation
 	 *	collision data, stores it and uploads to DDC */
-	virtual void Setup(class UBodySetup* BodySetup) override;
+	NAVIGATIONSYSTEM_API virtual void Setup(class UBodySetup* BodySetup) override;
 
 	/** copy user settings from other nav collision data */
-	void CopyUserSettings(const UNavCollision& OtherData);
+	NAVIGATIONSYSTEM_API void CopyUserSettings(const UNavCollision& OtherData);
 
 	/** show cylinder and box collisions */
-	virtual void DrawSimpleGeom(FPrimitiveDrawInterface* PDI, const FTransform& Transform, const FColor Color) override;
+	NAVIGATIONSYSTEM_API virtual void DrawSimpleGeom(FPrimitiveDrawInterface* PDI, const FTransform& Transform, const FColor Color) override;
 
 	/** Get data for dynamic obstacle */
-	virtual void GetNavigationModifier(FCompositeNavModifier& Modifier, const FTransform& LocalToWorld) override;
+	NAVIGATIONSYSTEM_API virtual void GetNavigationModifier(FCompositeNavModifier& Modifier, const FTransform& LocalToWorld) override;
 
 	/** Export collision data */
-	virtual bool ExportGeometry(const FTransform& LocalToWorld, FNavigableGeometryExport& GeoExport) const override;
+	NAVIGATIONSYSTEM_API virtual bool ExportGeometry(const FTransform& LocalToWorld, FNavigableGeometryExport& GeoExport) const override;
 
 	/** Read collisions data */
-	void GatherCollision();
+	NAVIGATIONSYSTEM_API void GatherCollision();
 
 #if WITH_EDITOR
-	virtual void InvalidateCollision() override;
+	NAVIGATIONSYSTEM_API virtual void InvalidateCollision() override;
 #endif // WITH_EDITOR
 
 protected:
-	void ClearCollision();
+	NAVIGATIONSYSTEM_API void ClearCollision();
 
 #if WITH_EDITOR
-	void InvalidatePhysicsData();
+	NAVIGATIONSYSTEM_API void InvalidatePhysicsData();
 #endif // WITH_EDITOR
-	FByteBulkData* GetCookedData(FName Format);
+	NAVIGATIONSYSTEM_API FByteBulkData* GetCookedData(FName Format);
 };
diff --git a/Engine/Source/Runtime/NavigationSystem/Public/NavFilters/NavigationQueryFilter.h b/Engine/Source/Runtime/NavigationSystem/Public/NavFilters/NavigationQueryFilter.h
--- a/Engine/Source/Runtime/NavigationSystem/Public/NavFilters/NavigationQueryFilter.h
+++ b/Engine/Source/Runtime/NavigationSystem/Public/NavFilters/NavigationQueryFilter.h
@@ -13,7 +13,7 @@
 class ANavigationData;
 
 USTRUCT()
-struct NAVIGATIONSYSTEM_API FNavigationFilterArea
+struct FNavigationFilterArea
 {
 	GENERATED_USTRUCT_BODY()
 
@@ -50,7 +50,7 @@ struct NAVIGATIONSYSTEM_API FNavigationFilterArea
 // Use UNavigationSystemV1.DescribeFilterFlags() to setup user friendly names of flags
 // 
 USTRUCT()
-struct NAVIGATIONSYSTEM_API FNavigationFilterFlags
+struct FNavigationFilterFlags
 {
 	GENERATED_USTRUCT_BODY()
 
@@ -100,8 +100,8 @@ struct NAVIGATIONSYSTEM_API FNavigationFilterFlags
 };
 
 /** Class containing definition of a navigation query filter */
-UCLASS(Abstract, Blueprintable)
-class NAVIGATIONSYSTEM_API UNavigationQueryFilter : public UObject
+UCLASS(Abstract, Blueprintable, MinimalAPI)
+class UNavigationQueryFilter : public UObject
 {
 	GENERATED_UCLASS_BODY()
 	
@@ -118,11 +118,11 @@ class NAVIGATIONSYSTEM_API UNavigationQueryFilter : public UObject
 	FNavigationFilterFlags ExcludeFlags;
 
 	/** get filter for given navigation data and initialize on first access */
-	FSharedConstNavQueryFilter GetQueryFilter(const ANavigationData& NavData, const UObject* Querier) const;
+	NAVIGATIONSYSTEM_API FSharedConstNavQueryFilter GetQueryFilter(const ANavigationData& NavData, const UObject* Querier) const;
 	
 	/** helper functions for accessing filter */
-	static FSharedConstNavQueryFilter GetQueryFilter(const ANavigationData& NavData, TSubclassOf<UNavigationQueryFilter> FilterClass);
-	static FSharedConstNavQueryFilter GetQueryFilter(const ANavigationData& NavData, const UObject* Querier, TSubclassOf<UNavigationQueryFilter> FilterClass);
+	static NAVIGATIONSYSTEM_API FSharedConstNavQueryFilter GetQueryFilter(const ANavigationData& NavData, TSubclassOf<UNavigationQueryFilter> FilterClass);
+	static NAVIGATIONSYSTEM_API FSharedConstNavQueryFilter GetQueryFilter(const ANavigationData& NavData, const UObject* Querier, TSubclassOf<UNavigationQueryFilter> FilterClass);
 
 	template<class T>
 	static FSharedConstNavQueryFilter GetQueryFilter(const ANavigationData& NavData, TSubclassOf<UNavigationQueryFilter> FilterClass = T::StaticClass())
@@ -131,7 +131,7 @@ class NAVIGATIONSYSTEM_API UNavigationQueryFilter : public UObject
 	}
 
 #if WITH_EDITOR
-	virtual void PostEditChangeProperty(struct FPropertyChangedEvent& PropertyChangedEvent) override;
+	NAVIGATIONSYSTEM_API virtual void PostEditChangeProperty(struct FPropertyChangedEvent& PropertyChangedEvent) override;
 #endif
 
 protected:
@@ -143,15 +143,15 @@ protected:
 	uint32 bIsMetaFilter : 1;
 
 	/** helper functions for adding area overrides */
-	void AddTravelCostOverride(TSubclassOf<UNavArea> AreaClass, float TravelCost);
-	void AddEnteringCostOverride(TSubclassOf<UNavArea> AreaClass, float EnteringCost);
-	void AddExcludedArea(TSubclassOf<UNavArea> AreaClass);
+	NAVIGATIONSYSTEM_API void AddTravelCostOverride(TSubclassOf<UNavArea> AreaClass, float TravelCost);
+	NAVIGATIONSYSTEM_API void AddEnteringCostOverride(TSubclassOf<UNavArea> AreaClass, float EnteringCost);
+	NAVIGATIONSYSTEM_API void AddExcludedArea(TSubclassOf<UNavArea> AreaClass);
 
 	/** find index of area data */
-	int32 FindAreaOverride(TSubclassOf<UNavArea> AreaClass) const;
+	NAVIGATIONSYSTEM_API int32 FindAreaOverride(TSubclassOf<UNavArea> AreaClass) const;
 	
 	/** setup filter for given navigation data, use to create custom filters */
-	virtual void InitializeFilter(const ANavigationData& NavData, const UObject* Querier, FNavigationQueryFilter& Filter) const;
+	NAVIGATIONSYSTEM_API virtual void InitializeFilter(const ANavigationData& NavData, const UObject* Querier, FNavigationQueryFilter& Filter) const;
 
 	virtual TSubclassOf<UNavigationQueryFilter> GetSimpleFilterForAgent(const UObject& Querier) const { return nullptr; }
 };
diff --git a/Engine/Source/Runtime/NavigationSystem/Public/NavLinkComponent.h b/Engine/Source/Runtime/NavigationSystem/Public/NavLinkComponent.h
--- a/Engine/Source/Runtime/NavigationSystem/Public/NavLinkComponent.h
+++ b/Engine/Source/Runtime/NavigationSystem/Public/NavLinkComponent.h
@@ -13,32 +13,32 @@
 class FPrimitiveSceneProxy;
 struct FNavigationRelevantData;
 
-UCLASS(ClassGroup = (Navigation), meta = (BlueprintSpawnableComponent), hidecategories = (Activation))
-class NAVIGATIONSYSTEM_API UNavLinkComponent : public UPrimitiveComponent, public INavLinkHostInterface
+UCLASS(ClassGroup = (Navigation), meta = (BlueprintSpawnableComponent), hidecategories = (Activation), MinimalAPI)
+class UNavLinkComponent : public UPrimitiveComponent, public INavLinkHostInterface
 {
 	GENERATED_UCLASS_BODY()
 
 	UPROPERTY(EditAnywhere, BlueprintReadOnly, Category = Navigation)
 	TArray<FNavigationLink> Links;
 
-	virtual void GetNavigationData(FNavigationRelevantData& Data) const override;
-	virtual bool IsNavigationRelevant() const override;
+	NAVIGATIONSYSTEM_API virtual void GetNavigationData(FNavigationRelevantData& Data) const override;
+	NAVIGATIONSYSTEM_API virtual bool IsNavigationRelevant() const override;
 
 	virtual bool GetNavigationLinksClasses(TArray<TSubclassOf<class UNavLinkDefinition> >& OutClasses) const override { return false; }
-	virtual bool GetNavigationLinksArray(TArray<FNavigationLink>& OutLink, TArray<FNavigationSegmentLink>& OutSegments) const override;
+	NAVIGATIONSYSTEM_API virtual bool GetNavigationLinksArray(TArray<FNavigationLink>& OutLink, TArray<FNavigationSegmentLink>& OutSegments) const override;
 
-	virtual FBoxSphereBounds CalcBounds(const FTransform &LocalToWorld) const override;
-	virtual FPrimitiveSceneProxy* CreateSceneProxy() override;
+	NAVIGATIONSYSTEM_API virtual FBoxSphereBounds CalcBounds(const FTransform &LocalToWorld) const override;
+	NAVIGATIONSYSTEM_API virtual FPrimitiveSceneProxy* CreateSceneProxy() override;
 	virtual bool ShouldRecreateProxyOnUpdateTransform() const override { return true; }
 
 #if WITH_EDITOR
-	virtual void PostEditChangeProperty(FPropertyChangedEvent& PropertyChangedEvent) override;
-	virtual void PostEditUndo() override;
-	virtual void PostEditImport() override;
+	NAVIGATIONSYSTEM_API virtual void PostEditChangeProperty(FPropertyChangedEvent& PropertyChangedEvent) override;
+	NAVIGATIONSYSTEM_API virtual void PostEditUndo() override;
+	NAVIGATIONSYSTEM_API virtual void PostEditImport() override;
 #endif // WITH_EDITOR
 
-	virtual void OnRegister() override;
+	NAVIGATIONSYSTEM_API virtual void OnRegister() override;
 
 protected:
-	void InitializeLinksAreaClasses();
+	NAVIGATIONSYSTEM_API void InitializeLinksAreaClasses();
 };
diff --git a/Engine/Source/Runtime/NavigationSystem/Public/NavLinkCustomComponent.h b/Engine/Source/Runtime/NavigationSystem/Public/NavLinkCustomComponent.h
--- a/Engine/Source/Runtime/NavigationSystem/Public/NavLinkCustomComponent.h
+++ b/Engine/Source/Runtime/NavigationSystem/Public/NavLinkCustomComponent.h
@@ -25,8 +25,8 @@ struct FNavigationRelevantData;
  *  - can broadcast state changes to nearby agents
  */
 
-UCLASS()
-class NAVIGATIONSYSTEM_API UNavLinkCustomComponent : public UNavRelevantComponent, public INavLinkCustomInterface
+UCLASS(MinimalAPI)
+class UNavLinkCustomComponent : public UNavRelevantComponent, public INavLinkCustomInterface
 {
 	GENERATED_UCLASS_BODY()
 
@@ -34,48 +34,48 @@ class NAVIGATIONSYSTEM_API UNavLinkCustomComponent : public UNavRelevantComponen
 	DECLARE_DELEGATE_TwoParams(FBroadcastFilter, UNavLinkCustomComponent* /*ThisComp*/, TArray<UObject*>& /*NotifyList*/);
 
 	// BEGIN INavLinkCustomInterface
-	virtual void GetLinkData(FVector& LeftPt, FVector& RightPt, ENavLinkDirection::Type& Direction) const override;
-	virtual void GetSupportedAgents(FNavAgentSelector& OutSupportedAgents) const override;
-	virtual TSubclassOf<UNavArea> GetLinkAreaClass() const override;
-	virtual FNavLinkAuxiliaryId GetAuxiliaryId() const override;
-	virtual FNavLinkId GetId() const override;
-	virtual void UpdateLinkId(FNavLinkId NewUniqueId) override;
-	virtual bool IsLinkPathfindingAllowed(const UObject* Querier) const override;
-	virtual bool OnLinkMoveStarted(UObject* PathComp, const FVector& DestPoint) override;
-	virtual void OnLinkMoveFinished(UObject* PathComp) override;
+	NAVIGATIONSYSTEM_API virtual void GetLinkData(FVector& LeftPt, FVector& RightPt, ENavLinkDirection::Type& Direction) const override;
+	NAVIGATIONSYSTEM_API virtual void GetSupportedAgents(FNavAgentSelector& OutSupportedAgents) const override;
+	NAVIGATIONSYSTEM_API virtual TSubclassOf<UNavArea> GetLinkAreaClass() const override;
+	NAVIGATIONSYSTEM_API virtual FNavLinkAuxiliaryId GetAuxiliaryId() const override;
+	NAVIGATIONSYSTEM_API virtual FNavLinkId GetId() const override;
+	NAVIGATIONSYSTEM_API virtual void UpdateLinkId(FNavLinkId NewUniqueId) override;
+	NAVIGATIONSYSTEM_API virtual bool IsLinkPathfindingAllowed(const UObject* Querier) const override;
+	NAVIGATIONSYSTEM_API virtual bool OnLinkMoveStarted(UObject* PathComp, const FVector& DestPoint) override;
+	NAVIGATIONSYSTEM_API virtual void OnLinkMoveFinished(UObject* PathComp) override;
 	// END INavLinkCustomInterface
 
 	//~ Begin UNavRelevantComponent Interface
-	virtual void GetNavigationData(FNavigationRelevantData& Data) const override;
-	virtual void CalcAndCacheBounds() const override;
+	NAVIGATIONSYSTEM_API virtual void GetNavigationData(FNavigationRelevantData& Data) const override;
+	NAVIGATIONSYSTEM_API virtual void CalcAndCacheBounds() const override;
 	//~ End UNavRelevantComponent Interface
 
 	//~ Begin UActorComponent Interface
-	virtual void OnRegister() override;
-	virtual void OnUnregister() override;
-	virtual TStructOnScope<FActorComponentInstanceData> GetComponentInstanceData() const override;
+	NAVIGATIONSYSTEM_API virtual void OnRegister() override;
+	NAVIGATIONSYSTEM_API virtual void OnUnregister() override;
+	NAVIGATIONSYSTEM_API virtual TStructOnScope<FActorComponentInstanceData> GetComponentInstanceData() const override;
 	//~ End UActorComponent Interface
 
-	void ApplyComponentInstanceData(struct FNavLinkCustomInstanceData* ComponentInstanceData);
+	NAVIGATIONSYSTEM_API void ApplyComponentInstanceData(struct FNavLinkCustomInstanceData* ComponentInstanceData);
 
 	//~ Begin UObject Interface
-	virtual void Serialize(FArchive& Ar) override;
-	virtual void PostLoad() override;
+	NAVIGATIONSYSTEM_API virtual void Serialize(FArchive& Ar) override;
+	NAVIGATIONSYSTEM_API virtual void PostLoad() override;
 #if WITH_EDITOR
-	virtual void PostEditImport() override;
+	NAVIGATIONSYSTEM_API virtual void PostEditImport() override;
 #endif
 	//~ End UObject Interface
 
 	/** set basic link data: end points and direction */
-	void SetLinkData(const FVector& RelativeStart, const FVector& RelativeEnd, ENavLinkDirection::Type Direction);
-	virtual FNavigationLink GetLinkModifier() const;
+	NAVIGATIONSYSTEM_API void SetLinkData(const FVector& RelativeStart, const FVector& RelativeEnd, ENavLinkDirection::Type Direction);
+	NAVIGATIONSYSTEM_API virtual FNavigationLink GetLinkModifier() const;
 
 	/** set area class to use when link is enabled */
-	void SetEnabledArea(TSubclassOf<UNavArea> AreaClass);
+	NAVIGATIONSYSTEM_API void SetEnabledArea(TSubclassOf<UNavArea> AreaClass);
 	TSubclassOf<UNavArea> GetEnabledArea() const { return EnabledAreaClass; }
 
 	/** set area class to use when link is disabled */
-	void SetDisabledArea(TSubclassOf<UNavArea> AreaClass);
+	NAVIGATIONSYSTEM_API void SetDisabledArea(TSubclassOf<UNavArea> AreaClass);
 	TSubclassOf<UNavArea> GetDisabledArea() const { return DisabledAreaClass; }
 
 	void SetSupportedAgents(const FNavAgentSelector& InSupportedAgents) { SupportedAgents = InSupportedAgents; }
@@ -83,35 +83,35 @@ class NAVIGATIONSYSTEM_API UNavLinkCustomComponent : public UNavRelevantComponen
 
 	/** add box obstacle during generation of navigation data
 	  * this can be used to create empty area under doors */
-	void AddNavigationObstacle(TSubclassOf<UNavArea> AreaClass, const FVector& BoxExtent, const FVector& BoxOffset = FVector::ZeroVector);
+	NAVIGATIONSYSTEM_API void AddNavigationObstacle(TSubclassOf<UNavArea> AreaClass, const FVector& BoxExtent, const FVector& BoxOffset = FVector::ZeroVector);
 
 	/** removes simple obstacle */
-	void ClearNavigationObstacle();
+	NAVIGATIONSYSTEM_API void ClearNavigationObstacle();
 
 	/** set properties of trigger around link entry point(s), that will notify nearby agents about link state change */
-	void SetBroadcastData(float Radius, ECollisionChannel TraceChannel = ECC_Pawn, float Interval = 0.0f);
+	NAVIGATIONSYSTEM_API void SetBroadcastData(float Radius, ECollisionChannel TraceChannel = ECC_Pawn, float Interval = 0.0f);
 
-	void SendBroadcastWhenEnabled(bool bEnabled);
-	void SendBroadcastWhenDisabled(bool bEnabled);
+	NAVIGATIONSYSTEM_API void SendBroadcastWhenEnabled(bool bEnabled);
+	NAVIGATIONSYSTEM_API void SendBroadcastWhenDisabled(bool bEnabled);
 
 	/** set delegate to filter  */
-	void SetBroadcastFilter(FBroadcastFilter const& InDelegate);
+	NAVIGATIONSYSTEM_API void SetBroadcastFilter(FBroadcastFilter const& InDelegate);
 
 	/** change state of smart link (used area class) */
-	void SetEnabled(bool bNewEnabled);
+	NAVIGATIONSYSTEM_API void SetEnabled(bool bNewEnabled);
 	bool IsEnabled() const { return bLinkEnabled; }
 	
 	/** set delegate to notify about reaching this link during path following */
-	void SetMoveReachedLink(FOnMoveReachedLink const& InDelegate);
+	NAVIGATIONSYSTEM_API void SetMoveReachedLink(FOnMoveReachedLink const& InDelegate);
 
 	/** check is any agent is currently moving though this link */
-	bool HasMovingAgents() const;
+	NAVIGATIONSYSTEM_API bool HasMovingAgents() const;
 
 	/** get link start point in world space */
-	FVector GetStartPoint() const;
+	NAVIGATIONSYSTEM_API FVector GetStartPoint() const;
 
 	/** get link end point in world space */
-	FVector GetEndPoint() const;
+	NAVIGATIONSYSTEM_API FVector GetEndPoint() const;
 
 	TSubclassOf<UNavArea> GetObstacleAreaClass() const { return ObstacleAreaClass; }
 
@@ -142,8 +142,8 @@ class NAVIGATIONSYSTEM_API UNavLinkCustomComponent : public UNavRelevantComponen
 
 protected:
 #if WITH_EDITOR
-	void OnNavAreaRegistered(const UWorld& World, const UClass* NavAreaClass);
-	void OnNavAreaUnregistered(const UWorld& World, const UClass* NavAreaClass);
+	NAVIGATIONSYSTEM_API void OnNavAreaRegistered(const UWorld& World, const UClass* NavAreaClass);
+	NAVIGATIONSYSTEM_API void OnNavAreaUnregistered(const UWorld& World, const UClass* NavAreaClass);
 #endif // WITH_EDITOR
 
 protected:
@@ -242,10 +242,10 @@ protected:
 	FTimerHandle TimerHandle_BroadcastStateChange;
 
 	/** notify nearby agents about link changing state */
-	void BroadcastStateChange();
+	NAVIGATIONSYSTEM_API void BroadcastStateChange();
 	
 	/** gather agents to notify about state change */
-	void CollectNearbyAgents(TArray<UObject*>& NotifyList);
+	NAVIGATIONSYSTEM_API void CollectNearbyAgents(TArray<UObject*>& NotifyList);
 
 #if WITH_EDITOR
 	FDelegateHandle OnNavAreaRegisteredDelegateHandle;
@@ -287,4 +287,4 @@ public:
 
 	UPROPERTY()
 	FNavLinkAuxiliaryId AuxiliaryCustomLinkId;
-};
\ No newline at end of file
+};
diff --git a/Engine/Source/Runtime/NavigationSystem/Public/NavLinkHostInterface.h b/Engine/Source/Runtime/NavigationSystem/Public/NavLinkHostInterface.h
--- a/Engine/Source/Runtime/NavigationSystem/Public/NavLinkHostInterface.h
+++ b/Engine/Source/Runtime/NavigationSystem/Public/NavLinkHostInterface.h
@@ -15,14 +15,14 @@ class UNavLinkHostInterface : public UInterface
 	GENERATED_UINTERFACE_BODY()
 };
 
-class NAVIGATIONSYSTEM_API INavLinkHostInterface
+class INavLinkHostInterface
 {
 	GENERATED_IINTERFACE_BODY()
 		
 	/**
 	 *	Retrieves UNavLinkDefinition derived UClasses hosted by this interface implementer
 	 */
-	virtual bool GetNavigationLinksClasses(TArray<TSubclassOf<class UNavLinkDefinition> >& OutClasses) const PURE_VIRTUAL(INavLinkHostInterface::GetNavigationLinksClasses,return false;);
+	NAVIGATIONSYSTEM_API virtual bool GetNavigationLinksClasses(TArray<TSubclassOf<class UNavLinkDefinition> >& OutClasses) const PURE_VIRTUAL(INavLinkHostInterface::GetNavigationLinksClasses,return false;);
 
 	/** 
 	 *	_Optional_ way of retrieving navigation link data - if INavLinkHostInterface 
diff --git a/Engine/Source/Runtime/NavigationSystem/Public/NavLinkRenderingComponent.h b/Engine/Source/Runtime/NavigationSystem/Public/NavLinkRenderingComponent.h
--- a/Engine/Source/Runtime/NavigationSystem/Public/NavLinkRenderingComponent.h
+++ b/Engine/Source/Runtime/NavigationSystem/Public/NavLinkRenderingComponent.h
@@ -11,25 +11,25 @@ class FPrimitiveSceneProxy;
 struct FConvexVolume;
 struct FEngineShowFlags;
 
-UCLASS(hidecategories=Object, editinlinenew)
-class NAVIGATIONSYSTEM_API UNavLinkRenderingComponent : public UPrimitiveComponent
+UCLASS(hidecategories=Object, editinlinenew, MinimalAPI)
+class UNavLinkRenderingComponent : public UPrimitiveComponent
 {
 	GENERATED_UCLASS_BODY()
 		
 	//~ Begin UPrimitiveComponent Interface
-	virtual FPrimitiveSceneProxy* CreateSceneProxy() override;
+	NAVIGATIONSYSTEM_API virtual FPrimitiveSceneProxy* CreateSceneProxy() override;
 
 	/** Should recreate proxy one very update */
 	virtual bool ShouldRecreateProxyOnUpdateTransform() const override { return true; }
 	virtual bool GetIgnoreBoundsForEditorFocus() const override { return true; }
 #if WITH_EDITOR
-	virtual bool ComponentIsTouchingSelectionBox(const FBox& InSelBBox, const bool bConsiderOnlyBSP, const bool bMustEncompassEntireComponent) const override;
-	virtual bool ComponentIsTouchingSelectionFrustum(const FConvexVolume& InFrustum, const bool bConsiderOnlyBSP, const bool bMustEncompassEntireComponent) const override;
+	NAVIGATIONSYSTEM_API virtual bool ComponentIsTouchingSelectionBox(const FBox& InSelBBox, const bool bConsiderOnlyBSP, const bool bMustEncompassEntireComponent) const override;
+	NAVIGATIONSYSTEM_API virtual bool ComponentIsTouchingSelectionFrustum(const FConvexVolume& InFrustum, const bool bConsiderOnlyBSP, const bool bMustEncompassEntireComponent) const override;
 #endif
 	//~ End UPrimitiveComponent Interface
 
 	//~ Begin USceneComponent Interface
-	virtual FBoxSphereBounds CalcBounds(const FTransform &LocalToWorld) const override;
+	NAVIGATIONSYSTEM_API virtual FBoxSphereBounds CalcBounds(const FTransform &LocalToWorld) const override;
 	//~ End USceneComponent Interface
 };
 
diff --git a/Engine/Source/Runtime/NavigationSystem/Public/NavLinkRenderingProxy.h b/Engine/Source/Runtime/NavigationSystem/Public/NavLinkRenderingProxy.h
--- a/Engine/Source/Runtime/NavigationSystem/Public/NavLinkRenderingProxy.h
+++ b/Engine/Source/Runtime/NavigationSystem/Public/NavLinkRenderingProxy.h
@@ -13,14 +13,14 @@ class FMeshElementCollector;
 class FPrimitiveDrawInterface;
 class UPrimitiveComponent;
 
-class NAVIGATIONSYSTEM_API FNavLinkRenderingProxy : public FPrimitiveSceneProxy
+class FNavLinkRenderingProxy : public FPrimitiveSceneProxy
 {
 private:
 	AActor* LinkOwnerActor;
 	class INavLinkHostInterface* LinkOwnerHost;
 
 public:
-	SIZE_T GetTypeHash() const override;
+	NAVIGATIONSYSTEM_API SIZE_T GetTypeHash() const override;
 
 	struct FNavLinkDrawing
 	{
@@ -73,16 +73,16 @@ private:
 
 public:
 	/** Initialization constructor. */
-	FNavLinkRenderingProxy(const UPrimitiveComponent* InComponent);
-	virtual void GetDynamicMeshElements(const TArray<const FSceneView*>& Views, const FSceneViewFamily& ViewFamily, uint32 VisibilityMap, FMeshElementCollector& Collector) const override;
-	virtual FPrimitiveViewRelevance GetViewRelevance(const FSceneView* View) const override;
-	virtual uint32 GetMemoryFootprint( void ) const override;
-	uint32 GetAllocatedSize( void ) const;
-	void StorePointLinks(const FTransform& LocalToWorld, const TArray<FNavigationLink>& LinksArray);
-	void StoreSegmentLinks(const FTransform& LocalToWorld, const TArray<FNavigationSegmentLink>& LinksArray);
+	NAVIGATIONSYSTEM_API FNavLinkRenderingProxy(const UPrimitiveComponent* InComponent);
+	NAVIGATIONSYSTEM_API virtual void GetDynamicMeshElements(const TArray<const FSceneView*>& Views, const FSceneViewFamily& ViewFamily, uint32 VisibilityMap, FMeshElementCollector& Collector) const override;
+	NAVIGATIONSYSTEM_API virtual FPrimitiveViewRelevance GetViewRelevance(const FSceneView* View) const override;
+	NAVIGATIONSYSTEM_API virtual uint32 GetMemoryFootprint( void ) const override;
+	NAVIGATIONSYSTEM_API uint32 GetAllocatedSize( void ) const;
+	NAVIGATIONSYSTEM_API void StorePointLinks(const FTransform& LocalToWorld, const TArray<FNavigationLink>& LinksArray);
+	NAVIGATIONSYSTEM_API void StoreSegmentLinks(const FTransform& LocalToWorld, const TArray<FNavigationSegmentLink>& LinksArray);
 
-	static void GetLinkMeshes(const TArray<FNavLinkDrawing>& OffMeshPointLinks, const TArray<FNavLinkSegmentDrawing>& OffMeshSegmentLinks, TArray<float>& StepHeights, FMaterialRenderProxy* const MeshColorInstance, int32 ViewIndex, FMeshElementCollector& Collector, uint32 AgentMask);
+	static NAVIGATIONSYSTEM_API void GetLinkMeshes(const TArray<FNavLinkDrawing>& OffMeshPointLinks, const TArray<FNavLinkSegmentDrawing>& OffMeshSegmentLinks, TArray<float>& StepHeights, FMaterialRenderProxy* const MeshColorInstance, int32 ViewIndex, FMeshElementCollector& Collector, uint32 AgentMask);
 
 	/** made static to allow consistent navlinks drawing even if something is drawing links without FNavLinkRenderingProxy */
-	static void DrawLinks(FPrimitiveDrawInterface* PDI, TArray<FNavLinkDrawing>& OffMeshPointLinks, TArray<FNavLinkSegmentDrawing>& OffMeshSegmentLinks, TArray<float>& StepHeights, FMaterialRenderProxy* const MeshColorInstance, uint32 AgentMask);
+	static NAVIGATIONSYSTEM_API void DrawLinks(FPrimitiveDrawInterface* PDI, TArray<FNavLinkDrawing>& OffMeshPointLinks, TArray<FNavLinkSegmentDrawing>& OffMeshSegmentLinks, TArray<float>& StepHeights, FMaterialRenderProxy* const MeshColorInstance, uint32 AgentMask);
 };
diff --git a/Engine/Source/Runtime/NavigationSystem/Public/NavLinkTrivial.h b/Engine/Source/Runtime/NavigationSystem/Public/NavLinkTrivial.h
--- a/Engine/Source/Runtime/NavigationSystem/Public/NavLinkTrivial.h
+++ b/Engine/Source/Runtime/NavigationSystem/Public/NavLinkTrivial.h
@@ -7,8 +7,8 @@
 #include "AI/Navigation/NavLinkDefinition.h"
 #include "NavLinkTrivial.generated.h"
 
-UCLASS()
-class NAVIGATIONSYSTEM_API UNavLinkTrivial : public UNavLinkDefinition
+UCLASS(MinimalAPI)
+class UNavLinkTrivial : public UNavLinkDefinition
 {
 	GENERATED_UCLASS_BODY()
 };
diff --git a/Engine/Source/Runtime/NavigationSystem/Public/NavMesh/NavMeshBoundsVolume.h b/Engine/Source/Runtime/NavigationSystem/Public/NavMesh/NavMeshBoundsVolume.h
--- a/Engine/Source/Runtime/NavigationSystem/Public/NavMesh/NavMeshBoundsVolume.h
+++ b/Engine/Source/Runtime/NavigationSystem/Public/NavMesh/NavMeshBoundsVolume.h
@@ -15,8 +15,8 @@
 #include "NavMeshBoundsVolume.generated.h"
 
 
-UCLASS()
-class NAVIGATIONSYSTEM_API ANavMeshBoundsVolume : public AVolume
+UCLASS(MinimalAPI)
+class ANavMeshBoundsVolume : public AVolume
 {
 	GENERATED_UCLASS_BODY()
 
@@ -24,16 +24,16 @@ class NAVIGATIONSYSTEM_API ANavMeshBoundsVolume : public AVolume
 	FNavAgentSelector SupportedAgents;
 
 	//~ Begin AActor Interface
-	virtual void PostRegisterAllComponents() override;
-	virtual void PostUnregisterAllComponents() override;
+	NAVIGATIONSYSTEM_API virtual void PostRegisterAllComponents() override;
+	NAVIGATIONSYSTEM_API virtual void PostUnregisterAllComponents() override;
 	//~ End AActor Interface
 #if WITH_EDITOR
 	//~ Begin UObject Interface
-	virtual void PostEditChangeProperty( struct FPropertyChangedEvent& PropertyChangedEvent) override;
-	virtual void PostEditUndo() override;
+	NAVIGATIONSYSTEM_API virtual void PostEditChangeProperty( struct FPropertyChangedEvent& PropertyChangedEvent) override;
+	NAVIGATIONSYSTEM_API virtual void PostEditUndo() override;
 	//~ End UObject Interface
 
-	static void OnPostEngineInit();
+	static NAVIGATIONSYSTEM_API void OnPostEngineInit();
 #endif // WITH_EDITOR
 };
 
