[Engine]

diff --git a/Engine/Source/Runtime/NavigationSystem/Public/NavMesh/NavMeshPath.h b/Engine/Source/Runtime/NavigationSystem/Public/NavMesh/NavMeshPath.h
--- a/Engine/Source/Runtime/NavigationSystem/Public/NavMesh/NavMeshPath.h
+++ b/Engine/Source/Runtime/NavigationSystem/Public/NavMesh/NavMeshPath.h
@@ -11,7 +11,7 @@
 #define RECAST_STRAIGHTPATH_OFFMESH_CONNECTION 0x04
 
 /** Helper to translate FNavPathPoint.Flags. */
-struct NAVIGATIONSYSTEM_API FNavMeshNodeFlags
+struct FNavMeshNodeFlags
 {
 	/** Extra node information (like "path start", "off-mesh connection"). */
 	uint8 PathFlags;
@@ -33,11 +33,11 @@ struct NAVIGATIONSYSTEM_API FNavMeshNodeFlags
 };
 
 
-struct NAVIGATIONSYSTEM_API FNavMeshPath : public FNavigationPath
+struct FNavMeshPath : public FNavigationPath
 {
 	typedef FNavigationPath Super;
 
-	FNavMeshPath();
+	NAVIGATIONSYSTEM_API FNavMeshPath();
 
 	PRAGMA_DISABLE_DEPRECATION_WARNINGS
 	FNavMeshPath(const FNavMeshPath&) = default;
@@ -51,7 +51,7 @@ struct NAVIGATIONSYSTEM_API FNavMeshPath : public FNavigationPath
 	FORCEINLINE bool IsStringPulled() const { return bStringPulled; }
 
 	/** find string pulled path from PathCorridor */
-	void PerformStringPulling(const FVector& StartLoc, const FVector& EndLoc);
+	NAVIGATIONSYSTEM_API void PerformStringPulling(const FVector& StartLoc, const FVector& EndLoc);
 
 	FORCEINLINE void SetWantsPathCorridor(const bool bNewWantsPathCorridor) { bWantsPathCorridor = bNewWantsPathCorridor; }
 	FORCEINLINE bool WantsPathCorridor() const { return bWantsPathCorridor; }
@@ -61,18 +61,18 @@ struct NAVIGATIONSYSTEM_API FNavMeshPath : public FNavigationPath
 
 	FORCEINLINE void OnPathCorridorUpdated() { bCorridorEdgesGenerated = false; }
 
-	virtual void DebugDraw(const ANavigationData* NavData, const FColor PathColor, UCanvas* Canvas, const bool bPersistent, const float LifeTime, const uint32 NextPathPointIndex = 0) const override;
+	NAVIGATIONSYSTEM_API virtual void DebugDraw(const ANavigationData* NavData, const FColor PathColor, UCanvas* Canvas, const bool bPersistent, const float LifeTime, const uint32 NextPathPointIndex = 0) const override;
 
-	bool ContainsWithSameEnd(const FNavMeshPath* Other) const;
+	NAVIGATIONSYSTEM_API bool ContainsWithSameEnd(const FNavMeshPath* Other) const;
 
-	void OffsetFromCorners(FVector::FReal Distance);
+	NAVIGATIONSYSTEM_API void OffsetFromCorners(FVector::FReal Distance);
 
-	void ApplyFlags(int32 NavDataFlags);
+	NAVIGATIONSYSTEM_API void ApplyFlags(int32 NavDataFlags);
 
-	virtual void ResetForRepath() override;
+	NAVIGATIONSYSTEM_API virtual void ResetForRepath() override;
 
 	/** get flags of path point or corridor poly (depends on bStringPulled flag) */
-	bool GetNodeFlags(int32 NodeIdx, FNavMeshNodeFlags& Flags) const;
+	NAVIGATIONSYSTEM_API bool GetNodeFlags(int32 NodeIdx, FNavMeshNodeFlags& Flags) const;
 
 	/** get cost of path, starting from next poly in corridor */
 	virtual FVector::FReal GetCostFromNode(NavNodeRef PathNode) const override { return GetCostFromIndex(PathCorridor.Find(PathNode) + 1); }
@@ -103,41 +103,41 @@ struct NAVIGATIONSYSTEM_API FNavMeshPath : public FNavigationPath
 	virtual bool ContainsCustomLink(FNavLinkId UniqueLinkId) const override { return CustomNavLinkIds.Contains(UniqueLinkId); }
 	virtual bool ContainsAnyCustomLink() const override { return CustomNavLinkIds.Num() > 0; }
 
-	bool IsPathSegmentANavLink(const int32 PathSegmentStartIndex) const;
+	NAVIGATIONSYSTEM_API bool IsPathSegmentANavLink(const int32 PathSegmentStartIndex) const;
 
-	virtual bool DoesIntersectBox(const FBox& Box, uint32 StartingIndex = 0, int32* IntersectingSegmentIndex = NULL, FVector* AgentExtent = NULL) const override;
-	virtual bool DoesIntersectBox(const FBox& Box, const FVector& AgentLocation, uint32 StartingIndex = 0, int32* IntersectingSegmentIndex = NULL, FVector* AgentExtent = NULL) const override;
+	NAVIGATIONSYSTEM_API virtual bool DoesIntersectBox(const FBox& Box, uint32 StartingIndex = 0, int32* IntersectingSegmentIndex = NULL, FVector* AgentExtent = NULL) const override;
+	NAVIGATIONSYSTEM_API virtual bool DoesIntersectBox(const FBox& Box, const FVector& AgentLocation, uint32 StartingIndex = 0, int32* IntersectingSegmentIndex = NULL, FVector* AgentExtent = NULL) const override;
 	/** retrieves normalized direction vector to given path segment. If path is not string pulled navigation corridor is being used */
-	virtual FVector GetSegmentDirection(uint32 SegmentEndIndex) const override;
+	NAVIGATIONSYSTEM_API virtual FVector GetSegmentDirection(uint32 SegmentEndIndex) const override;
 
-	void Invert();
+	NAVIGATIONSYSTEM_API void Invert();
 
 private:
-	bool DoesPathIntersectBoxImplementation(const FBox& Box, const FVector& StartLocation, uint32 StartingIndex, int32* IntersectingSegmentIndex, FVector* AgentExtent) const;
-	void InternalResetNavMeshPath();
+	NAVIGATIONSYSTEM_API bool DoesPathIntersectBoxImplementation(const FBox& Box, const FVector& StartLocation, uint32 StartingIndex, int32* IntersectingSegmentIndex, FVector* AgentExtent) const;
+	NAVIGATIONSYSTEM_API void InternalResetNavMeshPath();
 
 public:
 
 #if ENABLE_VISUAL_LOG
-	virtual void DescribeSelfToVisLog(struct FVisualLogEntry* Snapshot) const override;
-	virtual FString GetDescription() const override;
+	NAVIGATIONSYSTEM_API virtual void DescribeSelfToVisLog(struct FVisualLogEntry* Snapshot) const override;
+	NAVIGATIONSYSTEM_API virtual FString GetDescription() const override;
 #endif // ENABLE_VISUAL_LOG
 
 protected:
 	/** calculates total length of string pulled path. Does not generate string pulled
 	*	path if it's not already generated (see bWantsStringPulling and bStrigPulled)
 	*	Internal use only */
-	FVector::FReal GetStringPulledLength(const int32 StartingPoint) const;
+	NAVIGATIONSYSTEM_API FVector::FReal GetStringPulledLength(const int32 StartingPoint) const;
 
 	/** calculates estimated length of path expressed as sequence of navmesh edges.
 	*	It basically sums up distances between every subsequent nav edge pair edge middles.
 	*	Internal use only */
-	FVector::FReal GetPathCorridorLength(const int32 StartingEdge) const;
+	NAVIGATIONSYSTEM_API FVector::FReal GetPathCorridorLength(const int32 StartingEdge) const;
 
 	/** it's only const to be callable in const environment. It's not supposed to be called directly externally anyway,
 	*	just as part of retrieving corridor on demand or generating it in internal processes. It fills a mutable
 	*	array. */
-	const TArray<FNavigationPortalEdge>& GeneratePathCorridorEdges() const;
+	NAVIGATIONSYSTEM_API const TArray<FNavigationPortalEdge>& GeneratePathCorridorEdges() const;
 
 public:
 
@@ -181,5 +181,5 @@ protected:
 	uint32 bWantsPathCorridor : 1;
 
 public:
-	static const FNavPathType Type;
-};
\ No newline at end of file
+	static NAVIGATIONSYSTEM_API const FNavPathType Type;
+};
diff --git a/Engine/Source/Runtime/NavigationSystem/Public/NavMesh/NavMeshRenderingComponent.h b/Engine/Source/Runtime/NavigationSystem/Public/NavMesh/NavMeshRenderingComponent.h
--- a/Engine/Source/Runtime/NavigationSystem/Public/NavMesh/NavMeshRenderingComponent.h
+++ b/Engine/Source/Runtime/NavigationSystem/Public/NavMesh/NavMeshRenderingComponent.h
@@ -49,7 +49,7 @@ enum class ENavMeshDetailFlags : uint8
 };
 
 // exported to API for GameplayDebugger module
-struct NAVIGATIONSYSTEM_API FNavMeshSceneProxyData : public TSharedFromThis<FNavMeshSceneProxyData, ESPMode::ThreadSafe>
+struct FNavMeshSceneProxyData : public TSharedFromThis<FNavMeshSceneProxyData, ESPMode::ThreadSafe>
 {
 	struct FDebugMeshData
 	{
@@ -99,37 +99,37 @@ struct NAVIGATIONSYSTEM_API FNavMeshSceneProxyData : public TSharedFromThis<FNav
 	FNavMeshSceneProxyData() : NavMeshDrawOffset(0, 0, 10.f),
 		bDataGathered(false), bNeedsNewData(true), NavDetailFlags(0) {}
 
-	void Reset();
-	void Serialize(FArchive& Ar);
-	uint32 GetAllocatedSize() const;
+	NAVIGATIONSYSTEM_API void Reset();
+	NAVIGATIONSYSTEM_API void Serialize(FArchive& Ar);
+	NAVIGATIONSYSTEM_API uint32 GetAllocatedSize() const;
 
 #if WITH_RECAST
-	void GatherData(const ARecastNavMesh* NavMesh, int32 InNavDetailFlags, const TArray<int32>& TileSet);
+	NAVIGATIONSYSTEM_API void GatherData(const ARecastNavMesh* NavMesh, int32 InNavDetailFlags, const TArray<int32>& TileSet);
 
 #if RECAST_INTERNAL_DEBUG_DATA
-	void AddMeshForInternalData(const struct FRecastInternalDebugData& InInternalData);
+	NAVIGATIONSYSTEM_API void AddMeshForInternalData(const struct FRecastInternalDebugData& InInternalData);
 #endif //RECAST_INTERNAL_DEBUG_DATA
 
 #endif
 };
 
 // exported to API for GameplayDebugger module
-class NAVIGATIONSYSTEM_API FNavMeshSceneProxy final : public FDebugRenderSceneProxy, public FNoncopyable
+class FNavMeshSceneProxy final : public FDebugRenderSceneProxy, public FNoncopyable
 {
 	friend class FNavMeshDebugDrawDelegateHelper;
 public:
-	virtual SIZE_T GetTypeHash() const override;
+	NAVIGATIONSYSTEM_API virtual SIZE_T GetTypeHash() const override;
 
-	FNavMeshSceneProxy(const UPrimitiveComponent* InComponent, FNavMeshSceneProxyData* InProxyData, bool ForceToRender = false);
-	virtual ~FNavMeshSceneProxy() override;
+	NAVIGATIONSYSTEM_API FNavMeshSceneProxy(const UPrimitiveComponent* InComponent, FNavMeshSceneProxyData* InProxyData, bool ForceToRender = false);
+	NAVIGATIONSYSTEM_API virtual ~FNavMeshSceneProxy() override;
 
-	virtual void GetDynamicMeshElements(const TArray<const FSceneView*>& Views, const FSceneViewFamily& ViewFamily, uint32 VisibilityMap, FMeshElementCollector& Collector) const override;
+	NAVIGATIONSYSTEM_API virtual void GetDynamicMeshElements(const TArray<const FSceneView*>& Views, const FSceneViewFamily& ViewFamily, uint32 VisibilityMap, FMeshElementCollector& Collector) const override;
 
 protected:
-	virtual FPrimitiveViewRelevance GetViewRelevance(const FSceneView* View) const override;
+	NAVIGATIONSYSTEM_API virtual FPrimitiveViewRelevance GetViewRelevance(const FSceneView* View) const override;
 
 	virtual uint32 GetMemoryFootprint(void) const override { return sizeof(*this) + GetAllocatedSizeInternal(); }
-	uint32 GetAllocatedSizeInternal(void) const;
+	NAVIGATIONSYSTEM_API uint32 GetAllocatedSizeInternal(void) const;
 
 private:			
 	FNavMeshSceneProxyData ProxyData;
@@ -185,8 +185,8 @@ private:
 };
 #endif
 
-UCLASS(editinlinenew, ClassGroup = Debug)
-class NAVIGATIONSYSTEM_API UNavMeshRenderingComponent : public UDebugDrawComponent
+UCLASS(editinlinenew, ClassGroup = Debug, MinimalAPI)
+class UNavMeshRenderingComponent : public UDebugDrawComponent
 {
 	GENERATED_UCLASS_BODY()
 
@@ -194,24 +194,24 @@ public:
 	void ForceUpdate() { bForceUpdate = true; }
 	bool IsForcingUpdate() const { return bForceUpdate; }
 
-	static bool IsNavigationShowFlagSet(const UWorld* World);
+	static NAVIGATIONSYSTEM_API bool IsNavigationShowFlagSet(const UWorld* World);
 
 protected:
-	virtual void OnRegister()  override;
-	virtual void OnUnregister()  override;
+	NAVIGATIONSYSTEM_API virtual void OnRegister()  override;
+	NAVIGATIONSYSTEM_API virtual void OnUnregister()  override;
 
 #if UE_ENABLE_DEBUG_DRAWING
-  	virtual FDebugRenderSceneProxy* CreateDebugSceneProxy() override;
+  	NAVIGATIONSYSTEM_API virtual FDebugRenderSceneProxy* CreateDebugSceneProxy() override;
 	virtual FDebugDrawDelegateHelper& GetDebugDrawDelegateHelper() override { return NavMeshDebugDrawDelegateManager; }
 #endif
 
-	virtual FBoxSphereBounds CalcBounds(const FTransform &LocalToWorld) const override;
+	NAVIGATIONSYSTEM_API virtual FBoxSphereBounds CalcBounds(const FTransform &LocalToWorld) const override;
 
 	/** Gathers drawable information from NavMesh and puts it in OutProxyData. 
 	 *	Override to add additional information to OutProxyData.*/
-	virtual void GatherData(const ARecastNavMesh& NavMesh, FNavMeshSceneProxyData& OutProxyData) const;
+	NAVIGATIONSYSTEM_API virtual void GatherData(const ARecastNavMesh& NavMesh, FNavMeshSceneProxyData& OutProxyData) const;
 
-	void TimerFunction();
+	NAVIGATIONSYSTEM_API void TimerFunction();
 
 protected:
 	uint32 bCollectNavigationData : 1;
diff --git a/Engine/Source/Runtime/NavigationSystem/Public/NavMesh/PImplRecastNavMesh.h b/Engine/Source/Runtime/NavigationSystem/Public/NavMesh/PImplRecastNavMesh.h
--- a/Engine/Source/Runtime/NavigationSystem/Public/NavMesh/PImplRecastNavMesh.h
+++ b/Engine/Source/Runtime/NavigationSystem/Public/NavMesh/PImplRecastNavMesh.h
@@ -34,47 +34,47 @@ struct FNavLinkId;
 #define RECAST_VERY_SMALL_AGENT_RADIUS 0.0f
 
 /** Engine Private! - Private Implementation details of ARecastNavMesh */
-class NAVIGATIONSYSTEM_API FPImplRecastNavMesh
+class FPImplRecastNavMesh
 {
 public:
 
 	/** Constructor */
-	FPImplRecastNavMesh(ARecastNavMesh* Owner);
+	NAVIGATIONSYSTEM_API FPImplRecastNavMesh(ARecastNavMesh* Owner);
 
 	/** Dtor */
-	~FPImplRecastNavMesh();
+	NAVIGATIONSYSTEM_API ~FPImplRecastNavMesh();
 
 	/**
 	 * Serialization.
 	 * @param Ar - The archive with which to serialize.
 	 * @returns true if serialization was successful.
 	 */
-	void Serialize(FArchive& Ar, int32 NavMeshVersion);
+	NAVIGATIONSYSTEM_API void Serialize(FArchive& Ar, int32 NavMeshVersion);
 
 	/* Gather debug geometry.
 	 * @params OutGeometry Output geometry.
 	 * @params TileIndex Used to collect geometry for a specific tile, INDEX_NONE will gather all tiles.
 	 * @return True if done collecting.
 	 */
-	bool GetDebugGeometryForTile(FRecastDebugGeometry& OutGeometry, int32 TileIndex) const;
+	NAVIGATIONSYSTEM_API bool GetDebugGeometryForTile(FRecastDebugGeometry& OutGeometry, int32 TileIndex) const;
 	
 	/** Returns bounding box for the whole navmesh. */
-	FBox GetNavMeshBounds() const;
+	NAVIGATIONSYSTEM_API FBox GetNavMeshBounds() const;
 
 	/** Returns bounding box for a given navmesh tile. */
-	FBox GetNavMeshTileBounds(int32 TileIndex) const;
+	NAVIGATIONSYSTEM_API FBox GetNavMeshTileBounds(int32 TileIndex) const;
 
 	/** Retrieves XY and layer coordinates of tile specified by index */
-	bool GetNavMeshTileXY(int32 TileIndex, int32& OutX, int32& OutY, int32& OutLayer) const;
+	NAVIGATIONSYSTEM_API bool GetNavMeshTileXY(int32 TileIndex, int32& OutX, int32& OutY, int32& OutLayer) const;
 
 	/** Retrieves XY coordinates of tile specified by position */
-	bool GetNavMeshTileXY(const FVector& Point, int32& OutX, int32& OutY) const;
+	NAVIGATIONSYSTEM_API bool GetNavMeshTileXY(const FVector& Point, int32& OutX, int32& OutY) const;
 
 	/** Retrieves all tile indices at matching XY coordinates */
-	void GetNavMeshTilesAt(int32 TileX, int32 TileY, TArray<int32>& Indices) const;
+	NAVIGATIONSYSTEM_API void GetNavMeshTilesAt(int32 TileX, int32 TileY, TArray<int32>& Indices) const;
 
 	/** Retrieves list of tiles that intersect specified bounds */
-	void GetNavMeshTilesIn(const TArray<FBox>& InclusionBounds, TArray<int32>& Indices) const;
+	NAVIGATIONSYSTEM_API void GetNavMeshTilesIn(const TArray<FBox>& InclusionBounds, TArray<int32>& Indices) const;
 
 	/** Retrieves number of tiles in this navmesh */
 	FORCEINLINE int32 GetNavMeshTilesCount() const { return DetourNavMesh ? DetourNavMesh->getMaxTiles() : 0; }
@@ -83,179 +83,179 @@ public:
 
 	/** Generates path from the given query. Synchronous. */
 	UE_DEPRECATED(5.2, "Please use FindPath with the added bRequireNavigableEndLocation parameter (true can be used as default).")
-	ENavigationQueryResult::Type FindPath(const FVector& StartLoc, const FVector& EndLoc, const FVector::FReal CostLimit, FNavMeshPath& Path, const FNavigationQueryFilter& Filter, const UObject* Owner) const;
+	NAVIGATIONSYSTEM_API ENavigationQueryResult::Type FindPath(const FVector& StartLoc, const FVector& EndLoc, const FVector::FReal CostLimit, FNavMeshPath& Path, const FNavigationQueryFilter& Filter, const UObject* Owner) const;
 																																									  
 	/** Generates path from the given query. Synchronous. */
-	ENavigationQueryResult::Type FindPath(const FVector& StartLoc, const FVector& EndLoc, const FVector::FReal CostLimit, const bool bRequireNavigableEndLocation, FNavMeshPath& Path, const FNavigationQueryFilter& Filter, const UObject* Owner) const;
+	NAVIGATIONSYSTEM_API ENavigationQueryResult::Type FindPath(const FVector& StartLoc, const FVector& EndLoc, const FVector::FReal CostLimit, const bool bRequireNavigableEndLocation, FNavMeshPath& Path, const FNavigationQueryFilter& Filter, const UObject* Owner) const;
 
 	/** Check if path exists */
 	UE_DEPRECATED(5.2, "Please use TestPath with the added bRequireNavigableEndLocation parameter (true can be used as default).")
-	ENavigationQueryResult::Type TestPath(const FVector& StartLoc, const FVector& EndLoc, const FNavigationQueryFilter& Filter, const UObject* Owner, int32* NumVisitedNodes = 0) const;
+	NAVIGATIONSYSTEM_API ENavigationQueryResult::Type TestPath(const FVector& StartLoc, const FVector& EndLoc, const FNavigationQueryFilter& Filter, const UObject* Owner, int32* NumVisitedNodes = 0) const;
 
 	/** Check if path exists */
-	ENavigationQueryResult::Type TestPath(const FVector& StartLoc, const FVector& EndLoc, const bool bRequireNavigableEndLocation, const FNavigationQueryFilter& Filter, const UObject* Owner, int32* NumVisitedNodes = 0) const;
+	NAVIGATIONSYSTEM_API ENavigationQueryResult::Type TestPath(const FVector& StartLoc, const FVector& EndLoc, const bool bRequireNavigableEndLocation, const FNavigationQueryFilter& Filter, const UObject* Owner, int32* NumVisitedNodes = 0) const;
 
 	template< typename TRecastAStar, typename TRecastAStartGraph, typename TRecastGraphAStarFilter, typename TRecastAStarResult >
 	ENavigationQueryResult::Type FindPathCustomAStar(TRecastAStartGraph& RecastGraphWrapper, TRecastAStar& AStarAlgo, const FVector& StartLoc, const FVector& EndLoc, const FVector::FReal CostLimit, FNavMeshPath& Path, const FNavigationQueryFilter& Filter, const UObject* Owner) const;
 
 	/** Checks if the whole segment is in navmesh */
-	void Raycast(const FVector& StartLoc, const FVector& EndLoc, const FNavigationQueryFilter& InQueryFilter, const UObject* Owner,
+	NAVIGATIONSYSTEM_API void Raycast(const FVector& StartLoc, const FVector& EndLoc, const FNavigationQueryFilter& InQueryFilter, const UObject* Owner,
 		ARecastNavMesh::FRaycastResult& RaycastResult, NavNodeRef StartNode = INVALID_NAVNODEREF) const;
 
 	/** Generates path from given query and collect data for every step of A* algorithm */
 	UE_DEPRECATED(5.2, "Please use DebugPathfinding with the added bRequireNavigableEndLocation parameter (true can be used as default).")
-	int32 DebugPathfinding(const FVector& StartLoc, const FVector& EndLoc, const FVector::FReal CostLimit, const FNavigationQueryFilter& Filter, const UObject* Owner, TArray<FRecastDebugPathfindingData>& Steps);
+	NAVIGATIONSYSTEM_API int32 DebugPathfinding(const FVector& StartLoc, const FVector& EndLoc, const FVector::FReal CostLimit, const FNavigationQueryFilter& Filter, const UObject* Owner, TArray<FRecastDebugPathfindingData>& Steps);
 
 	/** Generates path from given query and collect data for every step of A* algorithm */
-	int32 DebugPathfinding(const FVector& StartLoc, const FVector& EndLoc, const FVector::FReal CostLimit, const bool bRequireNavigableEndLocation, const FNavigationQueryFilter& Filter, const UObject* Owner, TArray<FRecastDebugPathfindingData>& Steps);
+	NAVIGATIONSYSTEM_API int32 DebugPathfinding(const FVector& StartLoc, const FVector& EndLoc, const FVector::FReal CostLimit, const bool bRequireNavigableEndLocation, const FNavigationQueryFilter& Filter, const UObject* Owner, TArray<FRecastDebugPathfindingData>& Steps);
 
 	/** Returns a random location on the navmesh. */
-	FNavLocation GetRandomPoint(const FNavigationQueryFilter& Filter, const UObject* Owner) const;
+	NAVIGATIONSYSTEM_API FNavLocation GetRandomPoint(const FNavigationQueryFilter& Filter, const UObject* Owner) const;
 
 #if WITH_NAVMESH_CLUSTER_LINKS
 	/** Check if path exists using cluster graph */
-	ENavigationQueryResult::Type TestClusterPath(const FVector& StartLoc, const FVector& EndLoc, int32* NumVisitedNodes = 0) const;
+	NAVIGATIONSYSTEM_API ENavigationQueryResult::Type TestClusterPath(const FVector& StartLoc, const FVector& EndLoc, int32* NumVisitedNodes = 0) const;
 
 	/** Returns a random location on the navmesh within cluster */
-	bool GetRandomPointInCluster(NavNodeRef ClusterRef, FNavLocation& OutLocation) const;
+	NAVIGATIONSYSTEM_API bool GetRandomPointInCluster(NavNodeRef ClusterRef, FNavLocation& OutLocation) const;
 #endif // WITH_NAVMESH_CLUSTER_LINKS
 
 	/**	Tries to move current nav location towards target constrained to navigable area. Faster than ProjectPointToNavmesh.
 	 *	@param OutLocation if successful this variable will be filed with result
 	 *	@return true if successful, false otherwise
 	 */
-	bool FindMoveAlongSurface(const FNavLocation& StartLocation, const FVector& TargetPosition, FNavLocation& OutLocation, const FNavigationQueryFilter& Filter, const UObject* Owner) const;
+	NAVIGATIONSYSTEM_API bool FindMoveAlongSurface(const FNavLocation& StartLocation, const FVector& TargetPosition, FNavLocation& OutLocation, const FNavigationQueryFilter& Filter, const UObject* Owner) const;
 
-	bool ProjectPointToNavMesh(const FVector& Point, FNavLocation& Result, const FVector& Extent, const FNavigationQueryFilter& Filter, const UObject* Owner) const;
+	NAVIGATIONSYSTEM_API bool ProjectPointToNavMesh(const FVector& Point, FNavLocation& Result, const FVector& Extent, const FNavigationQueryFilter& Filter, const UObject* Owner) const;
 	
 	/** Project single point and grab all vertical intersections */
-	bool ProjectPointMulti(const FVector& Point, TArray<FNavLocation>& OutLocations, const FVector& Extent,
+	NAVIGATIONSYSTEM_API bool ProjectPointMulti(const FVector& Point, TArray<FNavLocation>& OutLocations, const FVector& Extent,
 		FVector::FReal MinZ, FVector::FReal MaxZ, const FNavigationQueryFilter& Filter, const UObject* Owner) const;
 
 	/** Returns nearest navmesh polygon to Loc, or INVALID_NAVMESHREF if Loc is not on the navmesh. */
-	NavNodeRef FindNearestPoly(FVector const& Loc, FVector const& Extent, const FNavigationQueryFilter& Filter, const UObject* Owner) const;
+	NAVIGATIONSYSTEM_API NavNodeRef FindNearestPoly(FVector const& Loc, FVector const& Extent, const FNavigationQueryFilter& Filter, const UObject* Owner) const;
 
 	/** Finds the polygons along the navigation graph that touch the specified circle. Return true if found any. */
-	bool FindPolysAroundCircle(const FVector& CenterPos, const NavNodeRef CenterNodeRef, const FVector::FReal Radius, const FNavigationQueryFilter& Filter, const UObject* Owner, TArray<NavNodeRef>* OutPolys = nullptr, TArray<NavNodeRef>* OutPolysParent = nullptr, TArray<float>* OutPolysCost = nullptr, int32* OutPolysCount = nullptr) const;
+	NAVIGATIONSYSTEM_API bool FindPolysAroundCircle(const FVector& CenterPos, const NavNodeRef CenterNodeRef, const FVector::FReal Radius, const FNavigationQueryFilter& Filter, const UObject* Owner, TArray<NavNodeRef>* OutPolys = nullptr, TArray<NavNodeRef>* OutPolysParent = nullptr, TArray<float>* OutPolysCost = nullptr, int32* OutPolysCount = nullptr) const;
 
 	/** Retrieves all polys within given pathing distance from StartLocation.
 	 *	@NOTE query is not using string-pulled path distance (for performance reasons),
 	 *		it measured distance between middles of portal edges, do you might want to 
 	 *		add an extra margin to PathingDistance */
-	bool GetPolysWithinPathingDistance(FVector const& StartLoc, const FVector::FReal PathingDistance,
+	NAVIGATIONSYSTEM_API bool GetPolysWithinPathingDistance(FVector const& StartLoc, const FVector::FReal PathingDistance,
 		const FNavigationQueryFilter& Filter, const UObject* Owner,
 		TArray<NavNodeRef>& FoundPolys, FRecastDebugPathfindingData* DebugData) const;
 
 	//@todo document
-	void GetEdgesForPathCorridor(const TArray<NavNodeRef>* PathCorridor, TArray<FNavigationPortalEdge>* PathCorridorEdges) const;
+	NAVIGATIONSYSTEM_API void GetEdgesForPathCorridor(const TArray<NavNodeRef>* PathCorridor, TArray<FNavigationPortalEdge>* PathCorridorEdges) const;
 
 	/** finds stringpulled path from given corridor */
-	bool FindStraightPath(const FVector& StartLoc, const FVector& EndLoc, const TArray<NavNodeRef>& PathCorridor, TArray<FNavPathPoint>& PathPoints, TArray<FNavLinkId>* CustomLinks = NULL) const;
+	NAVIGATIONSYSTEM_API bool FindStraightPath(const FVector& StartLoc, const FVector& EndLoc, const TArray<NavNodeRef>& PathCorridor, TArray<FNavPathPoint>& PathPoints, TArray<FNavLinkId>* CustomLinks = NULL) const;
 
 	/** finds stringpulled path from given corridor */
 	UE_DEPRECATED(5.3, "Please use FindStraightPath with the TArray<FNavPathPoint>* CustomLinks. This function has no effect.")
 	bool FindStraightPath(const FVector& StartLoc, const FVector& EndLoc, const TArray<NavNodeRef>& PathCorridor, TArray<FNavPathPoint>& PathPoints, TArray<uint32>* CustomLinks) const { return false; }
 
 	/** Filters nav polys in PolyRefs with Filter */
-	bool FilterPolys(TArray<NavNodeRef>& PolyRefs, const FRecastQueryFilter* Filter, const UObject* Owner) const;
+	NAVIGATIONSYSTEM_API bool FilterPolys(TArray<NavNodeRef>& PolyRefs, const FRecastQueryFilter* Filter, const UObject* Owner) const;
 
 	/** Get all polys from tile */
-	bool GetPolysInTile(int32 TileIndex, TArray<FNavPoly>& Polys) const;
+	NAVIGATIONSYSTEM_API bool GetPolysInTile(int32 TileIndex, TArray<FNavPoly>& Polys) const;
 
 	UE_DEPRECATED(5.3, "Please use the version of this function that takes a FNavLinkId. This function has no effect.")
 	void UpdateNavigationLinkArea(int32 UserId, uint8 AreaType, uint16 PolyFlags) const {}
 
 	/** Updates area on polygons creating point-to-point connection with given UserId */
-	void UpdateNavigationLinkArea(FNavLinkId UserId, uint8 AreaType, uint16 PolyFlags) const;
+	NAVIGATIONSYSTEM_API void UpdateNavigationLinkArea(FNavLinkId UserId, uint8 AreaType, uint16 PolyFlags) const;
 
 #if WITH_NAVMESH_SEGMENT_LINKS
 	/** Updates area on polygons creating segment-to-segment connection with given UserId */
-	void UpdateSegmentLinkArea(int32 UserId, uint8 AreaType, uint16 PolyFlags) const;
+	NAVIGATIONSYSTEM_API void UpdateSegmentLinkArea(int32 UserId, uint8 AreaType, uint16 PolyFlags) const;
 #endif // WITH_NAVMESH_SEGMENT_LINKS
 
 	/** Retrieves center of the specified polygon. Returns false on error. */
-	bool GetPolyCenter(NavNodeRef PolyID, FVector& OutCenter) const;
+	NAVIGATIONSYSTEM_API bool GetPolyCenter(NavNodeRef PolyID, FVector& OutCenter) const;
 	/** Retrieves the vertices for the specified polygon. Returns false on error. */
-	bool GetPolyVerts(NavNodeRef PolyID, TArray<FVector>& OutVerts) const;
+	NAVIGATIONSYSTEM_API bool GetPolyVerts(NavNodeRef PolyID, TArray<FVector>& OutVerts) const;
 	/** Retrieves a random point inside the specified polygon. Returns false on error. */
-	bool GetRandomPointInPoly(NavNodeRef PolyID, FVector& OutPoint) const;
+	NAVIGATIONSYSTEM_API bool GetRandomPointInPoly(NavNodeRef PolyID, FVector& OutPoint) const;
 	/** Retrieves the flags for the specified polygon. Returns false on error. */
-	bool GetPolyData(NavNodeRef PolyID, uint16& Flags, uint8& AreaType) const;
+	NAVIGATIONSYSTEM_API bool GetPolyData(NavNodeRef PolyID, uint16& Flags, uint8& AreaType) const;
 	/** Retrieves area ID for the specified polygon. */
-	uint32 GetPolyAreaID(NavNodeRef PolyID) const;
+	NAVIGATIONSYSTEM_API uint32 GetPolyAreaID(NavNodeRef PolyID) const;
 	/** Sets area ID for the specified polygon. */
-	void SetPolyAreaID(NavNodeRef PolyID, uint8 AreaID);
+	NAVIGATIONSYSTEM_API void SetPolyAreaID(NavNodeRef PolyID, uint8 AreaID);
 	/** Finds all polys connected with specified one */
-	bool GetPolyNeighbors(NavNodeRef PolyID, TArray<FNavigationPortalEdge>& Neighbors) const;
+	NAVIGATIONSYSTEM_API bool GetPolyNeighbors(NavNodeRef PolyID, TArray<FNavigationPortalEdge>& Neighbors) const;
 	/** Finds all polys connected with specified one, results expressed as array of NavNodeRefs */
-	bool GetPolyNeighbors(NavNodeRef PolyID, TArray<NavNodeRef>& Neighbors) const;
+	NAVIGATIONSYSTEM_API bool GetPolyNeighbors(NavNodeRef PolyID, TArray<NavNodeRef>& Neighbors) const;
 	/** Finds all polys connected with specified one */
-	bool GetPolyEdges(NavNodeRef PolyID, TArray<FNavigationPortalEdge>& Edges) const;
+	NAVIGATIONSYSTEM_API bool GetPolyEdges(NavNodeRef PolyID, TArray<FNavigationPortalEdge>& Edges) const;
 	/** Finds closest point constrained to given poly */
-	bool GetClosestPointOnPoly(NavNodeRef PolyID, const FVector& TestPt, FVector& PointOnPoly) const;
+	NAVIGATIONSYSTEM_API bool GetClosestPointOnPoly(NavNodeRef PolyID, const FVector& TestPt, FVector& PointOnPoly) const;
 	/** Decode poly ID into tile index and poly index */
-	bool GetPolyTileIndex(NavNodeRef PolyID, uint32& PolyIndex, uint32& TileIndex) const;
+	NAVIGATIONSYSTEM_API bool GetPolyTileIndex(NavNodeRef PolyID, uint32& PolyIndex, uint32& TileIndex) const;
 	/** Decode poly ID into FNavTileRef and poly index */
-	bool GetPolyTileRef(NavNodeRef PolyId, uint32& OutPolyIndex, FNavTileRef& OutTileRef) const;
+	NAVIGATIONSYSTEM_API bool GetPolyTileRef(NavNodeRef PolyId, uint32& OutPolyIndex, FNavTileRef& OutTileRef) const;
 	/** Retrieves user ID for given offmesh link poly */
 	UE_DEPRECATED(5.3, "Please use GetNavLinkUserId() instead. This function only returns Invalid.")
 	uint32 GetLinkUserId(NavNodeRef LinkPolyID) const
 	{
 		return FNavLinkId::Invalid.GetId();
 	}
-	FNavLinkId GetNavLinkUserId(NavNodeRef LinkPolyID) const;
+	NAVIGATIONSYSTEM_API FNavLinkId GetNavLinkUserId(NavNodeRef LinkPolyID) const;
 	/** Retrieves start and end point of offmesh link */
-	bool GetLinkEndPoints(NavNodeRef LinkPolyID, FVector& PointA, FVector& PointB) const;
+	NAVIGATIONSYSTEM_API bool GetLinkEndPoints(NavNodeRef LinkPolyID, FVector& PointA, FVector& PointB) const;
 	/** Check if poly is a custom link */
-	bool IsCustomLink(NavNodeRef PolyRef) const;
+	NAVIGATIONSYSTEM_API bool IsCustomLink(NavNodeRef PolyRef) const;
 
 #if	WITH_NAVMESH_CLUSTER_LINKS
 	/** Retrieves bounds of cluster. Returns false on error. */
-	bool GetClusterBounds(NavNodeRef ClusterRef, FBox& OutBounds) const;
-	NavNodeRef GetClusterRefFromPolyRef(const NavNodeRef PolyRef) const;
+	NAVIGATIONSYSTEM_API bool GetClusterBounds(NavNodeRef ClusterRef, FBox& OutBounds) const;
+	NAVIGATIONSYSTEM_API NavNodeRef GetClusterRefFromPolyRef(const NavNodeRef PolyRef) const;
 #endif // WITH_NAVMESH_CLUSTER_LINKS
 
 	uint32 GetTileIndexFromPolyRef(const NavNodeRef PolyRef) const { return DetourNavMesh != NULL ? DetourNavMesh->decodePolyIdTile(PolyRef) : uint32(INDEX_NONE); }
 
-	static uint16 GetFilterForbiddenFlags(const FRecastQueryFilter* Filter);
-	static void SetFilterForbiddenFlags(FRecastQueryFilter* Filter, uint16 ForbiddenFlags);
+	static NAVIGATIONSYSTEM_API uint16 GetFilterForbiddenFlags(const FRecastQueryFilter* Filter);
+	static NAVIGATIONSYSTEM_API void SetFilterForbiddenFlags(FRecastQueryFilter* Filter, uint16 ForbiddenFlags);
 
-	void OnAreaCostChanged();
+	NAVIGATIONSYSTEM_API void OnAreaCostChanged();
 
 public:
 	dtNavMesh const* GetRecastMesh() const { return DetourNavMesh; };
 	dtNavMesh* GetRecastMesh() { return DetourNavMesh; };
-	void ReleaseDetourNavMesh();
+	NAVIGATIONSYSTEM_API void ReleaseDetourNavMesh();
 
-	void RemoveTileCacheLayers(int32 TileX, int32 TileY);
-	void RemoveTileCacheLayer(int32 TileX, int32 TileY, int32 LayerIdx);
-	void AddTileCacheLayers(int32 TileX, int32 TileY, const TArray<FNavMeshTileData>& Layers);
-	void AddTileCacheLayer(int32 TileX, int32 TileY, int32 LayerIdx, const FNavMeshTileData& LayerData);
-	void MarkEmptyTileCacheLayers(int32 TileX, int32 TileY);
-	FNavMeshTileData GetTileCacheLayer(int32 TileX, int32 TileY, int32 LayerIdx) const;
-	TArray<FNavMeshTileData> GetTileCacheLayers(int32 TileX, int32 TileY) const;
-	bool HasTileCacheLayers(int32 TileX, int32 TileY) const;
+	NAVIGATIONSYSTEM_API void RemoveTileCacheLayers(int32 TileX, int32 TileY);
+	NAVIGATIONSYSTEM_API void RemoveTileCacheLayer(int32 TileX, int32 TileY, int32 LayerIdx);
+	NAVIGATIONSYSTEM_API void AddTileCacheLayers(int32 TileX, int32 TileY, const TArray<FNavMeshTileData>& Layers);
+	NAVIGATIONSYSTEM_API void AddTileCacheLayer(int32 TileX, int32 TileY, int32 LayerIdx, const FNavMeshTileData& LayerData);
+	NAVIGATIONSYSTEM_API void MarkEmptyTileCacheLayers(int32 TileX, int32 TileY);
+	NAVIGATIONSYSTEM_API FNavMeshTileData GetTileCacheLayer(int32 TileX, int32 TileY, int32 LayerIdx) const;
+	NAVIGATIONSYSTEM_API TArray<FNavMeshTileData> GetTileCacheLayers(int32 TileX, int32 TileY) const;
+	NAVIGATIONSYSTEM_API bool HasTileCacheLayers(int32 TileX, int32 TileY) const;
 
 	/** Assigns recast generated navmesh to this instance.
 	 *	@param bOwnData if true from now on this FPImplRecastNavMesh instance will be responsible for this piece 
 	 *		of memory
 	 */
-	void SetRecastMesh(dtNavMesh* NavMesh);
+	NAVIGATIONSYSTEM_API void SetRecastMesh(dtNavMesh* NavMesh);
 
-	float GetTotalDataSize() const;
+	NAVIGATIONSYSTEM_API float GetTotalDataSize() const;
 
 	/** Gets the size of the compressed tile cache, this is slow */
 #if !UE_BUILD_SHIPPING
-	int32 GetCompressedTileCacheSize();
+	NAVIGATIONSYSTEM_API int32 GetCompressedTileCacheSize();
 #endif
 
 	/** Called on world origin changes */
-	void ApplyWorldOffset(const FVector& InOffset, bool bWorldShift);
+	NAVIGATIONSYSTEM_API void ApplyWorldOffset(const FVector& InOffset, bool bWorldShift);
 
 	/** calculated cost of given segment if traversed on specified poly. Function measures distance between specified points
 	 *	and returns cost of traversing this distance on given poly.
 	 *	@note no check if segment is on poly is performed. */
-	FVector::FReal CalcSegmentCostOnPoly(NavNodeRef PolyID, const dtQueryFilter* Filter, const FVector& StartLoc, const FVector& EndLoc) const;
+	NAVIGATIONSYSTEM_API FVector::FReal CalcSegmentCostOnPoly(NavNodeRef PolyID, const dtQueryFilter* Filter, const FVector& StartLoc, const FVector& EndLoc) const;
 
 	ARecastNavMesh* NavMeshOwner;
 	
@@ -273,37 +273,37 @@ public:
 	mutable dtNavMeshQuery SharedNavQuery;
 
 	/** Helper function to serialize a single Recast tile. */
-	static void SerializeRecastMeshTile(FArchive& Ar, int32 NavMeshVersion, unsigned char*& TileData, int32& TileDataSize);
+	static NAVIGATIONSYSTEM_API void SerializeRecastMeshTile(FArchive& Ar, int32 NavMeshVersion, unsigned char*& TileData, int32& TileDataSize);
 
 	/** Helper function to serialize a Recast tile compressed data. */
-	static void SerializeCompressedTileCacheData(FArchive& Ar, int32 NavMeshVersion, unsigned char*& CompressedData, int32& CompressedDataSize);
+	static NAVIGATIONSYSTEM_API void SerializeCompressedTileCacheData(FArchive& Ar, int32 NavMeshVersion, unsigned char*& CompressedData, int32& CompressedDataSize);
 
 	/** Initialize data for pathfinding */
-	bool InitPathfinding(const FVector& UnrealStart, const FVector& UnrealEnd, 
+	NAVIGATIONSYSTEM_API bool InitPathfinding(const FVector& UnrealStart, const FVector& UnrealEnd, 
 		const dtNavMeshQuery& Query, const dtQueryFilter* Filter,
 		FVector& RecastStart, dtPolyRef& StartPoly,
 		FVector& RecastEnd, dtPolyRef& EndPoly) const;
 
 	/** Marks path flags, perform string pulling if needed */
-	void PostProcessPath(dtStatus PathfindResult, FNavMeshPath& Path,
+	NAVIGATIONSYSTEM_API void PostProcessPath(dtStatus PathfindResult, FNavMeshPath& Path,
 		const dtNavMeshQuery& Query, const dtQueryFilter* Filter,
 		NavNodeRef StartNode, NavNodeRef EndNode,
 		FVector UnrealStart, FVector UnrealEnd,
 		FVector RecastStart, FVector RecastEnd,
 		dtQueryResult& PathResult) const;
 
-	void GetDebugPolyEdges(const dtMeshTile& Tile, bool bInternalEdges, bool bNavMeshEdges, TArray<FVector>& InternalEdgeVerts, TArray<FVector>& NavMeshEdgeVerts) const;
+	NAVIGATIONSYSTEM_API void GetDebugPolyEdges(const dtMeshTile& Tile, bool bInternalEdges, bool bNavMeshEdges, TArray<FVector>& InternalEdgeVerts, TArray<FVector>& NavMeshEdgeVerts) const;
 
 	/** workhorse function finding portal edges between corridor polys */
-	void GetEdgesForPathCorridorImpl(const TArray<NavNodeRef>* PathCorridor, TArray<FNavigationPortalEdge>* PathCorridorEdges, const dtNavMeshQuery& NavQuery) const;
+	NAVIGATIONSYSTEM_API void GetEdgesForPathCorridorImpl(const TArray<NavNodeRef>* PathCorridor, TArray<FNavigationPortalEdge>* PathCorridorEdges, const dtNavMeshQuery& NavQuery) const;
 
 protected:
 	/** 
 	 *	@param ForbiddenFlags polys with flags matching the filter will get added to 
 	 */
-	int32 GetTilesDebugGeometry(const FRecastNavMeshGenerator* Generator, const dtMeshTile& Tile, int32 VertBase, FRecastDebugGeometry& OutGeometry, int32 TileIdx = INDEX_NONE, uint16 ForbiddenFlags = 0) const;
+	NAVIGATIONSYSTEM_API int32 GetTilesDebugGeometry(const FRecastNavMeshGenerator* Generator, const dtMeshTile& Tile, int32 VertBase, FRecastDebugGeometry& OutGeometry, int32 TileIdx = INDEX_NONE, uint16 ForbiddenFlags = 0) const;
 
-	ENavigationQueryResult::Type PostProcessPathInternal(dtStatus FindPathStatus, FNavMeshPath& Path, 
+	NAVIGATIONSYSTEM_API ENavigationQueryResult::Type PostProcessPathInternal(dtStatus FindPathStatus, FNavMeshPath& Path, 
 		const dtNavMeshQuery& NavQuery, const dtQueryFilter* QueryFilter, 
 		NavNodeRef StartPolyID, NavNodeRef EndPolyID, 
 		const FVector& RecastStartPos, const FVector& RecastEndPos, 
diff --git a/Engine/Source/Runtime/NavigationSystem/Public/NavMesh/RecastNavMesh.h b/Engine/Source/Runtime/NavigationSystem/Public/NavMesh/RecastNavMesh.h
--- a/Engine/Source/Runtime/NavigationSystem/Public/NavMesh/RecastNavMesh.h
+++ b/Engine/Source/Runtime/NavigationSystem/Public/NavMesh/RecastNavMesh.h
@@ -86,10 +86,10 @@ struct FDetourTileSizeInfo
 	unsigned short OffMeshBase = 0;
 };
 
-struct NAVIGATIONSYSTEM_API FDetourTileLayout
+struct FDetourTileLayout
 {
-	FDetourTileLayout(const dtMeshTile& tile);
-	FDetourTileLayout(const FDetourTileSizeInfo& SizeInfo);
+	NAVIGATIONSYSTEM_API FDetourTileLayout(const dtMeshTile& tile);
+	NAVIGATIONSYSTEM_API FDetourTileLayout(const FDetourTileSizeInfo& SizeInfo);
 
 private:
 	void InitFromSizeInfo(const FDetourTileSizeInfo& SizeInfo);
@@ -316,7 +316,7 @@ struct FNavigationWallEdge
 #endif //WITH_RECAST
 
 USTRUCT(meta=(Deprecated = "5.2"))
-struct NAVIGATIONSYSTEM_API FRecastNavMeshGenerationProperties
+struct FRecastNavMeshGenerationProperties
 {
 	GENERATED_BODY()
 
@@ -425,8 +425,8 @@ struct NAVIGATIONSYSTEM_API FRecastNavMeshGenerationProperties
 	UPROPERTY(EditAnywhere, Category = Generation)
 	uint32 bIsWorldPartitioned : 1;
 	
-	FRecastNavMeshGenerationProperties();
-	FRecastNavMeshGenerationProperties(const ARecastNavMesh& RecastNavMesh);
+	NAVIGATIONSYSTEM_API FRecastNavMeshGenerationProperties();
+	NAVIGATIONSYSTEM_API FRecastNavMeshGenerationProperties(const ARecastNavMesh& RecastNavMesh);
 };
 
 UENUM()
@@ -437,11 +437,11 @@ enum class EHeightFieldRenderMode : uint8
 };
 
 USTRUCT()
-struct NAVIGATIONSYSTEM_API FRecastNavMeshTileGenerationDebug
+struct FRecastNavMeshTileGenerationDebug
 {
 	GENERATED_BODY()
 	
-	FRecastNavMeshTileGenerationDebug();
+	NAVIGATIONSYSTEM_API FRecastNavMeshTileGenerationDebug();
 
 	/** If set, the selected internal debug data will be kept during tile generation to be displayed with the navmesh. */
 	UPROPERTY(Transient, EditAnywhere, Category = Debug)
@@ -611,7 +611,7 @@ namespace FNavMeshConfig
 }
 
 USTRUCT()
-struct NAVIGATIONSYSTEM_API FNavMeshResolutionParam
+struct FNavMeshResolutionParam
 {
 	GENERATED_BODY()
 
@@ -630,17 +630,17 @@ struct NAVIGATIONSYSTEM_API FNavMeshResolutionParam
 	float AgentMaxStepHeight = 35.f;
 };
 
-UCLASS(config=Engine, defaultconfig, hidecategories=(Input,Rendering,Tags,Transformation,Actor,Layers,Replication), notplaceable)
-class NAVIGATIONSYSTEM_API ARecastNavMesh : public ANavigationData
+UCLASS(config=Engine, defaultconfig, hidecategories=(Input,Rendering,Tags,Transformation,Actor,Layers,Replication), notplaceable, MinimalAPI)
+class ARecastNavMesh : public ANavigationData
 {
 	GENERATED_UCLASS_BODY()
 
 	typedef uint16 FNavPolyFlags;
 
-	virtual void Serialize( FArchive& Ar ) override;
+	NAVIGATIONSYSTEM_API virtual void Serialize( FArchive& Ar ) override;
 
 #if WITH_EDITOR
-	virtual bool CanEditChange(const FProperty* InProperty) const override;
+	NAVIGATIONSYSTEM_API virtual bool CanEditChange(const FProperty* InProperty) const override;
 #endif
 
 	/** Draw edges of every navmesh's triangle */
@@ -934,13 +934,13 @@ private:
 	float TileSetUpdateInterval;
 	
 	/** contains last available dtPoly's flag bit set (8th bit at the moment of writing) */
-	static FNavPolyFlags NavLinkFlag;
+	static NAVIGATIONSYSTEM_API FNavPolyFlags NavLinkFlag;
 
 	/** Squared draw distance */
-	static FVector::FReal DrawDistanceSq;
+	static NAVIGATIONSYSTEM_API FVector::FReal DrawDistanceSq;
 
 	/** MinimumSizeForChaosNavMeshInfluence*/
-	static float MinimumSizeForChaosNavMeshInfluenceSq;
+	static NAVIGATIONSYSTEM_API float MinimumSizeForChaosNavMeshInfluenceSq;
 
 public:
 
@@ -996,10 +996,10 @@ public:
 
 	//////////////////////////////////////////////////////////////////////////
 
-	bool HasValidNavmesh() const;
+	NAVIGATIONSYSTEM_API bool HasValidNavmesh() const;
 
 	/** Dtor */
-	virtual ~ARecastNavMesh();
+	NAVIGATIONSYSTEM_API virtual ~ARecastNavMesh();
 
 #if WITH_RECAST
 	//----------------------------------------------------------------------//
@@ -1008,14 +1008,14 @@ public:
 public:
 
 	//~ Begin UObject Interface
-	virtual void PostInitProperties() override;
-	virtual void PostLoad() override;
-	virtual void PostRegisterAllComponents() override;
-	virtual void BeginDestroy() override;
+	NAVIGATIONSYSTEM_API virtual void PostInitProperties() override;
+	NAVIGATIONSYSTEM_API virtual void PostLoad() override;
+	NAVIGATIONSYSTEM_API virtual void PostRegisterAllComponents() override;
+	NAVIGATIONSYSTEM_API virtual void BeginDestroy() override;
 
 #if WITH_EDITOR
-	virtual void PostEditChangeChainProperty( struct FPropertyChangedChainEvent& PropertyChangedChainEvent) override;
-	virtual void PostEditChangeProperty( struct FPropertyChangedEvent& PropertyChangedEvent) override;
+	NAVIGATIONSYSTEM_API virtual void PostEditChangeChainProperty( struct FPropertyChangedChainEvent& PropertyChangedChainEvent) override;
+	NAVIGATIONSYSTEM_API virtual void PostEditChangeProperty( struct FPropertyChangedEvent& PropertyChangedEvent) override;
 #endif // WITH_EDITOR
 	//~ End UObject Interface
 
@@ -1024,70 +1024,70 @@ public:
 	virtual bool ShouldExport() override { return false; }
 #endif
 
-	virtual void LoadBeforeGeneratorRebuild() override;
+	NAVIGATIONSYSTEM_API virtual void LoadBeforeGeneratorRebuild() override;
 	
-	virtual void CleanUp() override;
+	NAVIGATIONSYSTEM_API virtual void CleanUp() override;
 
 	//~ Begin ANavigationData Interface
-	virtual FNavLocation GetRandomPoint(FSharedConstNavQueryFilter Filter = NULL, const UObject* Querier = NULL) const override;
+	NAVIGATIONSYSTEM_API virtual FNavLocation GetRandomPoint(FSharedConstNavQueryFilter Filter = NULL, const UObject* Querier = NULL) const override;
 	/** finds a random location in Radius, reachable from Origin.
 	 *  @param Radius needs to be non-negative. The function fails for Radius < 0. Radius being 0 is still rasults in a valid request. */
-	virtual bool GetRandomReachablePointInRadius(const FVector& Origin, float Radius, FNavLocation& OutResult, FSharedConstNavQueryFilter Filter = NULL, const UObject* Querier = NULL) const override;
-	virtual bool GetRandomPointInNavigableRadius(const FVector& Origin, float Radius, FNavLocation& OutResult, FSharedConstNavQueryFilter Filter = NULL, const UObject* Querier = NULL) const override;
+	NAVIGATIONSYSTEM_API virtual bool GetRandomReachablePointInRadius(const FVector& Origin, float Radius, FNavLocation& OutResult, FSharedConstNavQueryFilter Filter = NULL, const UObject* Querier = NULL) const override;
+	NAVIGATIONSYSTEM_API virtual bool GetRandomPointInNavigableRadius(const FVector& Origin, float Radius, FNavLocation& OutResult, FSharedConstNavQueryFilter Filter = NULL, const UObject* Querier = NULL) const override;
 
-	virtual bool FindMoveAlongSurface(const FNavLocation& StartLocation, const FVector& TargetPosition, FNavLocation& OutLocation, FSharedConstNavQueryFilter Filter = NULL, const UObject* Querier = NULL) const override;
-	virtual bool FindOverlappingEdges(const FNavLocation& StartLocation, TConstArrayView<FVector> ConvexPolygon, TArray<FVector>& OutEdges, FSharedConstNavQueryFilter Filter = NULL, const UObject* Querier = NULL) const override;
-	virtual bool GetPathSegmentBoundaryEdges(const FNavigationPath& Path, const FNavPathPoint& StartPoint, const FNavPathPoint& EndPoint, const TConstArrayView<FVector> SearchArea, TArray<FVector>& OutEdges, const float MaxAreaEnterCost, FSharedConstNavQueryFilter Filter = NULL, const UObject* Querier = NULL) const override;
-	virtual bool ProjectPoint(const FVector& Point, FNavLocation& OutLocation, const FVector& Extent, FSharedConstNavQueryFilter Filter = NULL, const UObject* Querier = NULL) const override;
-	virtual bool IsNodeRefValid(NavNodeRef NodeRef) const override;
+	NAVIGATIONSYSTEM_API virtual bool FindMoveAlongSurface(const FNavLocation& StartLocation, const FVector& TargetPosition, FNavLocation& OutLocation, FSharedConstNavQueryFilter Filter = NULL, const UObject* Querier = NULL) const override;
+	NAVIGATIONSYSTEM_API virtual bool FindOverlappingEdges(const FNavLocation& StartLocation, TConstArrayView<FVector> ConvexPolygon, TArray<FVector>& OutEdges, FSharedConstNavQueryFilter Filter = NULL, const UObject* Querier = NULL) const override;
+	NAVIGATIONSYSTEM_API virtual bool GetPathSegmentBoundaryEdges(const FNavigationPath& Path, const FNavPathPoint& StartPoint, const FNavPathPoint& EndPoint, const TConstArrayView<FVector> SearchArea, TArray<FVector>& OutEdges, const float MaxAreaEnterCost, FSharedConstNavQueryFilter Filter = NULL, const UObject* Querier = NULL) const override;
+	NAVIGATIONSYSTEM_API virtual bool ProjectPoint(const FVector& Point, FNavLocation& OutLocation, const FVector& Extent, FSharedConstNavQueryFilter Filter = NULL, const UObject* Querier = NULL) const override;
+	NAVIGATIONSYSTEM_API virtual bool IsNodeRefValid(NavNodeRef NodeRef) const override;
 
 	/** Project batch of points using shared search extent and filter */
-	virtual void BatchProjectPoints(TArray<FNavigationProjectionWork>& Workload, const FVector& Extent, FSharedConstNavQueryFilter Filter = NULL, const UObject* Querier = NULL) const override;
+	NAVIGATIONSYSTEM_API virtual void BatchProjectPoints(TArray<FNavigationProjectionWork>& Workload, const FVector& Extent, FSharedConstNavQueryFilter Filter = NULL, const UObject* Querier = NULL) const override;
 
 	/** Project batch of points using shared search filter. This version is not requiring user to pass in Extent, 
 	 *	and is instead relying on FNavigationProjectionWork.ProjectionLimit.
 	 *	@note function will assert if item's FNavigationProjectionWork.ProjectionLimit is invalid */
-	virtual void BatchProjectPoints(TArray<FNavigationProjectionWork>& Workload, FSharedConstNavQueryFilter Filter = NULL, const UObject* Querier = NULL) const override;
+	NAVIGATIONSYSTEM_API virtual void BatchProjectPoints(TArray<FNavigationProjectionWork>& Workload, FSharedConstNavQueryFilter Filter = NULL, const UObject* Querier = NULL) const override;
 	
-	virtual ENavigationQueryResult::Type CalcPathCost(const FVector& PathStart, const FVector& PathEnd, FVector::FReal& OutPathCost, FSharedConstNavQueryFilter Filter = NULL, const UObject* Querier = NULL) const override;
-	virtual ENavigationQueryResult::Type CalcPathLength(const FVector& PathStart, const FVector& PathEnd, FVector::FReal& OutPathLength, FSharedConstNavQueryFilter QueryFilter = NULL, const UObject* Querier = NULL) const override;
-	virtual ENavigationQueryResult::Type CalcPathLengthAndCost(const FVector& PathStart, const FVector& PathEnd, FVector::FReal& OutPathLength, FVector::FReal& OutPathCost, FSharedConstNavQueryFilter QueryFilter = NULL, const UObject* Querier = NULL) const override;
-	virtual bool DoesNodeContainLocation(NavNodeRef NodeRef, const FVector& WorldSpaceLocation) const override;
+	NAVIGATIONSYSTEM_API virtual ENavigationQueryResult::Type CalcPathCost(const FVector& PathStart, const FVector& PathEnd, FVector::FReal& OutPathCost, FSharedConstNavQueryFilter Filter = NULL, const UObject* Querier = NULL) const override;
+	NAVIGATIONSYSTEM_API virtual ENavigationQueryResult::Type CalcPathLength(const FVector& PathStart, const FVector& PathEnd, FVector::FReal& OutPathLength, FSharedConstNavQueryFilter QueryFilter = NULL, const UObject* Querier = NULL) const override;
+	NAVIGATIONSYSTEM_API virtual ENavigationQueryResult::Type CalcPathLengthAndCost(const FVector& PathStart, const FVector& PathEnd, FVector::FReal& OutPathLength, FVector::FReal& OutPathCost, FSharedConstNavQueryFilter QueryFilter = NULL, const UObject* Querier = NULL) const override;
+	NAVIGATIONSYSTEM_API virtual bool DoesNodeContainLocation(NavNodeRef NodeRef, const FVector& WorldSpaceLocation) const override;
 
-	virtual UPrimitiveComponent* ConstructRenderingComponent() override;
+	NAVIGATIONSYSTEM_API virtual UPrimitiveComponent* ConstructRenderingComponent() override;
 	/** Returns bounding box for the navmesh. */
 	virtual FBox GetBounds() const override { return GetNavMeshBounds(); }
 	/** Called on world origin changes **/
-	virtual void ApplyWorldOffset(const FVector& InOffset, bool bWorldShift) override;
+	NAVIGATIONSYSTEM_API virtual void ApplyWorldOffset(const FVector& InOffset, bool bWorldShift) override;
 
-	virtual void FillNavigationDataChunkActor(const FBox& InQueryBounds, class ANavigationDataChunkActor& DataChunkActor, FBox& OutTilesBounds) const override;
+	NAVIGATIONSYSTEM_API virtual void FillNavigationDataChunkActor(const FBox& InQueryBounds, class ANavigationDataChunkActor& DataChunkActor, FBox& OutTilesBounds) const override;
 
-	virtual void OnStreamingNavDataAdded(class ANavigationDataChunkActor& InActor) override;
-	virtual void OnStreamingNavDataRemoved(class ANavigationDataChunkActor& InActor) override;
+	NAVIGATIONSYSTEM_API virtual void OnStreamingNavDataAdded(class ANavigationDataChunkActor& InActor) override;
+	NAVIGATIONSYSTEM_API virtual void OnStreamingNavDataRemoved(class ANavigationDataChunkActor& InActor) override;
 	
-	virtual void OnStreamingLevelAdded(ULevel* InLevel, UWorld* InWorld) override;
-	virtual void OnStreamingLevelRemoved(ULevel* InLevel, UWorld* InWorld) override;
+	NAVIGATIONSYSTEM_API virtual void OnStreamingLevelAdded(ULevel* InLevel, UWorld* InWorld) override;
+	NAVIGATIONSYSTEM_API virtual void OnStreamingLevelRemoved(ULevel* InLevel, UWorld* InWorld) override;
 
 #if WITH_EDITOR
-	virtual double GetWorldPartitionNavigationDataBuilderOverlap() const override;
+	NAVIGATIONSYSTEM_API virtual double GetWorldPartitionNavigationDataBuilderOverlap() const override;
 #endif
 	//~ End ANavigationData Interface
 
-	virtual void AttachNavMeshDataChunk(URecastNavMeshDataChunk& NavDataChunk);
-	virtual void DetachNavMeshDataChunk(URecastNavMeshDataChunk& NavDataChunk);
+	NAVIGATIONSYSTEM_API virtual void AttachNavMeshDataChunk(URecastNavMeshDataChunk& NavDataChunk);
+	NAVIGATIONSYSTEM_API virtual void DetachNavMeshDataChunk(URecastNavMeshDataChunk& NavDataChunk);
 
-	const TArray<FIntPoint>& GetActiveTiles() const;
-	TArray<FIntPoint>& GetActiveTiles(); 
+	NAVIGATIONSYSTEM_API const TArray<FIntPoint>& GetActiveTiles() const;
+	NAVIGATIONSYSTEM_API TArray<FIntPoint>& GetActiveTiles(); 
 
-	void LogRecastTile(const TCHAR* Caller, const FName& Prefix, const FName& OperationName, const dtNavMesh& DetourMesh, const int32 TileX, const int32 TileY, const int32 LayerIndex, const uint64 TileRef) const;
+	NAVIGATIONSYSTEM_API void LogRecastTile(const TCHAR* Caller, const FName& Prefix, const FName& OperationName, const dtNavMesh& DetourMesh, const int32 TileX, const int32 TileY, const int32 LayerIndex, const uint64 TileRef) const;
 	
 protected:
 	/** Serialization helper. */
-	void SerializeRecastNavMesh(FArchive& Ar, FPImplRecastNavMesh*& NavMesh, int32 NavMeshVersion);
+	NAVIGATIONSYSTEM_API void SerializeRecastNavMesh(FArchive& Ar, FPImplRecastNavMesh*& NavMesh, int32 NavMeshVersion);
 
-	virtual void RestrictBuildingToActiveTiles(bool InRestrictBuildingToActiveTiles) override;
+	NAVIGATIONSYSTEM_API virtual void RestrictBuildingToActiveTiles(bool InRestrictBuildingToActiveTiles) override;
 
-	virtual void OnRegistered() override;
+	NAVIGATIONSYSTEM_API virtual void OnRegistered() override;
 
 public:
 	/** Get the CellSize for the given resolution. */
@@ -1109,79 +1109,79 @@ public:
 	void SetAgentMaxStepHeight(const ENavigationDataResolution Resolution, const float MaxStepHeight) { NavMeshResolutionParams[(uint8)Resolution].AgentMaxStepHeight = MaxStepHeight; }
 	
 	/** Returns the tile size in world units. */
-	float GetTileSizeUU() const;
+	NAVIGATIONSYSTEM_API float GetTileSizeUU() const;
 	
 	/** Whether NavMesh should adjust its tile pool size when NavBounds are changed */
-	bool IsResizable() const;
+	NAVIGATIONSYSTEM_API bool IsResizable() const;
 
 	/** Returns bounding box for the whole navmesh. */
-	FBox GetNavMeshBounds() const;
+	NAVIGATIONSYSTEM_API FBox GetNavMeshBounds() const;
 
 	/** Returns bounding box for a given navmesh tile. */
-	FBox GetNavMeshTileBounds(int32 TileIndex) const;
+	NAVIGATIONSYSTEM_API FBox GetNavMeshTileBounds(int32 TileIndex) const;
 
 	/** Retrieves XY coordinates of tile specified by index */
-	bool GetNavMeshTileXY(int32 TileIndex, int32& OutX, int32& OutY, int32& Layer) const;
+	NAVIGATIONSYSTEM_API bool GetNavMeshTileXY(int32 TileIndex, int32& OutX, int32& OutY, int32& Layer) const;
 
 	/** Retrieves XY coordinates of tile specified by position */
-	bool GetNavMeshTileXY(const FVector& Point, int32& OutX, int32& OutY) const;
+	NAVIGATIONSYSTEM_API bool GetNavMeshTileXY(const FVector& Point, int32& OutX, int32& OutY) const;
 
 	/** Retrieves the tile resolution */
-	bool GetNavmeshTileResolution(int32 TileIndex, ENavigationDataResolution& OutResolution) const;
+	NAVIGATIONSYSTEM_API bool GetNavmeshTileResolution(int32 TileIndex, ENavigationDataResolution& OutResolution) const;
 
 	/** Checks the supplied Points tile indicies can fit in the range of an int32 */
-	bool CheckTileIndicesInValidRange(const FVector& Point, bool& bOutInRange) const;
+	NAVIGATIONSYSTEM_API bool CheckTileIndicesInValidRange(const FVector& Point, bool& bOutInRange) const;
 
 	/** Retrieves all tile indices at matching XY coordinates */
-	void GetNavMeshTilesAt(int32 TileX, int32 TileY, TArray<int32>& Indices) const;
+	NAVIGATIONSYSTEM_API void GetNavMeshTilesAt(int32 TileX, int32 TileY, TArray<int32>& Indices) const;
 
 	/** Retrieves number of tiles in this navmesh */
-	int32 GetNavMeshTilesCount() const;
+	NAVIGATIONSYSTEM_API int32 GetNavMeshTilesCount() const;
 
 	/** Removes compressed tile data at given tile coord */
-	void RemoveTileCacheLayers(int32 TileX, int32 TileY);
+	NAVIGATIONSYSTEM_API void RemoveTileCacheLayers(int32 TileX, int32 TileY);
 	
 	/** Stores compressed tile data for given tile coord */
-	void AddTileCacheLayers(int32 TileX, int32 TileY, const TArray<FNavMeshTileData>& InLayers);
+	NAVIGATIONSYSTEM_API void AddTileCacheLayers(int32 TileX, int32 TileY, const TArray<FNavMeshTileData>& InLayers);
 
 #if RECAST_INTERNAL_DEBUG_DATA
-	void RemoveTileDebugData(int32 TileX, int32 TileY);
-	void AddTileDebugData(int32 TileX, int32 TileY, const struct FRecastInternalDebugData& InTileDebugData);
+	NAVIGATIONSYSTEM_API void RemoveTileDebugData(int32 TileX, int32 TileY);
+	NAVIGATIONSYSTEM_API void AddTileDebugData(int32 TileX, int32 TileY, const struct FRecastInternalDebugData& InTileDebugData);
 #endif
 
 	/** Marks tile coord as rebuild and empty */
-	void MarkEmptyTileCacheLayers(int32 TileX, int32 TileY);
+	NAVIGATIONSYSTEM_API void MarkEmptyTileCacheLayers(int32 TileX, int32 TileY);
 	
 	/** Returns compressed tile data at given tile coord */
-	TArray<FNavMeshTileData> GetTileCacheLayers(int32 TileX, int32 TileY) const;
+	NAVIGATIONSYSTEM_API TArray<FNavMeshTileData> GetTileCacheLayers(int32 TileX, int32 TileY) const;
 
 	/** Gets the size of the compressed tile cache, this is slow */
 #if !UE_BUILD_SHIPPING
-	int32 GetCompressedTileCacheSize();
+	NAVIGATIONSYSTEM_API int32 GetCompressedTileCacheSize();
 #endif
 
-	void GetEdgesForPathCorridor(const TArray<NavNodeRef>* PathCorridor, TArray<struct FNavigationPortalEdge>* PathCorridorEdges) const;
+	NAVIGATIONSYSTEM_API void GetEdgesForPathCorridor(const TArray<NavNodeRef>* PathCorridor, TArray<struct FNavigationPortalEdge>* PathCorridorEdges) const;
 
-	void UpdateDrawing();
+	NAVIGATIONSYSTEM_API void UpdateDrawing();
 
 	/** Creates a task to be executed on GameThread calling UpdateDrawing */
-	void RequestDrawingUpdate(bool bForce = false);
+	NAVIGATIONSYSTEM_API void RequestDrawingUpdate(bool bForce = false);
 
 	/** called after regenerating tiles */
 	UE_DEPRECATED(5.1, "Use new version with FNavTileRef")
-	virtual void OnNavMeshTilesUpdated(const TArray<uint32>& ChangedTiles);
+	NAVIGATIONSYSTEM_API virtual void OnNavMeshTilesUpdated(const TArray<uint32>& ChangedTiles);
 
 	/** called after regenerating tiles */
-	virtual void OnNavMeshTilesUpdated(const TArray<FNavTileRef>& ChangedTiles);
+	NAVIGATIONSYSTEM_API virtual void OnNavMeshTilesUpdated(const TArray<FNavTileRef>& ChangedTiles);
 
 	/** Event from generator that navmesh build has finished */
-	virtual void OnNavMeshGenerationFinished();
+	NAVIGATIONSYSTEM_API virtual void OnNavMeshGenerationFinished();
 
-	virtual void EnsureBuildCompletion() override;
+	NAVIGATIONSYSTEM_API virtual void EnsureBuildCompletion() override;
 
-	virtual void SetConfig(const FNavDataConfig& Src) override;
+	NAVIGATIONSYSTEM_API virtual void SetConfig(const FNavDataConfig& Src) override;
 protected:
-	virtual void FillConfig(FNavDataConfig& Dest) override;
+	NAVIGATIONSYSTEM_API virtual void FillConfig(FNavDataConfig& Dest) override;
 
 	FORCEINLINE const FNavigationQueryFilter& GetRightFilterRef(FSharedConstNavQueryFilter Filter) const 
 	{
@@ -1190,52 +1190,52 @@ protected:
 
 public:
 
-	static bool IsVoxelCacheEnabled();
+	static NAVIGATIONSYSTEM_API bool IsVoxelCacheEnabled();
 
 	//----------------------------------------------------------------------//
 	// Debug                                                                
 	//----------------------------------------------------------------------//
 	/** Debug rendering support. */
 	UE_DEPRECATED(5.1, "Please use the new signature of GetDebugGeometryForTile()")
-	void GetDebugGeometry(FRecastDebugGeometry& OutGeometry, int32 TileIndex = INDEX_NONE) const;
+	NAVIGATIONSYSTEM_API void GetDebugGeometry(FRecastDebugGeometry& OutGeometry, int32 TileIndex = INDEX_NONE) const;
 
 	/* Gather debug geometry.
 	 * @params OutGeometry Output geometry.
 	 * @params TileIndex Used to collect geometry for a specific tile, INDEX_NONE will gather all tiles
 	 * @return True if done collecting.
 	 */
-	bool GetDebugGeometryForTile(FRecastDebugGeometry& OutGeometry, int32 TileIndex) const;
+	NAVIGATIONSYSTEM_API bool GetDebugGeometryForTile(FRecastDebugGeometry& OutGeometry, int32 TileIndex) const;
 
 	// @todo docuement
-	void DrawDebugPathCorridor(NavNodeRef const* PathPolys, int32 NumPathPolys, bool bPersistent=true) const;
+	NAVIGATIONSYSTEM_API void DrawDebugPathCorridor(NavNodeRef const* PathPolys, int32 NumPathPolys, bool bPersistent=true) const;
 
 #if !UE_BUILD_SHIPPING
-	virtual uint32 LogMemUsed() const override;
+	NAVIGATIONSYSTEM_API virtual uint32 LogMemUsed() const override;
 #endif // !UE_BUILD_SHIPPING
 
-	void UpdateNavMeshDrawing();
+	NAVIGATIONSYSTEM_API void UpdateNavMeshDrawing();
 
 	//----------------------------------------------------------------------//
 	// Utilities
 	//----------------------------------------------------------------------//
-	virtual void OnNavAreaChanged() override;
-	virtual void OnNavAreaAdded(const UClass* NavAreaClass, int32 AgentIndex) override;
-	virtual void OnNavAreaRemoved(const UClass* NavAreaClass) override;
-	virtual int32 GetNewAreaID(const UClass* AreaClass) const override;
+	NAVIGATIONSYSTEM_API virtual void OnNavAreaChanged() override;
+	NAVIGATIONSYSTEM_API virtual void OnNavAreaAdded(const UClass* NavAreaClass, int32 AgentIndex) override;
+	NAVIGATIONSYSTEM_API virtual void OnNavAreaRemoved(const UClass* NavAreaClass) override;
+	NAVIGATIONSYSTEM_API virtual int32 GetNewAreaID(const UClass* AreaClass) const override;
 	virtual int32 GetMaxSupportedAreas() const override { return RECAST_MAX_AREAS; }
 
 	/** Get forbidden area flags from default query filter */
-	uint16 GetDefaultForbiddenFlags() const;
+	NAVIGATIONSYSTEM_API uint16 GetDefaultForbiddenFlags() const;
 	/** Change forbidden area flags in default query filter */
-	void SetDefaultForbiddenFlags(uint16 ForbiddenAreaFlags);
+	NAVIGATIONSYSTEM_API void SetDefaultForbiddenFlags(uint16 ForbiddenAreaFlags);
 
 	/** Area sort function */
-	virtual void SortAreasForGenerator(TArray<FRecastAreaNavModifierElement>& Areas) const;
+	NAVIGATIONSYSTEM_API virtual void SortAreasForGenerator(TArray<FRecastAreaNavModifierElement>& Areas) const;
 
-	virtual void RecreateDefaultFilter();
+	NAVIGATIONSYSTEM_API virtual void RecreateDefaultFilter();
 
 	int32 GetMaxSimultaneousTileGenerationJobsCount() const { return MaxSimultaneousTileGenerationJobsCount; }
-	void SetMaxSimultaneousTileGenerationJobsCount(int32 NewJobsCountLimit);
+	NAVIGATIONSYSTEM_API void SetMaxSimultaneousTileGenerationJobsCount(int32 NewJobsCountLimit);
 
 	/** Returns query extent including adjustments for voxelization error compensation */
 	FVector GetModifiedQueryExtent(const FVector& QueryExtent) const
@@ -1248,17 +1248,17 @@ public:
 	// Custom navigation links
 	//----------------------------------------------------------------------//
 
-	virtual void UpdateCustomLink(const INavLinkCustomInterface* CustomLink) override;
+	NAVIGATIONSYSTEM_API virtual void UpdateCustomLink(const INavLinkCustomInterface* CustomLink) override;
 
 	UE_DEPRECATED(5.3, "Use version of this function that takes a FNavLinkId. This function now has no effect.")
 	void UpdateNavigationLinkArea(int32 UserId, TSubclassOf<UNavArea> AreaClass) const {}
 
 	/** update area class and poly flags for all offmesh links with given UserId */
-	void UpdateNavigationLinkArea(FNavLinkId UserId, TSubclassOf<UNavArea> AreaClass) const;
+	NAVIGATIONSYSTEM_API void UpdateNavigationLinkArea(FNavLinkId UserId, TSubclassOf<UNavArea> AreaClass) const;
 
 #if WITH_NAVMESH_SEGMENT_LINKS
 	/** update area class and poly flags for all offmesh segment links with given UserId */
-	void UpdateSegmentLinkArea(int32 UserId, TSubclassOf<UNavArea> AreaClass) const;
+	NAVIGATIONSYSTEM_API void UpdateSegmentLinkArea(int32 UserId, TSubclassOf<UNavArea> AreaClass) const;
 #endif // WITH_NAVMESH_SEGMENT_LINKS
 
 	//----------------------------------------------------------------------//
@@ -1266,20 +1266,20 @@ public:
 	//----------------------------------------------------------------------//
 
 	/** Starts batch processing and locks access to navmesh from other threads */
-	virtual void BeginBatchQuery() const override;
+	NAVIGATIONSYSTEM_API virtual void BeginBatchQuery() const override;
 
 	/** Finishes batch processing and release locks */
-	virtual void FinishBatchQuery() const override;
+	NAVIGATIONSYSTEM_API virtual void FinishBatchQuery() const override;
 
 	//----------------------------------------------------------------------//
 	// Querying                                                                
 	//----------------------------------------------------------------------//
 	
 	/** dtNavMesh getter */
-	dtNavMesh* GetRecastMesh();
+	NAVIGATIONSYSTEM_API dtNavMesh* GetRecastMesh();
 
 	/** dtNavMesh getter */
-	const dtNavMesh* GetRecastMesh() const;
+	NAVIGATIONSYSTEM_API const dtNavMesh* GetRecastMesh() const;
 
 	UE_DEPRECATED(5.3, "Please use GetNavLinkUserId() instead. This function only returns Invalid.")
 	int32 GetLinkUserId(NavNodeRef LinkPolyID) const
@@ -1288,186 +1288,186 @@ public:
 	}
 
 	/** Retrieves LinkUserID associated with indicated PolyID */
-	FNavLinkId GetNavLinkUserId(NavNodeRef LinkPolyID) const;
+	NAVIGATIONSYSTEM_API FNavLinkId GetNavLinkUserId(NavNodeRef LinkPolyID) const;
 
-	FColor GetAreaIDColor(uint8 AreaID) const;
+	NAVIGATIONSYSTEM_API FColor GetAreaIDColor(uint8 AreaID) const;
 
 	/** Finds the polygons along the navigation graph that touch the specified circle. */
-	bool FindPolysAroundCircle(const FVector& CenterPos, const NavNodeRef CenterNodeRef, const FVector::FReal Radius, const FSharedConstNavQueryFilter& Filter, const UObject* QueryOwner, TArray<NavNodeRef>* OutPolys = nullptr, TArray<NavNodeRef>* OutPolysParent = nullptr, TArray<float>* OutPolysCost = nullptr, int32* OutPolysCount = nullptr) const;
+	NAVIGATIONSYSTEM_API bool FindPolysAroundCircle(const FVector& CenterPos, const NavNodeRef CenterNodeRef, const FVector::FReal Radius, const FSharedConstNavQueryFilter& Filter, const UObject* QueryOwner, TArray<NavNodeRef>* OutPolys = nullptr, TArray<NavNodeRef>* OutPolysParent = nullptr, TArray<float>* OutPolysCost = nullptr, int32* OutPolysCount = nullptr) const;
 	
 	/** Returns nearest navmesh polygon to Loc, or INVALID_NAVMESHREF if Loc is not on the navmesh. */
-	NavNodeRef FindNearestPoly(FVector const& Loc, FVector const& Extent, FSharedConstNavQueryFilter Filter = NULL, const UObject* Querier = NULL) const;
+	NAVIGATIONSYSTEM_API NavNodeRef FindNearestPoly(FVector const& Loc, FVector const& Extent, FSharedConstNavQueryFilter Filter = NULL, const UObject* Querier = NULL) const;
 
 	/** Finds the distance to the closest wall, limited to MaxDistance
 	 *	[out] OutClosestPointOnWall, if supplied, will be set to closest point on closest wall. Will not be set if no wall in the area (return value 0.f) */
-	FVector::FReal FindDistanceToWall(const FVector& StartLoc, FSharedConstNavQueryFilter Filter = nullptr, FVector::FReal MaxDistance = TNumericLimits<FVector::FReal>::Max(), FVector* OutClosestPointOnWall = nullptr) const;
+	NAVIGATIONSYSTEM_API FVector::FReal FindDistanceToWall(const FVector& StartLoc, FSharedConstNavQueryFilter Filter = nullptr, FVector::FReal MaxDistance = TNumericLimits<FVector::FReal>::Max(), FVector* OutClosestPointOnWall = nullptr) const;
 
 	/** Retrieves center of the specified polygon. Returns false on error. */
-	bool GetPolyCenter(NavNodeRef PolyID, FVector& OutCenter) const;
+	NAVIGATIONSYSTEM_API bool GetPolyCenter(NavNodeRef PolyID, FVector& OutCenter) const;
 
 	/** Retrieves the vertices for the specified polygon. Returns false on error. */
-	bool GetPolyVerts(NavNodeRef PolyID, TArray<FVector>& OutVerts) const;
+	NAVIGATIONSYSTEM_API bool GetPolyVerts(NavNodeRef PolyID, TArray<FVector>& OutVerts) const;
 
 	/** Retrieves a random point inside the specified polygon. Returns false on error. */
-	bool GetRandomPointInPoly(NavNodeRef PolyID, FVector& OutPoint) const;
+	NAVIGATIONSYSTEM_API bool GetRandomPointInPoly(NavNodeRef PolyID, FVector& OutPoint) const;
 
 	/** Retrieves area ID for the specified polygon. */
-	uint32 GetPolyAreaID(NavNodeRef PolyID) const;
+	NAVIGATIONSYSTEM_API uint32 GetPolyAreaID(NavNodeRef PolyID) const;
 
 	/** Sets area ID for the specified polygon. */
-	bool SetPolyArea(NavNodeRef PolyID, TSubclassOf<UNavArea> AreaClass);
+	NAVIGATIONSYSTEM_API bool SetPolyArea(NavNodeRef PolyID, TSubclassOf<UNavArea> AreaClass);
 
 	/** Sets area ID for the specified polygons */
-	void SetPolyArrayArea(const TArray<FNavPoly>& Polys, TSubclassOf<UNavArea> AreaClass);
+	NAVIGATIONSYSTEM_API void SetPolyArrayArea(const TArray<FNavPoly>& Polys, TSubclassOf<UNavArea> AreaClass);
 
 	/** In given Bounds find all areas of class OldArea and replace them with NewArea
 	 *	@return number of polys touched */
-	int32 ReplaceAreaInTileBounds(const FBox& Bounds, TSubclassOf<UNavArea> OldArea, TSubclassOf<UNavArea> NewArea, bool ReplaceLinks = true, TArray<NavNodeRef>* OutTouchedNodes = nullptr);
+	NAVIGATIONSYSTEM_API int32 ReplaceAreaInTileBounds(const FBox& Bounds, TSubclassOf<UNavArea> OldArea, TSubclassOf<UNavArea> NewArea, bool ReplaceLinks = true, TArray<NavNodeRef>* OutTouchedNodes = nullptr);
 	
 	/** Retrieves poly and area flags for specified polygon */
-	bool GetPolyFlags(NavNodeRef PolyID, uint16& PolyFlags, uint16& AreaFlags) const;
-	bool GetPolyFlags(NavNodeRef PolyID, FNavMeshNodeFlags& Flags) const;
+	NAVIGATIONSYSTEM_API bool GetPolyFlags(NavNodeRef PolyID, uint16& PolyFlags, uint16& AreaFlags) const;
+	NAVIGATIONSYSTEM_API bool GetPolyFlags(NavNodeRef PolyID, FNavMeshNodeFlags& Flags) const;
 
 	/** Finds all polys connected with specified one */
-	bool GetPolyNeighbors(NavNodeRef PolyID, TArray<FNavigationPortalEdge>& Neighbors) const;
+	NAVIGATIONSYSTEM_API bool GetPolyNeighbors(NavNodeRef PolyID, TArray<FNavigationPortalEdge>& Neighbors) const;
 
 	/** Finds all polys connected with specified one, results expressed as array of NavNodeRefs */
-	bool GetPolyNeighbors(NavNodeRef PolyID, TArray<NavNodeRef>& Neighbors) const;
+	NAVIGATIONSYSTEM_API bool GetPolyNeighbors(NavNodeRef PolyID, TArray<NavNodeRef>& Neighbors) const;
 
 	/** Finds edges of specified poly */
-	bool GetPolyEdges(NavNodeRef PolyID, TArray<FNavigationPortalEdge>& Neighbors) const;
+	NAVIGATIONSYSTEM_API bool GetPolyEdges(NavNodeRef PolyID, TArray<FNavigationPortalEdge>& Neighbors) const;
 
 	/** Finds closest point constrained to given poly */
-	bool GetClosestPointOnPoly(NavNodeRef PolyID, const FVector& TestPt, FVector& PointOnPoly) const;
+	NAVIGATIONSYSTEM_API bool GetClosestPointOnPoly(NavNodeRef PolyID, const FVector& TestPt, FVector& PointOnPoly) const;
 
 	/** Decode poly ID into tile index and poly index */
-	bool GetPolyTileIndex(NavNodeRef PolyID, uint32& PolyIndex, uint32& TileIndex) const;
+	NAVIGATIONSYSTEM_API bool GetPolyTileIndex(NavNodeRef PolyID, uint32& PolyIndex, uint32& TileIndex) const;
 
 	/** Retrieves start and end point of offmesh link */
-	bool GetLinkEndPoints(NavNodeRef LinkPolyID, FVector& PointA, FVector& PointB) const;
+	NAVIGATIONSYSTEM_API bool GetLinkEndPoints(NavNodeRef LinkPolyID, FVector& PointA, FVector& PointB) const;
 
 	/** Retrieves bounds of cluster. Returns false on error. */
-	bool GetClusterBounds(NavNodeRef ClusterRef, FBox& OutBounds) const;
+	NAVIGATIONSYSTEM_API bool GetClusterBounds(NavNodeRef ClusterRef, FBox& OutBounds) const;
 
 	/** Get random point in given cluster */
-	bool GetRandomPointInCluster(NavNodeRef ClusterRef, FNavLocation& OutLocation) const;
+	NAVIGATIONSYSTEM_API bool GetRandomPointInCluster(NavNodeRef ClusterRef, FNavLocation& OutLocation) const;
 
 	/** Get cluster ref containing given poly ref */
-	NavNodeRef GetClusterRef(NavNodeRef PolyRef) const;
+	NAVIGATIONSYSTEM_API NavNodeRef GetClusterRef(NavNodeRef PolyRef) const;
 
 	/** Retrieves all polys within given pathing distance from StartLocation.
 	 *	@NOTE query is not using string-pulled path distance (for performance reasons),
 	 *		it measured distance between middles of portal edges, do you might want to 
 	 *		add an extra margin to PathingDistance */
-	bool GetPolysWithinPathingDistance(FVector const& StartLoc, const FVector::FReal PathingDistance, TArray<NavNodeRef>& FoundPolys,
+	NAVIGATIONSYSTEM_API bool GetPolysWithinPathingDistance(FVector const& StartLoc, const FVector::FReal PathingDistance, TArray<NavNodeRef>& FoundPolys,
 		FSharedConstNavQueryFilter Filter = nullptr, const UObject* Querier = nullptr, FRecastDebugPathfindingData* DebugData = nullptr) const;
 
 	/** Filters nav polys in PolyRefs with Filter */
-	bool FilterPolys(TArray<NavNodeRef>& PolyRefs, const FRecastQueryFilter* Filter, const UObject* Querier = NULL) const;
+	NAVIGATIONSYSTEM_API bool FilterPolys(TArray<NavNodeRef>& PolyRefs, const FRecastQueryFilter* Filter, const UObject* Querier = NULL) const;
 
 	/** Get all polys from tile */
-	bool GetPolysInTile(int32 TileIndex, TArray<FNavPoly>& Polys) const;
+	NAVIGATIONSYSTEM_API bool GetPolysInTile(int32 TileIndex, TArray<FNavPoly>& Polys) const;
 
 	/** Get up to 256 polys that overlap the specified box */
-	bool GetPolysInBox(const FBox& Box, TArray<FNavPoly>& Polys, FSharedConstNavQueryFilter Filter = nullptr, const UObject* Owner = nullptr) const;
+	NAVIGATIONSYSTEM_API bool GetPolysInBox(const FBox& Box, TArray<FNavPoly>& Polys, FSharedConstNavQueryFilter Filter = nullptr, const UObject* Owner = nullptr) const;
 
 	/** Find up to 64 navmesh eges in up to 64 polys around the center */
-	bool FindEdges(const NavNodeRef CenterNodeRef, const FVector Center, const FVector::FReal Radius, const FSharedConstNavQueryFilter Filter, TArray<FNavigationWallEdge>& OutEdges) const;
+	NAVIGATIONSYSTEM_API bool FindEdges(const NavNodeRef CenterNodeRef, const FVector Center, const FVector::FReal Radius, const FSharedConstNavQueryFilter Filter, TArray<FNavigationWallEdge>& OutEdges) const;
 
 	/** Get all polys from tile */
-	bool GetNavLinksInTile(const int32 TileIndex, TArray<FNavPoly>& Polys, const bool bIncludeLinksFromNeighborTiles) const;
+	NAVIGATIONSYSTEM_API bool GetNavLinksInTile(const int32 TileIndex, TArray<FNavPoly>& Polys, const bool bIncludeLinksFromNeighborTiles) const;
 
 	/** Projects point on navmesh, returning all hits along vertical line defined by min-max Z params */
-	bool ProjectPointMulti(const FVector& Point, TArray<FNavLocation>& OutLocations, const FVector& Extent,
+	NAVIGATIONSYSTEM_API bool ProjectPointMulti(const FVector& Point, TArray<FNavLocation>& OutLocations, const FVector& Extent,
 		FVector::FReal MinZ, FVector::FReal MaxZ, FSharedConstNavQueryFilter Filter = NULL, const UObject* Querier = NULL) const;
 	
 	// @todo docuement
-	static FPathFindingResult FindPath(const FNavAgentProperties& AgentProperties, const FPathFindingQuery& Query);
-	static bool TestPath(const FNavAgentProperties& AgentProperties, const FPathFindingQuery& Query, int32* NumVisitedNodes);
-	static bool TestHierarchicalPath(const FNavAgentProperties& AgentProperties, const FPathFindingQuery& Query, int32* NumVisitedNodes);
-	static bool NavMeshRaycast(const ANavigationData* Self, const FVector& RayStart, const FVector& RayEnd, FVector& HitLocation, FSharedConstNavQueryFilter QueryFilter, const UObject* Querier, FRaycastResult& Result);
-	static bool NavMeshRaycast(const ANavigationData* Self, const FVector& RayStart, const FVector& RayEnd, FVector& HitLocation, FSharedConstNavQueryFilter QueryFilter, const UObject* Querier = NULL);
-	static bool NavMeshRaycast(const ANavigationData* Self, NavNodeRef RayStartNode, const FVector& RayStart, const FVector& RayEnd, FVector& HitLocation, FSharedConstNavQueryFilter QueryFilter, const UObject* Querier = NULL);
+	static NAVIGATIONSYSTEM_API FPathFindingResult FindPath(const FNavAgentProperties& AgentProperties, const FPathFindingQuery& Query);
+	static NAVIGATIONSYSTEM_API bool TestPath(const FNavAgentProperties& AgentProperties, const FPathFindingQuery& Query, int32* NumVisitedNodes);
+	static NAVIGATIONSYSTEM_API bool TestHierarchicalPath(const FNavAgentProperties& AgentProperties, const FPathFindingQuery& Query, int32* NumVisitedNodes);
+	static NAVIGATIONSYSTEM_API bool NavMeshRaycast(const ANavigationData* Self, const FVector& RayStart, const FVector& RayEnd, FVector& HitLocation, FSharedConstNavQueryFilter QueryFilter, const UObject* Querier, FRaycastResult& Result);
+	static NAVIGATIONSYSTEM_API bool NavMeshRaycast(const ANavigationData* Self, const FVector& RayStart, const FVector& RayEnd, FVector& HitLocation, FSharedConstNavQueryFilter QueryFilter, const UObject* Querier = NULL);
+	static NAVIGATIONSYSTEM_API bool NavMeshRaycast(const ANavigationData* Self, NavNodeRef RayStartNode, const FVector& RayStart, const FVector& RayEnd, FVector& HitLocation, FSharedConstNavQueryFilter QueryFilter, const UObject* Querier = NULL);
 
-	virtual void BatchRaycast(TArray<FNavigationRaycastWork>& Workload, FSharedConstNavQueryFilter QueryFilter, const UObject* Querier = NULL) const override;
+	NAVIGATIONSYSTEM_API virtual void BatchRaycast(TArray<FNavigationRaycastWork>& Workload, FSharedConstNavQueryFilter QueryFilter, const UObject* Querier = NULL) const override;
 
 	/** finds a Filter-passing navmesh location closest to specified StartLoc
 	 *	@return true if adjusting was required, false otherwise */
-	bool AdjustLocationWithFilter(const FVector& StartLoc, FVector& OutAdjustedLocation, const FNavigationQueryFilter& Filter, const UObject* Querier = NULL) const;
+	NAVIGATIONSYSTEM_API bool AdjustLocationWithFilter(const FVector& StartLoc, FVector& OutAdjustedLocation, const FNavigationQueryFilter& Filter, const UObject* Querier = NULL) const;
 	
 	/** Check if navmesh is defined (either built/streamed or recognized as empty tile by generator) in given radius.
 	  * @returns true if ALL tiles inside are ready
 	  */
-	bool HasCompleteDataInRadius(const FVector& TestLocation, FVector::FReal TestRadius) const;
+	NAVIGATIONSYSTEM_API bool HasCompleteDataInRadius(const FVector& TestLocation, FVector::FReal TestRadius) const;
 
 	/** @return true is specified segment is fully on navmesh (respecting the optional filter) */
-	bool IsSegmentOnNavmesh(const FVector& SegmentStart, const FVector& SegmentEnd, FSharedConstNavQueryFilter Filter = NULL, const UObject* Querier = NULL) const;
+	NAVIGATIONSYSTEM_API bool IsSegmentOnNavmesh(const FVector& SegmentStart, const FVector& SegmentEnd, FSharedConstNavQueryFilter Filter = NULL, const UObject* Querier = NULL) const;
 
 	/** Check if poly is a custom link */
-	bool IsCustomLink(NavNodeRef PolyRef) const;
+	NAVIGATIONSYSTEM_API bool IsCustomLink(NavNodeRef PolyRef) const;
 
 	UE_DEPRECATED(5.3, "Use new override of this function with Array<FNavLinkId>* CustomLinks. This function has no effect.")
 	bool FindStraightPath(const FVector& StartLoc, const FVector& EndLoc, const TArray<NavNodeRef>& PathCorridor, TArray<FNavPathPoint>& PathPoints, TArray<uint32>* CustomLinks) const { return false; }
 
 	/** finds stringpulled path from given corridor */
-	bool FindStraightPath(const FVector& StartLoc, const FVector& EndLoc, const TArray<NavNodeRef>& PathCorridor, TArray<FNavPathPoint>& PathPoints, TArray<FNavLinkId>* CustomLinks = NULL) const;
+	NAVIGATIONSYSTEM_API bool FindStraightPath(const FVector& StartLoc, const FVector& EndLoc, const TArray<NavNodeRef>& PathCorridor, TArray<FNavPathPoint>& PathPoints, TArray<FNavLinkId>* CustomLinks = NULL) const;
 
 	/** Runs A* pathfinding on navmesh and collect data for every step */
-	int32 DebugPathfinding(const FPathFindingQuery& Query, TArray<FRecastDebugPathfindingData>& Steps);
+	NAVIGATIONSYSTEM_API int32 DebugPathfinding(const FPathFindingQuery& Query, TArray<FRecastDebugPathfindingData>& Steps);
 
-	static const FRecastQueryFilter* GetNamedFilter(ERecastNamedFilter::Type FilterType);
+	static NAVIGATIONSYSTEM_API const FRecastQueryFilter* GetNamedFilter(ERecastNamedFilter::Type FilterType);
 	FORCEINLINE static FNavPolyFlags GetNavLinkFlag() { return NavLinkFlag; }
 	
-	virtual bool NeedsRebuild() const override;
-	virtual bool SupportsRuntimeGeneration() const override;
-	virtual bool SupportsStreaming() const override;
+	NAVIGATIONSYSTEM_API virtual bool NeedsRebuild() const override;
+	NAVIGATIONSYSTEM_API virtual bool SupportsRuntimeGeneration() const override;
+	NAVIGATIONSYSTEM_API virtual bool SupportsStreaming() const override;
 
-	bool IsWorldPartitionedDynamicNavmesh() const;
+	NAVIGATIONSYSTEM_API bool IsWorldPartitionedDynamicNavmesh() const;
 	
 	/** When using active tiles generation, navigation is only allowed to be runtime generated on a subset of tiles.
 	 *  The subset is be defined by navinvokers or loaded world partitioned cells. */
-	bool IsUsingActiveTilesGeneration(const UNavigationSystemV1& NavSys) const;
+	NAVIGATIONSYSTEM_API bool IsUsingActiveTilesGeneration(const UNavigationSystemV1& NavSys) const;
 
-	virtual void ConditionalConstructGenerator() override;
+	NAVIGATIONSYSTEM_API virtual void ConditionalConstructGenerator() override;
 
 PRAGMA_DISABLE_DEPRECATION_WARNINGS	
 	UE_DEPRECATED(5.2, "UpdateGenerationProperties is unused, it will be removed")
-	void UpdateGenerationProperties(const FRecastNavMeshGenerationProperties& GenerationProps);
+	NAVIGATIONSYSTEM_API void UpdateGenerationProperties(const FRecastNavMeshGenerationProperties& GenerationProps);
 PRAGMA_ENABLE_DEPRECATION_WARNINGS	
 	
 	bool ShouldGatherDataOnGameThread() const { return bDoFullyAsyncNavDataGathering == false; }
 	int32 GetTileNumberHardLimit() const { return TileNumberHardLimit; }
 
-	virtual void UpdateActiveTiles(const TArray<FNavigationInvokerRaw>& InvokerLocations);
-	virtual void RemoveTiles(const TArray<FIntPoint>& Tiles);
+	NAVIGATIONSYSTEM_API virtual void UpdateActiveTiles(const TArray<FNavigationInvokerRaw>& InvokerLocations);
+	NAVIGATIONSYSTEM_API virtual void RemoveTiles(const TArray<FIntPoint>& Tiles);
 	
 	UE_DEPRECATED(5.3, "Use overload with FNavMeshDirtyTileElement instead.")
-	void RebuildTile(const TArray<FIntPoint>& Tiles);
+	NAVIGATIONSYSTEM_API void RebuildTile(const TArray<FIntPoint>& Tiles);
 
-	void RebuildTile(const TArray<FNavMeshDirtyTileElement>& Tiles);
+	NAVIGATIONSYSTEM_API void RebuildTile(const TArray<FNavMeshDirtyTileElement>& Tiles);
 	
-	void DirtyTilesInBounds(const FBox& Bounds);
+	NAVIGATIONSYSTEM_API void DirtyTilesInBounds(const FBox& Bounds);
 
 #if RECAST_INTERNAL_DEBUG_DATA
-	const TMap<FIntPoint, struct FRecastInternalDebugData>* GetDebugDataMap() const;
+	NAVIGATIONSYSTEM_API const TMap<FIntPoint, struct FRecastInternalDebugData>* GetDebugDataMap() const;
 #endif
 
 protected:
 
-	void UpdatePolyRefBitsPreview();
+	NAVIGATIONSYSTEM_API void UpdatePolyRefBitsPreview();
 	
 	/** Invalidates active paths that go through changed tiles  */
 	UE_DEPRECATED(5.1, "Use new version with FNavTileRef")
-	void InvalidateAffectedPaths(const TArray<uint32>& ChangedTiles);
+	NAVIGATIONSYSTEM_API void InvalidateAffectedPaths(const TArray<uint32>& ChangedTiles);
 
 	/** Invalidates active paths that go through changed tiles  */
-	void InvalidateAffectedPaths(const TArray<FNavTileRef>& ChangedTiles);
+	NAVIGATIONSYSTEM_API void InvalidateAffectedPaths(const TArray<FNavTileRef>& ChangedTiles);
 
 	/** created a new FRecastNavMeshGenerator instance. Overrider to supply your
 	 *	own extentions. Note: needs to derive from FRecastNavMeshGenerator */
-	virtual FRecastNavMeshGenerator* CreateGeneratorInstance();
+	NAVIGATIONSYSTEM_API virtual FRecastNavMeshGenerator* CreateGeneratorInstance();
 
-	void CheckToDiscardSubLevelNavData(const UNavigationSystemBase& NavSys);
+	NAVIGATIONSYSTEM_API void CheckToDiscardSubLevelNavData(const UNavigationSystemBase& NavSys);
 
 private:
 	friend struct FRecastGraphWrapper;
@@ -1475,16 +1475,16 @@ private:
 	friend class FPImplRecastNavMesh;
 	friend class URecastNavMeshDataChunk;
 	// destroys FPImplRecastNavMesh instance if it has been created 
-	void DestroyRecastPImpl();
+	NAVIGATIONSYSTEM_API void DestroyRecastPImpl();
 	// @todo docuement
-	void UpdateNavVersion();
-	void UpdateNavObject();
+	NAVIGATIONSYSTEM_API void UpdateNavVersion();
+	NAVIGATIONSYSTEM_API void UpdateNavObject();
 
 	/** @return Navmesh data chunk that belongs to this actor */
-	URecastNavMeshDataChunk* GetNavigationDataChunk(ULevel* InLevel) const;
+	NAVIGATIONSYSTEM_API URecastNavMeshDataChunk* GetNavigationDataChunk(ULevel* InLevel) const;
 
 	/** @return Navmesh data chunk that belongs to this actor */
-	URecastNavMeshDataChunk* GetNavigationDataChunk(const ANavigationDataChunkActor& InActor) const;
+	NAVIGATIONSYSTEM_API URecastNavMeshDataChunk* GetNavigationDataChunk(const ANavigationDataChunkActor& InActor) const;
 
 protected:
 	// retrieves RecastNavMeshImpl
@@ -1493,7 +1493,7 @@ protected:
 
 private:
 	/** @return Navmesh data chunk that belongs to this actor */
-	URecastNavMeshDataChunk* GetNavigationDataChunk(const TArray<UNavigationDataChunk*>& InChunks) const;
+	NAVIGATIONSYSTEM_API URecastNavMeshDataChunk* GetNavigationDataChunk(const TArray<UNavigationDataChunk*>& InChunks) const;
 
 	/** NavMesh versioning. */
 	uint32 NavMeshVersion;
@@ -1513,7 +1513,7 @@ private:
 #endif // RECAST_ASYNC_REBUILDING
 
 private:
-	static const FRecastQueryFilter* NamedFilters[ERecastNamedFilter::NamedFiltersCount];
+	static NAVIGATIONSYSTEM_API const FRecastQueryFilter* NamedFilters[ERecastNamedFilter::NamedFiltersCount];
 #endif // WITH_RECAST
 
 public:
@@ -1523,7 +1523,7 @@ public:
 
 	/** @return true if any polygon/link has been touched */
 	UFUNCTION(BlueprintCallable, Category = NavMesh, meta = (DisplayName = "ReplaceAreaInTileBounds"))
-	bool K2_ReplaceAreaInTileBounds(FBox Bounds, TSubclassOf<UNavArea> OldArea, TSubclassOf<UNavArea> NewArea, bool ReplaceLinks = true);
+	NAVIGATIONSYSTEM_API bool K2_ReplaceAreaInTileBounds(FBox Bounds, TSubclassOf<UNavArea> OldArea, TSubclassOf<UNavArea> NewArea, bool ReplaceLinks = true);
 };
 
 #if WITH_RECAST
diff --git a/Engine/Source/Runtime/NavigationSystem/Public/NavMesh/RecastNavMeshDataChunk.h b/Engine/Source/Runtime/NavigationSystem/Public/NavMesh/RecastNavMeshDataChunk.h
--- a/Engine/Source/Runtime/NavigationSystem/Public/NavMesh/RecastNavMeshDataChunk.h
+++ b/Engine/Source/Runtime/NavigationSystem/Public/NavMesh/RecastNavMeshDataChunk.h
@@ -57,42 +57,42 @@ enum EGatherTilesCopyMode
 /** 
  * 
  */
-UCLASS()
-class NAVIGATIONSYSTEM_API URecastNavMeshDataChunk : public UNavigationDataChunk
+UCLASS(MinimalAPI)
+class URecastNavMeshDataChunk : public UNavigationDataChunk
 {
 	GENERATED_UCLASS_BODY()
 
 	//~ Begin UObject Interface
-	virtual void Serialize(FArchive& Ar) override;
+	NAVIGATIONSYSTEM_API virtual void Serialize(FArchive& Ar) override;
 	//~ End UObject Interface
 
 	/** Attaches tiles to specified navmesh, transferring tile ownership to navmesh */
 	UE_DEPRECATED(5.1, "Use overload using ARecastNavMesh& instead")
-	TArray<uint32> AttachTiles(FPImplRecastNavMesh& NavMeshImpl);
+	NAVIGATIONSYSTEM_API TArray<uint32> AttachTiles(FPImplRecastNavMesh& NavMeshImpl);
 
 	/** Attaches tiles to specified navmesh */
 	UE_DEPRECATED(5.1, "Use overload using ARecastNavMesh& instead")
-	TArray<uint32> AttachTiles(FPImplRecastNavMesh& NavMeshImpl, const bool bKeepCopyOfData, const bool bKeepCopyOfCacheData);
+	NAVIGATIONSYSTEM_API TArray<uint32> AttachTiles(FPImplRecastNavMesh& NavMeshImpl, const bool bKeepCopyOfData, const bool bKeepCopyOfCacheData);
 
 	/** Detaches tiles from specified navmesh, taking tile ownership */
 	UE_DEPRECATED(5.1, "Use overload using ARecastNavMesh& instead")
-	TArray<uint32> DetachTiles(FPImplRecastNavMesh& NavMeshImpl);
+	NAVIGATIONSYSTEM_API TArray<uint32> DetachTiles(FPImplRecastNavMesh& NavMeshImpl);
 
 	/** Detaches tiles from specified navmesh */
 	UE_DEPRECATED(5.1, "Use overload using ARecastNavMesh& instead")
-	TArray<uint32> DetachTiles(FPImplRecastNavMesh& NavMeshImpl, const bool bTakeDataOwnership, const bool bTakeCacheDataOwnership);
+	NAVIGATIONSYSTEM_API TArray<uint32> DetachTiles(FPImplRecastNavMesh& NavMeshImpl, const bool bTakeDataOwnership, const bool bTakeCacheDataOwnership);
 
 	/** Attaches tiles to specified navmesh, transferring tile ownership to navmesh */
-	TArray<FNavTileRef> AttachTiles(ARecastNavMesh& NavMesh);
+	NAVIGATIONSYSTEM_API TArray<FNavTileRef> AttachTiles(ARecastNavMesh& NavMesh);
 
 	/** Attaches tiles to specified navmesh */
-	TArray<FNavTileRef> AttachTiles(ARecastNavMesh& NavMesh, const bool bKeepCopyOfData, const bool bKeepCopyOfCacheData);
+	NAVIGATIONSYSTEM_API TArray<FNavTileRef> AttachTiles(ARecastNavMesh& NavMesh, const bool bKeepCopyOfData, const bool bKeepCopyOfCacheData);
 
 	/** Detaches tiles from specified navmesh, taking tile ownership */
-	TArray<FNavTileRef> DetachTiles(ARecastNavMesh& NavMesh);
+	NAVIGATIONSYSTEM_API TArray<FNavTileRef> DetachTiles(ARecastNavMesh& NavMesh);
 
 	/** Detaches tiles from specified navmesh */
-	TArray<FNavTileRef> DetachTiles(ARecastNavMesh& NavMesh, const bool bTakeDataOwnership, const bool bTakeCacheDataOwnership);
+	NAVIGATIONSYSTEM_API TArray<FNavTileRef> DetachTiles(ARecastNavMesh& NavMesh, const bool bTakeDataOwnership, const bool bTakeCacheDataOwnership);
 
 	/** 
 	 * Experimental: Moves tiles data on the xy plane by the offset (in tile coordinates) and rotation (in degree).
@@ -101,29 +101,29 @@ class NAVIGATIONSYSTEM_API URecastNavMeshDataChunk : public UNavigationDataChunk
 	 * @param RotationDeg		Rotation in degrees.
 	 * @param RotationCenter	World position
 	 */
-	void MoveTiles(FPImplRecastNavMesh& NavMeshImpl, const FIntPoint& Offset, const FVector::FReal RotationDeg, const FVector2D& RotationCenter);
+	NAVIGATIONSYSTEM_API void MoveTiles(FPImplRecastNavMesh& NavMeshImpl, const FIntPoint& Offset, const FVector::FReal RotationDeg, const FVector2D& RotationCenter);
 	
 	/** Number of tiles in this chunk */
-	int32 GetNumTiles() const;
+	NAVIGATIONSYSTEM_API int32 GetNumTiles() const;
 
 	/** Const accessor to the list of tiles in the data chunk. */
 	const TArray<FRecastTileData>& GetTiles() const { return Tiles; }
 
 	/** Returns the AABB for the given tiles. */
-	void GetTilesBounds(const FPImplRecastNavMesh& NavMeshImpl, const TArray<int32>& TileIndices, FBox& OutBounds) const;
+	NAVIGATIONSYSTEM_API void GetTilesBounds(const FPImplRecastNavMesh& NavMeshImpl, const TArray<int32>& TileIndices, FBox& OutBounds) const;
 
 	/** Mutable accessor to the list of tiles in the data chunk. */
 	TArray<FRecastTileData>& GetMutableTiles() { return Tiles; }
 
 	/** Releases all tiles that this chunk holds */
-	void ReleaseTiles();
+	NAVIGATIONSYSTEM_API void ReleaseTiles();
 
 	/** Collect tiles with data and/or cache data from the provided TileIndices. */
-	void GetTiles(const FPImplRecastNavMesh* NavMeshImpl, const TArray<int32>& TileIndices, const EGatherTilesCopyMode CopyMode, const bool bMarkAsAttached = true);
+	NAVIGATIONSYSTEM_API void GetTiles(const FPImplRecastNavMesh* NavMeshImpl, const TArray<int32>& TileIndices, const EGatherTilesCopyMode CopyMode, const bool bMarkAsAttached = true);
 
 private:
 #if WITH_RECAST
-	void SerializeRecastData(FArchive& Ar, int32 NavMeshVersion);
+	NAVIGATIONSYSTEM_API void SerializeRecastData(FArchive& Ar, int32 NavMeshVersion);
 #endif//WITH_RECAST
 
 private:
