[Engine]

diff --git a/Engine/Source/Runtime/NavigationSystem/Public/NavMesh/RecastNavMeshGenerator.h b/Engine/Source/Runtime/NavigationSystem/Public/NavMesh/RecastNavMeshGenerator.h
--- a/Engine/Source/Runtime/NavigationSystem/Public/NavMesh/RecastNavMeshGenerator.h
+++ b/Engine/Source/Runtime/NavigationSystem/Public/NavMesh/RecastNavMeshGenerator.h
@@ -294,20 +294,20 @@ struct FRecastTileTimeSliceSettings
 /**
  * Class handling generation of a single tile, caching data that can speed up subsequent tile generations
  */
-class NAVIGATIONSYSTEM_API FRecastTileGenerator : public FNoncopyable, public FGCObject
+class FRecastTileGenerator : public FNoncopyable, public FGCObject
 {
 	friend FRecastNavMeshGenerator;
 
 public:
-	FRecastTileGenerator(FRecastNavMeshGenerator& ParentGenerator, const FIntPoint& Location, const double PendingTileCreationTime);
-	virtual ~FRecastTileGenerator();
+	NAVIGATIONSYSTEM_API FRecastTileGenerator(FRecastNavMeshGenerator& ParentGenerator, const FIntPoint& Location, const double PendingTileCreationTime);
+	NAVIGATIONSYSTEM_API virtual ~FRecastTileGenerator();
 		
 	/** Does the work involved with regenerating this tile using time slicing.
 	 *  The return value determines the result of the time slicing
 	 */
-	ETimeSliceWorkResult DoWorkTimeSliced();
+	NAVIGATIONSYSTEM_API ETimeSliceWorkResult DoWorkTimeSliced();
 	/** Does the work involved with regenerating this tile */
-	bool DoWork();
+	NAVIGATIONSYSTEM_API bool DoWork();
 
 	FORCEINLINE int32 GetTileX() const { return TileX; }
 	FORCEINLINE int32 GetTileY() const { return TileY; }
@@ -318,25 +318,25 @@ public:
 	/** Whether tile data was fully regenerated */
 	FORCEINLINE bool IsFullyRegenerated() const { return bRegenerateCompressedLayers; }
 	/** Whether tile task has anything to build */
-	bool HasDataToBuild() const;
+	NAVIGATIONSYSTEM_API bool HasDataToBuild() const;
 
 	const TArray<FNavMeshTileData>& GetCompressedLayers() const { return CompressedLayers; }
 
-	static FBox CalculateTileBounds(int32 X, int32 Y, const FVector& RcNavMeshOrigin, const FBox& TotalNavBounds, FVector::FReal TileSizeInWorldUnits);
+	static NAVIGATIONSYSTEM_API FBox CalculateTileBounds(int32 X, int32 Y, const FVector& RcNavMeshOrigin, const FBox& TotalNavBounds, FVector::FReal TileSizeInWorldUnits);
 
 protected:
 	// to be used solely by FRecastNavMeshGenerator
 	TArray<FNavMeshTileData>& GetNavigationData() { return NavigationData; }
 	
 public:
-	uint32 GetUsedMemCount() const;
+	NAVIGATIONSYSTEM_API uint32 GetUsedMemCount() const;
 
 	// Memory amount used to construct generator 
 	uint32 UsedMemoryOnStartup;
 
 	// FGCObject begin
-	virtual void AddReferencedObjects(FReferenceCollector& Collector) override;
-	virtual FString GetReferencerName() const override;
+	NAVIGATIONSYSTEM_API virtual void AddReferencedObjects(FReferenceCollector& Collector) override;
+	NAVIGATIONSYSTEM_API virtual FString GetReferencerName() const override;
 	// FGCObject end
 
 #if RECAST_INTERNAL_DEBUG_DATA
@@ -354,99 +354,99 @@ protected:
 	 *	@return Failed if failed or no need to generate (still valid).
 	 *  @return CallAgainNextTimeSlice, time slice is finished this frame but we need to call this function again next frame
 	 */
-	ETimeSliceWorkResult GenerateTileTimeSliced();
+	NAVIGATIONSYSTEM_API ETimeSliceWorkResult GenerateTileTimeSliced();
 	/** Does the actual tile generation.
 	 *	@note always trigger tile generation only via DoWorkTime(). This is a worker function
 	 *  The return value determines the result of the time slicing
 	 *	@return true if new tile navigation data has been generated and is ready to be added to navmesh instance,
 	 *	@return false if failed or no need to generate (still valid).
 	 */
-	bool GenerateTile();
+	NAVIGATIONSYSTEM_API bool GenerateTile();
 
-	void Setup(const FRecastNavMeshGenerator& ParentGenerator, const TArray<FBox>& DirtyAreas);
+	NAVIGATIONSYSTEM_API void Setup(const FRecastNavMeshGenerator& ParentGenerator, const TArray<FBox>& DirtyAreas);
 	
 	/** Gather geometry */
-	virtual void GatherGeometry(const FRecastNavMeshGenerator& ParentGenerator, bool bGeometryChanged);
+	NAVIGATIONSYSTEM_API virtual void GatherGeometry(const FRecastNavMeshGenerator& ParentGenerator, bool bGeometryChanged);
 	/** Gather geometry sources to be processed later by the GatherGeometryFromSources */
-	virtual void PrepareGeometrySources(const FRecastNavMeshGenerator& ParentGenerator, bool bGeometryChanged);
+	NAVIGATIONSYSTEM_API virtual void PrepareGeometrySources(const FRecastNavMeshGenerator& ParentGenerator, bool bGeometryChanged);
 	/** Gather geometry from the prefetched sources */
-	void GatherGeometryFromSources();
+	NAVIGATIONSYSTEM_API void GatherGeometryFromSources();
 	/** Gather geometry from the prefetched sources time sliced version */
-	ETimeSliceWorkResult GatherGeometryFromSourcesTimeSliced();
+	NAVIGATIONSYSTEM_API ETimeSliceWorkResult GatherGeometryFromSourcesTimeSliced();
 	/** Gather geometry from a specified Navigation Data */
-	void GatherNavigationDataGeometry(const TSharedRef<FNavigationRelevantData, ESPMode::ThreadSafe>& ElementData, UNavigationSystemV1& NavSys, const FNavDataConfig& OwnerNavDataConfig, bool bGeometryChanged);
+	NAVIGATIONSYSTEM_API void GatherNavigationDataGeometry(const TSharedRef<FNavigationRelevantData, ESPMode::ThreadSafe>& ElementData, UNavigationSystemV1& NavSys, const FNavDataConfig& OwnerNavDataConfig, bool bGeometryChanged);
 
 	/** Start functions used by GenerateCompressedLayersTimeSliced / GenerateCompressedLayers */
-	bool CreateHeightField(FNavMeshBuildContext& BuildContext, FTileRasterizationContext& RasterContext);
-	ETimeSliceWorkResult RasterizeTrianglesTimeSliced(FNavMeshBuildContext& BuildContext, FTileRasterizationContext& RasterContext);
-	void RasterizeTriangles(FNavMeshBuildContext& BuildContext, FTileRasterizationContext& RasterContext);
-	ETimeSliceWorkResult RasterizeGeometryRecastTimeSliced(FNavMeshBuildContext& BuildContext, const TArray<FVector::FReal>& Coords, const TArray<int32>& Indices, const rcRasterizationFlags RasterizationFlags, FTileRasterizationContext& RasterContext);
+	NAVIGATIONSYSTEM_API bool CreateHeightField(FNavMeshBuildContext& BuildContext, FTileRasterizationContext& RasterContext);
+	NAVIGATIONSYSTEM_API ETimeSliceWorkResult RasterizeTrianglesTimeSliced(FNavMeshBuildContext& BuildContext, FTileRasterizationContext& RasterContext);
+	NAVIGATIONSYSTEM_API void RasterizeTriangles(FNavMeshBuildContext& BuildContext, FTileRasterizationContext& RasterContext);
+	NAVIGATIONSYSTEM_API ETimeSliceWorkResult RasterizeGeometryRecastTimeSliced(FNavMeshBuildContext& BuildContext, const TArray<FVector::FReal>& Coords, const TArray<int32>& Indices, const rcRasterizationFlags RasterizationFlags, FTileRasterizationContext& RasterContext);
 	UE_DEPRECATED(5.0, "Call the version of this function where Coords are now a TArray of FReals!")
-	ETimeSliceWorkResult RasterizeGeometryRecastTimeSliced(FNavMeshBuildContext& BuildContext, const TArray<float>& Coords, const TArray<int32>& Indices, const rcRasterizationFlags RasterizationFlags, FTileRasterizationContext& RasterContext);
-	void RasterizeGeometryRecast(FNavMeshBuildContext& BuildContext, const TArray<FVector::FReal>& Coords, const TArray<int32>& Indices, const rcRasterizationFlags RasterizationFlags, FTileRasterizationContext& RasterContext);
+	NAVIGATIONSYSTEM_API ETimeSliceWorkResult RasterizeGeometryRecastTimeSliced(FNavMeshBuildContext& BuildContext, const TArray<float>& Coords, const TArray<int32>& Indices, const rcRasterizationFlags RasterizationFlags, FTileRasterizationContext& RasterContext);
+	NAVIGATIONSYSTEM_API void RasterizeGeometryRecast(FNavMeshBuildContext& BuildContext, const TArray<FVector::FReal>& Coords, const TArray<int32>& Indices, const rcRasterizationFlags RasterizationFlags, FTileRasterizationContext& RasterContext);
 	UE_DEPRECATED(5.0, "Call the version of this function where Coords are now a TArray of FReals!")
-	void RasterizeGeometryRecast(FNavMeshBuildContext& BuildContext, const TArray<float>& Coords, const TArray<int32>& Indices, const rcRasterizationFlags RasterizationFlags, FTileRasterizationContext& RasterContext);
-	void RasterizeGeometryTransformCoords(const TArray<FVector::FReal>& Coords, const FTransform& LocalToWorld);
+	NAVIGATIONSYSTEM_API void RasterizeGeometryRecast(FNavMeshBuildContext& BuildContext, const TArray<float>& Coords, const TArray<int32>& Indices, const rcRasterizationFlags RasterizationFlags, FTileRasterizationContext& RasterContext);
+	NAVIGATIONSYSTEM_API void RasterizeGeometryTransformCoords(const TArray<FVector::FReal>& Coords, const FTransform& LocalToWorld);
 	UE_DEPRECATED(5.0, "Call the version of this function where Coords are now a TArray of FReals!")
-	void RasterizeGeometryTransformCoords(const TArray<float>& Coords, const FTransform& LocalToWorld);
-	ETimeSliceWorkResult RasterizeGeometryTimeSliced(FNavMeshBuildContext& BuildContext, const TArray<FVector::FReal>& Coords, const TArray<int32>& Indices, const FTransform& LocalToWorld, const rcRasterizationFlags RasterizationFlags, FTileRasterizationContext& RasterContext);
+	NAVIGATIONSYSTEM_API void RasterizeGeometryTransformCoords(const TArray<float>& Coords, const FTransform& LocalToWorld);
+	NAVIGATIONSYSTEM_API ETimeSliceWorkResult RasterizeGeometryTimeSliced(FNavMeshBuildContext& BuildContext, const TArray<FVector::FReal>& Coords, const TArray<int32>& Indices, const FTransform& LocalToWorld, const rcRasterizationFlags RasterizationFlags, FTileRasterizationContext& RasterContext);
 	UE_DEPRECATED(5.0, "Call the version of this function where Coords are now a TArray of FReals!")
-	ETimeSliceWorkResult RasterizeGeometryTimeSliced(FNavMeshBuildContext& BuildContext, const TArray<float>& Coords, const TArray<int32>& Indices, const FTransform& LocalToWorld, const rcRasterizationFlags RasterizationFlags, FTileRasterizationContext& RasterContext);
-	void RasterizeGeometry(FNavMeshBuildContext& BuildContext, const TArray<FVector::FReal>& Coords, const TArray<int32>& Indices, const FTransform& LocalToWorld, const rcRasterizationFlags RasterizationFlags, FTileRasterizationContext& RasterContext);
+	NAVIGATIONSYSTEM_API ETimeSliceWorkResult RasterizeGeometryTimeSliced(FNavMeshBuildContext& BuildContext, const TArray<float>& Coords, const TArray<int32>& Indices, const FTransform& LocalToWorld, const rcRasterizationFlags RasterizationFlags, FTileRasterizationContext& RasterContext);
+	NAVIGATIONSYSTEM_API void RasterizeGeometry(FNavMeshBuildContext& BuildContext, const TArray<FVector::FReal>& Coords, const TArray<int32>& Indices, const FTransform& LocalToWorld, const rcRasterizationFlags RasterizationFlags, FTileRasterizationContext& RasterContext);
 	UE_DEPRECATED(5.0, "Call the version of this function where Coords are now a TArray of FReals!")
-	void RasterizeGeometry(FNavMeshBuildContext& BuildContext, const TArray<float>& Coords, const TArray<int32>& Indices, const FTransform& LocalToWorld, const rcRasterizationFlags RasterizationFlags, FTileRasterizationContext& RasterContext);
-	void GenerateRecastFilter(FNavMeshBuildContext& BuildContext, FTileRasterizationContext& RasterContext);
-	ETimeSliceWorkResult GenerateRecastFilterTimeSliced(FNavMeshBuildContext& BuildContext, FTileRasterizationContext& RasterContext);
-	bool BuildCompactHeightField(FNavMeshBuildContext& BuildContext, FTileRasterizationContext& RasterContext);
-	bool RecastErodeWalkable(FNavMeshBuildContext& BuildContext, FTileRasterizationContext& RasterContext);
-	bool RecastBuildLayers(FNavMeshBuildContext& BuildContext, FTileRasterizationContext& RasterContext);
-	bool RecastBuildTileCache(FNavMeshBuildContext& BuildContext, FTileRasterizationContext& RasterContext);
+	NAVIGATIONSYSTEM_API void RasterizeGeometry(FNavMeshBuildContext& BuildContext, const TArray<float>& Coords, const TArray<int32>& Indices, const FTransform& LocalToWorld, const rcRasterizationFlags RasterizationFlags, FTileRasterizationContext& RasterContext);
+	NAVIGATIONSYSTEM_API void GenerateRecastFilter(FNavMeshBuildContext& BuildContext, FTileRasterizationContext& RasterContext);
+	NAVIGATIONSYSTEM_API ETimeSliceWorkResult GenerateRecastFilterTimeSliced(FNavMeshBuildContext& BuildContext, FTileRasterizationContext& RasterContext);
+	NAVIGATIONSYSTEM_API bool BuildCompactHeightField(FNavMeshBuildContext& BuildContext, FTileRasterizationContext& RasterContext);
+	NAVIGATIONSYSTEM_API bool RecastErodeWalkable(FNavMeshBuildContext& BuildContext, FTileRasterizationContext& RasterContext);
+	NAVIGATIONSYSTEM_API bool RecastBuildLayers(FNavMeshBuildContext& BuildContext, FTileRasterizationContext& RasterContext);
+	NAVIGATIONSYSTEM_API bool RecastBuildTileCache(FNavMeshBuildContext& BuildContext, FTileRasterizationContext& RasterContext);
 	/** End functions used by GenerateCompressedLayersTimeSliced / GenerateCompressedLayers */
 
 	/** builds CompressedLayers array (geometry + modifiers) time sliced*/
-	virtual ETimeSliceWorkResult GenerateCompressedLayersTimeSliced(FNavMeshBuildContext& BuildContext);
+	NAVIGATIONSYSTEM_API virtual ETimeSliceWorkResult GenerateCompressedLayersTimeSliced(FNavMeshBuildContext& BuildContext);
 	/** builds CompressedLayers array (geometry + modifiers) */
-	virtual bool GenerateCompressedLayers(FNavMeshBuildContext& BuildContext);
+	NAVIGATIONSYSTEM_API virtual bool GenerateCompressedLayers(FNavMeshBuildContext& BuildContext);
 
 	/** Builds a navigation data layer */
-	bool GenerateNavigationDataLayer(FNavMeshBuildContext& BuildContext, FTileCacheCompressor& TileCompressor, FTileCacheAllocator& GenNavAllocator, FTileGenerationContext& GenerationContext, int32 LayerIdx);
+	NAVIGATIONSYSTEM_API bool GenerateNavigationDataLayer(FNavMeshBuildContext& BuildContext, FTileCacheCompressor& TileCompressor, FTileCacheAllocator& GenNavAllocator, FTileGenerationContext& GenerationContext, int32 LayerIdx);
 
 	/** builds NavigationData array (layers + obstacles) time sliced */
-	ETimeSliceWorkResult GenerateNavigationDataTimeSliced(FNavMeshBuildContext& BuildContext);
+	NAVIGATIONSYSTEM_API ETimeSliceWorkResult GenerateNavigationDataTimeSliced(FNavMeshBuildContext& BuildContext);
 
 	/** builds NavigationData array (layers + obstacles) */
-	bool GenerateNavigationData(FNavMeshBuildContext& BuildContext);
+	NAVIGATIONSYSTEM_API bool GenerateNavigationData(FNavMeshBuildContext& BuildContext);
 
-	virtual void ApplyVoxelFilter(struct rcHeightfield* SolidHF, FVector::FReal WalkableRadius);
+	NAVIGATIONSYSTEM_API virtual void ApplyVoxelFilter(struct rcHeightfield* SolidHF, FVector::FReal WalkableRadius);
 
 	/** Compute rasterization mask */
-	void InitRasterizationMaskArray(const rcHeightfield* SolidHF, TInlineMaskArray& OutRasterizationMasks);
-	void ComputeRasterizationMasks(FNavMeshBuildContext& BuildContext, FTileRasterizationContext& RasterContext);
-	void MarkRasterizationMask(rcContext* /*BuildContext*/, rcHeightfield* SolidHF,
+	NAVIGATIONSYSTEM_API void InitRasterizationMaskArray(const rcHeightfield* SolidHF, TInlineMaskArray& OutRasterizationMasks);
+	NAVIGATIONSYSTEM_API void ComputeRasterizationMasks(FNavMeshBuildContext& BuildContext, FTileRasterizationContext& RasterContext);
+	NAVIGATIONSYSTEM_API void MarkRasterizationMask(rcContext* /*BuildContext*/, rcHeightfield* SolidHF,
 		const FAreaNavModifier& Modifier, const FTransform& LocalToWorld, const int32 Mask, TInlineMaskArray& OutMaskArray);
 
 	/** apply areas from DynamicAreas to layer */
-	void MarkDynamicAreas(dtTileCacheLayer& Layer);
-	void MarkDynamicArea(const FAreaNavModifier& Modifier, const FTransform& LocalToWorld, dtTileCacheLayer& Layer);
-	void MarkDynamicArea(const FAreaNavModifier& Modifier, const FTransform& LocalToWorld, dtTileCacheLayer& Layer, const int32 AreaId, const int32* ReplaceAreaIdPtr);
+	NAVIGATIONSYSTEM_API void MarkDynamicAreas(dtTileCacheLayer& Layer);
+	NAVIGATIONSYSTEM_API void MarkDynamicArea(const FAreaNavModifier& Modifier, const FTransform& LocalToWorld, dtTileCacheLayer& Layer);
+	NAVIGATIONSYSTEM_API void MarkDynamicArea(const FAreaNavModifier& Modifier, const FTransform& LocalToWorld, dtTileCacheLayer& Layer, const int32 AreaId, const int32* ReplaceAreaIdPtr);
 
-	void AppendModifier(const FCompositeNavModifier& Modifier, const FNavDataPerInstanceTransformDelegate& InTransformsDelegate);
+	NAVIGATIONSYSTEM_API void AppendModifier(const FCompositeNavModifier& Modifier, const FNavDataPerInstanceTransformDelegate& InTransformsDelegate);
 	/** Appends specified geometry to tile's geometry */
-	void ValidateAndAppendGeometry(const TSharedRef<FNavigationRelevantData, ESPMode::ThreadSafe>& ElementData, const FCompositeNavModifier& InModifier);
-	void AppendGeometry(const FNavigationRelevantData& DataRef, const FCompositeNavModifier& InModifier, const FNavDataPerInstanceTransformDelegate& InTransformsDelegate);
-	void AppendVoxels(rcSpanCache* SpanData, int32 NumSpans);
+	NAVIGATIONSYSTEM_API void ValidateAndAppendGeometry(const TSharedRef<FNavigationRelevantData, ESPMode::ThreadSafe>& ElementData, const FCompositeNavModifier& InModifier);
+	NAVIGATIONSYSTEM_API void AppendGeometry(const FNavigationRelevantData& DataRef, const FCompositeNavModifier& InModifier, const FNavDataPerInstanceTransformDelegate& InTransformsDelegate);
+	NAVIGATIONSYSTEM_API void AppendVoxels(rcSpanCache* SpanData, int32 NumSpans);
 	
 	/** prepare voxel cache from collision data */
-	void PrepareVoxelCache(const TNavStatArray<uint8>& RawCollisionCache, const FCompositeNavModifier& InModifier, TNavStatArray<rcSpanCache>& SpanData);
-	bool HasVoxelCache(const TNavStatArray<uint8>& RawVoxelCache, rcSpanCache*& CachedVoxels, int32& NumCachedVoxels) const;
-	void AddVoxelCache(TNavStatArray<uint8>& RawVoxelCache, const rcSpanCache* CachedVoxels, const int32 NumCachedVoxels) const;
+	NAVIGATIONSYSTEM_API void PrepareVoxelCache(const TNavStatArray<uint8>& RawCollisionCache, const FCompositeNavModifier& InModifier, TNavStatArray<rcSpanCache>& SpanData);
+	NAVIGATIONSYSTEM_API bool HasVoxelCache(const TNavStatArray<uint8>& RawVoxelCache, rcSpanCache*& CachedVoxels, int32& NumCachedVoxels) const;
+	NAVIGATIONSYSTEM_API void AddVoxelCache(TNavStatArray<uint8>& RawVoxelCache, const rcSpanCache* CachedVoxels, const int32 NumCachedVoxels) const;
 
-	void DumpAsyncData();
-	void DumpSyncData();
+	NAVIGATIONSYSTEM_API void DumpAsyncData();
+	NAVIGATIONSYSTEM_API void DumpSyncData();
 
 #if RECAST_INTERNAL_DEBUG_DATA
-	bool IsTileDebugActive() const;
-	bool IsTileDebugAllowingGeneration() const;
+	NAVIGATIONSYSTEM_API bool IsTileDebugActive() const;
+	NAVIGATIONSYSTEM_API bool IsTileDebugAllowingGeneration() const;
 #endif
 
 protected:
@@ -530,7 +530,7 @@ protected:
 #endif
 };
 
-struct NAVIGATIONSYSTEM_API FRecastTileGeneratorWrapper : public FNonAbandonableTask
+struct FRecastTileGeneratorWrapper : public FNonAbandonableTask
 {
 	TSharedRef<FRecastTileGenerator> TileGenerator;
 
@@ -672,11 +672,11 @@ enum class EAddGeneratedTilesTimeSlicedState : uint8
 /**
  * Class that handles generation of the whole Recast-based navmesh.
  */
-class NAVIGATIONSYSTEM_API FRecastNavMeshGenerator : public FNavDataGenerator
+class FRecastNavMeshGenerator : public FNavDataGenerator
 {
 public:
-	FRecastNavMeshGenerator(ARecastNavMesh& InDestNavMesh);
-	virtual ~FRecastNavMeshGenerator();
+	NAVIGATIONSYSTEM_API FRecastNavMeshGenerator(ARecastNavMesh& InDestNavMesh);
+	NAVIGATIONSYSTEM_API virtual ~FRecastNavMeshGenerator();
 
 private:
 	/** Prevent copying. */
@@ -684,41 +684,41 @@ private:
 	FRecastNavMeshGenerator& operator=(FRecastNavMeshGenerator const& NoCopy) { check(0); return *this; }
 
 public:
-	virtual bool RebuildAll() override;
-	virtual void EnsureBuildCompletion() override;
-	virtual void CancelBuild() override;
-	virtual void TickAsyncBuild(float DeltaSeconds) override;
-	virtual void OnNavigationBoundsChanged() override;
+	NAVIGATIONSYSTEM_API virtual bool RebuildAll() override;
+	NAVIGATIONSYSTEM_API virtual void EnsureBuildCompletion() override;
+	NAVIGATIONSYSTEM_API virtual void CancelBuild() override;
+	NAVIGATIONSYSTEM_API virtual void TickAsyncBuild(float DeltaSeconds) override;
+	NAVIGATIONSYSTEM_API virtual void OnNavigationBoundsChanged() override;
 
 	/** Asks generator to update navigation affected by DirtyAreas */
-	virtual void RebuildDirtyAreas(const TArray<FNavigationDirtyArea>& DirtyAreas) override;
+	NAVIGATIONSYSTEM_API virtual void RebuildDirtyAreas(const TArray<FNavigationDirtyArea>& DirtyAreas) override;
 
 	/** determines whether this generator is performing navigation building actions at the moment, dirty areas are also checked */
-	virtual bool IsBuildInProgressCheckDirty() const override;
+	NAVIGATIONSYSTEM_API virtual bool IsBuildInProgressCheckDirty() const override;
 
 #if !RECAST_ASYNC_REBUILDING
 	/** returns true if we are time slicing and the data is valid to use false otherwise*/
-	virtual bool GetTimeSliceData(int32& OutNumRemainingBuildTasks, double& OutCurrentBuildTaskDuration) const override;
+	NAVIGATIONSYSTEM_API virtual bool GetTimeSliceData(int32& OutNumRemainingBuildTasks, double& OutCurrentBuildTaskDuration) const override;
 #endif
 
 	int32 GetNumRemaningBuildTasksHelper() const { return RunningDirtyTiles.Num() + PendingDirtyTiles.Num() + static_cast<int32>(SyncTimeSlicedData.TileGeneratorSync.IsValid()); }
-	virtual int32 GetNumRemaningBuildTasks() const override;
-	virtual int32 GetNumRunningBuildTasks() const override;
+	NAVIGATIONSYSTEM_API virtual int32 GetNumRemaningBuildTasks() const override;
+	NAVIGATIONSYSTEM_API virtual int32 GetNumRunningBuildTasks() const override;
 
 	/** Checks if a given tile is being build or has just finished building */
 	UE_DEPRECATED(5.1, "Use new version with FNavTileRef")
-	bool IsTileChanged(int32 TileIdx) const;
+	NAVIGATIONSYSTEM_API bool IsTileChanged(int32 TileIdx) const;
 
 	/** Checks if a given tile is being build or has just finished building */
-	bool IsTileChanged(const FNavTileRef InTileRef) const;
+	NAVIGATIONSYSTEM_API bool IsTileChanged(const FNavTileRef InTileRef) const;
 		
 	FORCEINLINE uint32 GetVersion() const { return Version; }
 
 	const ARecastNavMesh* GetOwner() const { return DestNavMesh; }
 
 	/** update area data */
-	void OnAreaAdded(const UClass* AreaClass, int32 AreaID);
-	void OnAreaRemoved(const UClass* AreaClass);
+	NAVIGATIONSYSTEM_API void OnAreaAdded(const UClass* AreaClass, int32 AreaID);
+	NAVIGATIONSYSTEM_API void OnAreaRemoved(const UClass* AreaClass);
 		
 	//--- accessors --- //
 	FORCEINLINE class UWorld* GetWorld() const { return DestNavMesh->GetWorld(); }
@@ -733,43 +733,43 @@ public:
 
 	/** checks if any on InclusionBounds encapsulates given box.
 	 *	@return index to first item in InclusionBounds that meets expectations */
-	int32 FindInclusionBoundEncapsulatingBox(const FBox& Box) const;
+	NAVIGATIONSYSTEM_API int32 FindInclusionBoundEncapsulatingBox(const FBox& Box) const;
 
 	/** Total navigable area box, sum of all navigation volumes bounding boxes */
 	FBox GetTotalBounds() const { return TotalNavBounds; }
 
 	const FRecastNavMeshCachedData& GetAdditionalCachedData() const { return AdditionalCachedData; }
 
-	bool HasDirtyTiles() const;
-	bool HasDirtyTiles(const FBox& AreaBounds) const;
-	int32 GetDirtyTilesCount(const FBox& AreaBounds) const;
+	NAVIGATIONSYSTEM_API bool HasDirtyTiles() const;
+	NAVIGATIONSYSTEM_API bool HasDirtyTiles(const FBox& AreaBounds) const;
+	NAVIGATIONSYSTEM_API int32 GetDirtyTilesCount(const FBox& AreaBounds) const;
 
-	bool GatherGeometryOnGameThread() const;
-	bool IsTimeSliceRegenActive() const;
+	NAVIGATIONSYSTEM_API bool GatherGeometryOnGameThread() const;
+	NAVIGATIONSYSTEM_API bool IsTimeSliceRegenActive() const;
 
-	FBox GrowBoundingBox(const FBox& BBox, bool bIncludeAgentHeight) const;
+	NAVIGATIONSYSTEM_API FBox GrowBoundingBox(const FBox& BBox, bool bIncludeAgentHeight) const;
 	
 	/** Returns if the provided Octree Element should generate geometry on the provided NavDataConfig. Can be used to extend the logic to decide what geometry is generated on what Navmesh */
-	virtual bool ShouldGenerateGeometryForOctreeElement(const FNavigationOctreeElement& Element, const FNavDataConfig& NavDataConfig) const;
+	NAVIGATIONSYSTEM_API virtual bool ShouldGenerateGeometryForOctreeElement(const FNavigationOctreeElement& Element, const FNavDataConfig& NavDataConfig) const;
 	
 #if !(UE_BUILD_SHIPPING || UE_BUILD_TEST) && ENABLE_VISUAL_LOG
-	virtual void ExportNavigationData(const FString& FileName) const override;
-	virtual void GrabDebugSnapshot(struct FVisualLogEntry* Snapshot, const FBox& BoundingBox, const FName& CategoryName, ELogVerbosity::Type Verbosity) const override;
+	NAVIGATIONSYSTEM_API virtual void ExportNavigationData(const FString& FileName) const override;
+	NAVIGATIONSYSTEM_API virtual void GrabDebugSnapshot(struct FVisualLogEntry* Snapshot, const FBox& BoundingBox, const FName& CategoryName, ELogVerbosity::Type Verbosity) const override;
 #endif
 
 	/** 
 	 *	@param Actor is a reference to make callee responsible for assuring it's valid
 	 */
-	static void ExportComponentGeometry(UActorComponent* Component, FNavigationRelevantData& Data);
-	static void ExportVertexSoupGeometry(const TArray<FVector>& Verts, FNavigationRelevantData& Data);
+	static NAVIGATIONSYSTEM_API void ExportComponentGeometry(UActorComponent* Component, FNavigationRelevantData& Data);
+	static NAVIGATIONSYSTEM_API void ExportVertexSoupGeometry(const TArray<FVector>& Verts, FNavigationRelevantData& Data);
 
-	static void ExportRigidBodyGeometry(UBodySetup& BodySetup, TNavStatArray<FVector>& OutVertexBuffer, TNavStatArray<int32>& OutIndexBuffer, const FTransform& LocalToWorld = FTransform::Identity);
-	static void ExportRigidBodyGeometry(UBodySetup& BodySetup, TNavStatArray<FVector>& OutTriMeshVertexBuffer, TNavStatArray<int32>& OutTriMeshIndexBuffer, TNavStatArray<FVector>& OutConvexVertexBuffer, TNavStatArray<int32>& OutConvexIndexBuffer, TNavStatArray<int32>& OutShapeBuffer, const FTransform& LocalToWorld = FTransform::Identity);
-	static void ExportAggregatedGeometry(const FKAggregateGeom& AggGeom, TNavStatArray<FVector>& OutConvexVertexBuffer, TNavStatArray<int32>& OutConvexIndexBuffer, TNavStatArray<int32>& OutShapeBuffer, const FTransform& LocalToWorld = FTransform::Identity);
+	static NAVIGATIONSYSTEM_API void ExportRigidBodyGeometry(UBodySetup& BodySetup, TNavStatArray<FVector>& OutVertexBuffer, TNavStatArray<int32>& OutIndexBuffer, const FTransform& LocalToWorld = FTransform::Identity);
+	static NAVIGATIONSYSTEM_API void ExportRigidBodyGeometry(UBodySetup& BodySetup, TNavStatArray<FVector>& OutTriMeshVertexBuffer, TNavStatArray<int32>& OutTriMeshIndexBuffer, TNavStatArray<FVector>& OutConvexVertexBuffer, TNavStatArray<int32>& OutConvexIndexBuffer, TNavStatArray<int32>& OutShapeBuffer, const FTransform& LocalToWorld = FTransform::Identity);
+	static NAVIGATIONSYSTEM_API void ExportAggregatedGeometry(const FKAggregateGeom& AggGeom, TNavStatArray<FVector>& OutConvexVertexBuffer, TNavStatArray<int32>& OutConvexIndexBuffer, TNavStatArray<int32>& OutShapeBuffer, const FTransform& LocalToWorld = FTransform::Identity);
 
 #if UE_ENABLE_DEBUG_DRAWING
 	/** Converts data encoded in EncodedData.CollisionData to FNavDebugMeshData format */
-	static void GetDebugGeometry(const FNavigationRelevantData& EncodedData, FNavDebugMeshData& DebugMeshData);
+	static NAVIGATIONSYSTEM_API void GetDebugGeometry(const FNavigationRelevantData& EncodedData, FNavDebugMeshData& DebugMeshData);
 #endif  // UE_ENABLE_DEBUG_DRAWING
 
 	const FNavRegenTimeSliceManager* GetTimeSliceManager() const { return SyncTimeSlicedData.TimeSliceManager; }
@@ -777,114 +777,114 @@ public:
 	void SetNextTimeSliceRegenActive(bool bRegenState) { SyncTimeSlicedData.bNextTimeSliceRegenActive = bRegenState; }
 
 	/** Update the config according to the resolution */
-	virtual void SetupTileConfig(const ENavigationDataResolution TileResolution, FRecastBuildConfig& OutConfig) const;
+	NAVIGATIONSYSTEM_API virtual void SetupTileConfig(const ENavigationDataResolution TileResolution, FRecastBuildConfig& OutConfig) const;
 
 protected:
 	// Performs initial setup of member variables so that generator is ready to
 	// do its thing from this point on. Called just after construction by ARecastNavMesh
-	virtual void Init();
+	NAVIGATIONSYSTEM_API virtual void Init();
 
 	// Used to configure Config. Override to influence build properties
-	virtual void ConfigureBuildProperties(FRecastBuildConfig& OutConfig);
+	NAVIGATIONSYSTEM_API virtual void ConfigureBuildProperties(FRecastBuildConfig& OutConfig);
 
 	// Updates cached list of navigation bounds
-	void UpdateNavigationBounds();
+	NAVIGATIONSYSTEM_API void UpdateNavigationBounds();
 		
 	// Sorts pending build tiles by proximity to player, so tiles closer to player will get generated first
-	virtual void SortPendingBuildTiles();
+	NAVIGATIONSYSTEM_API virtual void SortPendingBuildTiles();
 
 	// Get seed locations used for sorting pending build tiles. Tiles closer to these locations will be prioritized first.
-	virtual void GetSeedLocations(UWorld& World, TArray<FVector2D>& OutSeedLocations) const;
+	NAVIGATIONSYSTEM_API virtual void GetSeedLocations(UWorld& World, TArray<FVector2D>& OutSeedLocations) const;
 
 	/** Instantiates dtNavMesh and configures it for tiles generation. Returns false if failed */
-	bool ConstructTiledNavMesh();
+	NAVIGATIONSYSTEM_API bool ConstructTiledNavMesh();
 
 	/** Determine bit masks for poly address */
-	void CalcNavMeshProperties(int32& MaxTiles, int32& MaxPolys);
+	NAVIGATIONSYSTEM_API void CalcNavMeshProperties(int32& MaxTiles, int32& MaxPolys);
 	
 	/** Marks grid tiles affected by specified areas as dirty */
-	virtual void MarkDirtyTiles(const TArray<FNavigationDirtyArea>& DirtyAreas);
+	NAVIGATIONSYSTEM_API virtual void MarkDirtyTiles(const TArray<FNavigationDirtyArea>& DirtyAreas);
 
 	/** Returns if the provided UObject that requested a navmesh dirtying should dirty this Navmesh. Useful to avoid tiles regeneration from objects that are excluded from the provided NavDataConfig */
-	virtual bool ShouldDirtyTilesRequestedByObject(const UNavigationSystemV1& NavSys, const FNavigationOctree& NavOctreeInstance, const UObject& SourceObject, const FNavDataConfig& NavDataConfig) const;
+	NAVIGATIONSYSTEM_API virtual bool ShouldDirtyTilesRequestedByObject(const UNavigationSystemV1& NavSys, const FNavigationOctree& NavOctreeInstance, const UObject& SourceObject, const FNavDataConfig& NavDataConfig) const;
 
 	/** Marks all tiles overlapping with InclusionBounds dirty (via MarkDirtyTiles). */
-	bool MarkNavBoundsDirty();
+	NAVIGATIONSYSTEM_API bool MarkNavBoundsDirty();
 
 	UE_DEPRECATED(5.1, "Use new version with FNavTileRef")
-	void RemoveLayers(const FIntPoint& Tile, TArray<uint32>& UpdatedTiles);
+	NAVIGATIONSYSTEM_API void RemoveLayers(const FIntPoint& Tile, TArray<uint32>& UpdatedTiles);
 
-	void RemoveLayers(const FIntPoint& Tile, TArray<FNavTileRef>& UpdatedTiles);
+	NAVIGATIONSYSTEM_API void RemoveLayers(const FIntPoint& Tile, TArray<FNavTileRef>& UpdatedTiles);
 	
-	void StoreCompressedTileCacheLayers(const FRecastTileGenerator& TileGenerator, int32 TileX, int32 TileY);
+	NAVIGATIONSYSTEM_API void StoreCompressedTileCacheLayers(const FRecastTileGenerator& TileGenerator, int32 TileX, int32 TileY);
 
 #if RECAST_INTERNAL_DEBUG_DATA
-	void StoreDebugData(const FRecastTileGenerator& TileGenerator, int32 TileX, int32 TileY);
+	NAVIGATIONSYSTEM_API void StoreDebugData(const FRecastTileGenerator& TileGenerator, int32 TileX, int32 TileY);
 #endif
 
 #if RECAST_ASYNC_REBUILDING
 	/** Processes pending tile generation tasks Async*/
 	UE_DEPRECATED(5.1, "Use ProcessTileTasksAsyncAndGetUpdatedTiles instead")
-	TArray<uint32> ProcessTileTasksAsync(const int32 NumTasksToProcess);
+	NAVIGATIONSYSTEM_API TArray<uint32> ProcessTileTasksAsync(const int32 NumTasksToProcess);
 
 	/** Processes pending tile generation tasks Async*/
-	TArray<FNavTileRef> ProcessTileTasksAsyncAndGetUpdatedTiles(const int32 NumTasksToProcess);
+	NAVIGATIONSYSTEM_API TArray<FNavTileRef> ProcessTileTasksAsyncAndGetUpdatedTiles(const int32 NumTasksToProcess);
 #else
-	TSharedRef<FRecastTileGenerator> CreateTileGeneratorFromPendingElement(FIntPoint &OutTileLocation, const int32 ForcedPendingTileIdx = INDEX_NONE);
+	NAVIGATIONSYSTEM_API TSharedRef<FRecastTileGenerator> CreateTileGeneratorFromPendingElement(FIntPoint &OutTileLocation, const int32 ForcedPendingTileIdx = INDEX_NONE);
 
 	/** Processes pending tile generation tasks Sync with option for time slicing currently an experimental feature. */
 	UE_DEPRECATED(5.1, "Use ProcessTileTasksSyncTimeSlicedAndGetUpdatedTiles instead")
-	virtual TArray<uint32> ProcessTileTasksSyncTimeSliced();
+	NAVIGATIONSYSTEM_API virtual TArray<uint32> ProcessTileTasksSyncTimeSliced();
 	UE_DEPRECATED(5.1, "Use ProcessTileTasksSyncAndGetUpdatedTiles instead")
-	TArray<uint32> ProcessTileTasksSync(const int32 NumTasksToProcess);
+	NAVIGATIONSYSTEM_API TArray<uint32> ProcessTileTasksSync(const int32 NumTasksToProcess);
 
 	/** Processes pending tile generation tasks Sync with option for time slicing currently an experimental feature. */
-	virtual TArray<FNavTileRef> ProcessTileTasksSyncTimeSlicedAndGetUpdatedTiles();
-	TArray<FNavTileRef> ProcessTileTasksSyncAndGetUpdatedTiles(const int32 NumTasksToProcess);
+	NAVIGATIONSYSTEM_API virtual TArray<FNavTileRef> ProcessTileTasksSyncTimeSlicedAndGetUpdatedTiles();
+	NAVIGATIONSYSTEM_API TArray<FNavTileRef> ProcessTileTasksSyncAndGetUpdatedTiles(const int32 NumTasksToProcess);
 
-	virtual int32 GetNextPendingDirtyTileToBuild() const;
+	NAVIGATIONSYSTEM_API virtual int32 GetNextPendingDirtyTileToBuild() const;
 #endif
 	/** Processes pending tile generation tasks */
 	UE_DEPRECATED(5.1, "Use ProcessTileTasksAndGetUpdatedTiles instead")
-	TArray<uint32> ProcessTileTasks(const int32 NumTasksToProcess);
+	NAVIGATIONSYSTEM_API TArray<uint32> ProcessTileTasks(const int32 NumTasksToProcess);
 
 	/** Processes pending tile generation tasks */
-	TArray<FNavTileRef> ProcessTileTasksAndGetUpdatedTiles(const int32 NumTasksToProcess);
+	NAVIGATIONSYSTEM_API TArray<FNavTileRef> ProcessTileTasksAndGetUpdatedTiles(const int32 NumTasksToProcess);
 
-	void ResetTimeSlicedTileGeneratorSync();
+	NAVIGATIONSYSTEM_API void ResetTimeSlicedTileGeneratorSync();
 
 public:
 	/** Adds generated tiles to NavMesh, replacing old ones, uses time slicing returns Failed if any layer failed */
 	UE_DEPRECATED(5.1, "Use new version with FNavTileRef")
-	ETimeSliceWorkResult AddGeneratedTilesTimeSliced(FRecastTileGenerator& TileGenerator, TArray<uint32>& OutResultTileIndices);
+	NAVIGATIONSYSTEM_API ETimeSliceWorkResult AddGeneratedTilesTimeSliced(FRecastTileGenerator& TileGenerator, TArray<uint32>& OutResultTileIndices);
 
 	/** Adds generated tiles to NavMesh, replacing old ones, uses time slicing returns Failed if any layer failed */
-	ETimeSliceWorkResult AddGeneratedTilesTimeSliced(FRecastTileGenerator& TileGenerator, TArray<FNavTileRef>& OutResultTileRefs);
+	NAVIGATIONSYSTEM_API ETimeSliceWorkResult AddGeneratedTilesTimeSliced(FRecastTileGenerator& TileGenerator, TArray<FNavTileRef>& OutResultTileRefs);
 
 	/** Adds generated tiles to NavMesh, replacing old ones */
 	UE_DEPRECATED(5.1, "Use AddGeneratedTilesAndGetUpdatedTiles instead")
-	TArray<uint32> AddGeneratedTiles(FRecastTileGenerator& TileGenerator);
+	NAVIGATIONSYSTEM_API TArray<uint32> AddGeneratedTiles(FRecastTileGenerator& TileGenerator);
 
 	/** Adds generated tiles to NavMesh, replacing old ones */
-	TArray<FNavTileRef> AddGeneratedTilesAndGetUpdatedTiles(FRecastTileGenerator& TileGenerator);
+	NAVIGATIONSYSTEM_API TArray<FNavTileRef> AddGeneratedTilesAndGetUpdatedTiles(FRecastTileGenerator& TileGenerator);
 
 public:
 	/** Removes all tiles at specified grid location */
 	UE_DEPRECATED(5.1, "Use RemoveTileLayersAndGetUpdatedTiles instead")
-	TArray<uint32> RemoveTileLayers(const int32 TileX, const int32 TileY, TMap<int32, dtPolyRef>* OldLayerTileIdMap = nullptr);
+	NAVIGATIONSYSTEM_API TArray<uint32> RemoveTileLayers(const int32 TileX, const int32 TileY, TMap<int32, dtPolyRef>* OldLayerTileIdMap = nullptr);
 
 	/** Removes all tiles at specified grid location */
-	TArray<FNavTileRef> RemoveTileLayersAndGetUpdatedTiles(const int32 TileX, const int32 TileY, TMap<int32, dtPolyRef>* OldLayerTileIdMap = nullptr);
+	NAVIGATIONSYSTEM_API TArray<FNavTileRef> RemoveTileLayersAndGetUpdatedTiles(const int32 TileX, const int32 TileY, TMap<int32, dtPolyRef>* OldLayerTileIdMap = nullptr);
 
-	void RemoveTiles(const TArray<FIntPoint>& Tiles);
+	NAVIGATIONSYSTEM_API void RemoveTiles(const TArray<FIntPoint>& Tiles);
 
 	UE_DEPRECATED(5.3, "Use overload with FNavMeshDirtyTileElement instead.")
-	void ReAddTiles(const TArray<FIntPoint>& Tiles);
+	NAVIGATIONSYSTEM_API void ReAddTiles(const TArray<FIntPoint>& Tiles);
 
-	void ReAddTiles(const TArray<FNavMeshDirtyTileElement>& Tiles);
+	NAVIGATIONSYSTEM_API void ReAddTiles(const TArray<FNavMeshDirtyTileElement>& Tiles);
 
 	bool IsBuildingRestrictedToActiveTiles() const { return bRestrictBuildingToActiveTiles; }
-	bool IsInActiveSet(const FIntPoint& Tile) const;
+	NAVIGATIONSYSTEM_API bool IsInActiveSet(const FIntPoint& Tile) const;
 
 	/** sets a limit to number of asynchronous tile generators running at one time
 	 *	@note if used at runtime will not result in killing tasks above limit count
@@ -893,15 +893,15 @@ public:
 
 	void SetSortPendingTileMethod(const ENavigationSortPendingTilesMethod InMethod) { SortPendingTilesMethod = InMethod; }
 
-	static void CalcPolyRefBits(ARecastNavMesh* NavMeshOwner, int32& MaxTileBits, int32& MaxPolyBits);
+	static NAVIGATIONSYSTEM_API void CalcPolyRefBits(ARecastNavMesh* NavMeshOwner, int32& MaxTileBits, int32& MaxPolyBits);
 
 protected:
-	virtual void RestrictBuildingToActiveTiles(bool InRestrictBuildingToActiveTiles);
+	NAVIGATIONSYSTEM_API virtual void RestrictBuildingToActiveTiles(bool InRestrictBuildingToActiveTiles);
 	
 	/** Blocks until build for specified list of tiles is complete and discard results */
-	void DiscardCurrentBuildingTasks();
+	NAVIGATIONSYSTEM_API void DiscardCurrentBuildingTasks();
 
-	virtual TSharedRef<FRecastTileGenerator> CreateTileGenerator(const FIntPoint& Coord, const TArray<FBox>& DirtyAreas, const double PendingTileCreationTime = 0.);
+	NAVIGATIONSYSTEM_API virtual TSharedRef<FRecastTileGenerator> CreateTileGenerator(const FIntPoint& Coord, const TArray<FBox>& DirtyAreas, const double PendingTileCreationTime = 0.);
 
 	template <typename T>
 	UE_DEPRECATED(5.3, "Use ConstructTileGeneratorImpl instead.")
@@ -923,13 +923,13 @@ protected:
 	//----------------------------------------------------------------------//
 	// debug
 	//----------------------------------------------------------------------//
-	virtual uint32 LogMemUsed() const override;
+	NAVIGATIONSYSTEM_API virtual uint32 LogMemUsed() const override;
 
 	UE_DEPRECATED(5.1, "Use new version with FNavTileRef")
-	void AddGeneratedTileLayer(int32 LayerIndex, FRecastTileGenerator& TileGenerator, const TMap<int32, dtPolyRef>& OldLayerTileIdMap, TArray<uint32>& OutResultTileIndices);
+	NAVIGATIONSYSTEM_API void AddGeneratedTileLayer(int32 LayerIndex, FRecastTileGenerator& TileGenerator, const TMap<int32, dtPolyRef>& OldLayerTileIdMap, TArray<uint32>& OutResultTileIndices);
 
-	bool IsAllowedToAddTileLayers(const FIntPoint Tile) const;
-	void AddGeneratedTileLayer(int32 LayerIndex, FRecastTileGenerator& TileGenerator, const TMap<int32, dtPolyRef>& OldLayerTileIdMap, TArray<FNavTileRef>& OutResultTileRefs);
+	NAVIGATIONSYSTEM_API bool IsAllowedToAddTileLayers(const FIntPoint Tile) const;
+	NAVIGATIONSYSTEM_API void AddGeneratedTileLayer(int32 LayerIndex, FRecastTileGenerator& TileGenerator, const TMap<int32, dtPolyRef>& OldLayerTileIdMap, TArray<FNavTileRef>& OutResultTileRefs);
 
 #if !UE_BUILD_SHIPPING
 	/** Data struct used by 'LogDirtyAreas' that contains all the information regarding the areas that are being dirtied, per dirtied tile. */
@@ -940,7 +940,7 @@ protected:
 	};
 
 	/** Used internally, when LogNavigationDirtyArea is VeryVerbose, to log the number of tiles a dirty area is requesting. */
-	void LogDirtyAreas(const TMap<FPendingTileElement, TArray<FNavigationDirtyAreaPerTileDebugInformation>>& DirtyAreasDebuggingInformation) const; 
+	NAVIGATIONSYSTEM_API void LogDirtyAreas(const TMap<FPendingTileElement, TArray<FNavigationDirtyAreaPerTileDebugInformation>>& DirtyAreasDebuggingInformation) const; 
 #endif
 	
 protected:
diff --git a/Engine/Source/Runtime/NavigationSystem/Public/NavMesh/RecastQueryFilter.h b/Engine/Source/Runtime/NavigationSystem/Public/NavMesh/RecastQueryFilter.h
--- a/Engine/Source/Runtime/NavigationSystem/Public/NavMesh/RecastQueryFilter.h
+++ b/Engine/Source/Runtime/NavigationSystem/Public/NavMesh/RecastQueryFilter.h
@@ -17,38 +17,38 @@
 
 class UNavigationSystemV1;
 
-class NAVIGATIONSYSTEM_API FRecastQueryFilter : public INavigationQueryFilterInterface, public dtQueryFilter
+class FRecastQueryFilter : public INavigationQueryFilterInterface, public dtQueryFilter
 {
 public:
-	FRecastQueryFilter(bool bIsVirtual = true);
+	NAVIGATIONSYSTEM_API FRecastQueryFilter(bool bIsVirtual = true);
 	virtual ~FRecastQueryFilter(){}
 
-	virtual void Reset() override;
+	NAVIGATIONSYSTEM_API virtual void Reset() override;
 
-	virtual void SetAreaCost(uint8 AreaType, float Cost) override;
-	virtual void SetFixedAreaEnteringCost(uint8 AreaType, float Cost) override;
-	virtual void SetExcludedArea(uint8 AreaType) override;
-	virtual void SetAllAreaCosts(const float* CostArray, const int32 Count) override;
-	virtual void GetAllAreaCosts(float* CostArray, float* FixedCostArray, const int32 Count) const override;
-	virtual void SetBacktrackingEnabled(const bool bBacktracking) override;
-	virtual bool IsBacktrackingEnabled() const override;
-	virtual float GetHeuristicScale() const override;
-	virtual bool IsEqual(const INavigationQueryFilterInterface* Other) const override;
-	virtual void SetIncludeFlags(uint16 Flags) override;
-	virtual uint16 GetIncludeFlags() const override;
-	virtual void SetExcludeFlags(uint16 Flags) override;
-	virtual uint16 GetExcludeFlags() const override;
+	NAVIGATIONSYSTEM_API virtual void SetAreaCost(uint8 AreaType, float Cost) override;
+	NAVIGATIONSYSTEM_API virtual void SetFixedAreaEnteringCost(uint8 AreaType, float Cost) override;
+	NAVIGATIONSYSTEM_API virtual void SetExcludedArea(uint8 AreaType) override;
+	NAVIGATIONSYSTEM_API virtual void SetAllAreaCosts(const float* CostArray, const int32 Count) override;
+	NAVIGATIONSYSTEM_API virtual void GetAllAreaCosts(float* CostArray, float* FixedCostArray, const int32 Count) const override;
+	NAVIGATIONSYSTEM_API virtual void SetBacktrackingEnabled(const bool bBacktracking) override;
+	NAVIGATIONSYSTEM_API virtual bool IsBacktrackingEnabled() const override;
+	NAVIGATIONSYSTEM_API virtual float GetHeuristicScale() const override;
+	NAVIGATIONSYSTEM_API virtual bool IsEqual(const INavigationQueryFilterInterface* Other) const override;
+	NAVIGATIONSYSTEM_API virtual void SetIncludeFlags(uint16 Flags) override;
+	NAVIGATIONSYSTEM_API virtual uint16 GetIncludeFlags() const override;
+	NAVIGATIONSYSTEM_API virtual void SetExcludeFlags(uint16 Flags) override;
+	NAVIGATIONSYSTEM_API virtual uint16 GetExcludeFlags() const override;
 	virtual FVector GetAdjustedEndLocation(const FVector& EndLocation) const override { return EndLocation; }
-	virtual INavigationQueryFilterInterface* CreateCopy() const override;
+	NAVIGATIONSYSTEM_API virtual INavigationQueryFilterInterface* CreateCopy() const override;
 
 	const dtQueryFilter* GetAsDetourQueryFilter() const { return this; }
 
 	/** Changes whether the filter will use virtual set of filtering functions (getVirtualCost and passVirtualFilter)
 	 *	or the inlined ones (getInlineCost and passInlineFilter) */
-	void SetIsVirtual(bool bIsVirtual);
+	NAVIGATIONSYSTEM_API void SetIsVirtual(bool bIsVirtual);
 
 	/** Instruct filter whether it can reopen nodes already on closed list */
-	void SetShouldIgnoreClosedNodes(const bool bIgnoreClosed);
+	NAVIGATIONSYSTEM_API void SetShouldIgnoreClosedNodes(const bool bIgnoreClosed);
 
 	//----------------------------------------------------------------------//
 	// @note you might also want to override following functions from dtQueryFilter	
@@ -56,11 +56,11 @@ public:
 	// virtual FVector::FReal getVirtualCost(const FVector::FReal* pa, const FVector::FReal* pb, const dtPolyRef prevRef, const dtMeshTile* prevTile, const dtPoly* prevPoly, const dtPolyRef curRef, const dtMeshTile* curTile, const dtPoly* curPoly, const dtPolyRef nextRef, const dtMeshTile* nextTile, const dtPoly* nextPoly) const;
 };
 
-struct NAVIGATIONSYSTEM_API FRecastSpeciaLinkFilter : public dtQuerySpecialLinkFilter
+struct FRecastSpeciaLinkFilter : public dtQuerySpecialLinkFilter
 {
 	FRecastSpeciaLinkFilter(UNavigationSystemV1* NavSystem, const UObject* Owner) : NavSys(NavSystem), SearchOwner(Owner), CachedOwnerOb(nullptr) {}
-	virtual bool isLinkAllowed(const uint64 UserId) const override;
-	virtual void initialize() override;
+	NAVIGATIONSYSTEM_API virtual bool isLinkAllowed(const uint64 UserId) const override;
+	NAVIGATIONSYSTEM_API virtual void initialize() override;
 
 	UNavigationSystemV1* NavSys;
 	FWeakObjectPtr SearchOwner;
diff --git a/Engine/Source/Runtime/NavigationSystem/Public/NavModifierComponent.h b/Engine/Source/Runtime/NavigationSystem/Public/NavModifierComponent.h
--- a/Engine/Source/Runtime/NavigationSystem/Public/NavModifierComponent.h
+++ b/Engine/Source/Runtime/NavigationSystem/Public/NavModifierComponent.h
@@ -12,8 +12,8 @@
 struct FNavigationRelevantData;
 class UBodySetup;
 
-UCLASS(ClassGroup = (Navigation), meta = (BlueprintSpawnableComponent), hidecategories = (Activation), config = Engine, defaultconfig)
-class NAVIGATIONSYSTEM_API UNavModifierComponent : public UNavRelevantComponent
+UCLASS(ClassGroup = (Navigation), meta = (BlueprintSpawnableComponent), hidecategories = (Activation), config = Engine, defaultconfig, MinimalAPI)
+class UNavModifierComponent : public UNavRelevantComponent
 {
 	GENERATED_UCLASS_BODY()
 
@@ -33,26 +33,26 @@ class NAVIGATIONSYSTEM_API UNavModifierComponent : public UNavRelevantComponent
 	UPROPERTY(config, EditAnywhere, Category = Navigation)
 	uint8 bIncludeAgentHeight : 1;
 
-	virtual void CalcAndCacheBounds() const override;
-	virtual void GetNavigationData(FNavigationRelevantData& Data) const override;
+	NAVIGATIONSYSTEM_API virtual void CalcAndCacheBounds() const override;
+	NAVIGATIONSYSTEM_API virtual void GetNavigationData(FNavigationRelevantData& Data) const override;
 
 	UFUNCTION(BlueprintCallable, Category = "AI|Navigation")
-	void SetAreaClass(TSubclassOf<UNavArea> NewAreaClass);
+	NAVIGATIONSYSTEM_API void SetAreaClass(TSubclassOf<UNavArea> NewAreaClass);
 
 protected:
-	void OnTransformUpdated(USceneComponent* RootComponent, EUpdateTransformFlags UpdateTransformFlags, ETeleportType Teleport);
+	NAVIGATIONSYSTEM_API void OnTransformUpdated(USceneComponent* RootComponent, EUpdateTransformFlags UpdateTransformFlags, ETeleportType Teleport);
 
 #if WITH_EDITOR
-	void OnNavAreaRegistered(const UWorld& World, const UClass* NavAreaClass);
-	void OnNavAreaUnregistered(const UWorld& World, const UClass* NavAreaClass);
+	NAVIGATIONSYSTEM_API void OnNavAreaRegistered(const UWorld& World, const UClass* NavAreaClass);
+	NAVIGATIONSYSTEM_API void OnNavAreaUnregistered(const UWorld& World, const UClass* NavAreaClass);
 #endif // WITH_EDITOR 
 
 	//~ Begin UActorComponent Interface
-	virtual void OnRegister() override;
-	virtual void OnUnregister() override;
+	NAVIGATIONSYSTEM_API virtual void OnRegister() override;
+	NAVIGATIONSYSTEM_API virtual void OnUnregister() override;
 	//~ End UActorComponent Interface
 
-	void PopulateComponentBounds(FTransform InParentTransform, const UBodySetup& InBodySetup) const;
+	NAVIGATIONSYSTEM_API void PopulateComponentBounds(FTransform InParentTransform, const UBodySetup& InBodySetup) const;
 	
 	struct FRotatedBox
 	{
diff --git a/Engine/Source/Runtime/NavigationSystem/Public/NavModifierVolume.h b/Engine/Source/Runtime/NavigationSystem/Public/NavModifierVolume.h
--- a/Engine/Source/Runtime/NavigationSystem/Public/NavModifierVolume.h
+++ b/Engine/Source/Runtime/NavigationSystem/Public/NavModifierVolume.h
@@ -16,8 +16,8 @@ struct FNavigationRelevantData;
 /** 
  *	Allows applying selected AreaClass to navmesh, using Volume's shape
  */
-UCLASS(hidecategories=(Navigation))
-class NAVIGATIONSYSTEM_API ANavModifierVolume : public AVolume, public INavRelevantInterface
+UCLASS(hidecategories=(Navigation), MinimalAPI)
+class ANavModifierVolume : public AVolume, public INavRelevantInterface
 {
 	GENERATED_BODY()
 
@@ -40,28 +40,28 @@ protected:
 #endif
 
 public:
-	ANavModifierVolume(const FObjectInitializer& ObjectInitializer);
+	NAVIGATIONSYSTEM_API ANavModifierVolume(const FObjectInitializer& ObjectInitializer);
 
 	UFUNCTION(BlueprintCallable, Category = "AI|Navigation")
-	void SetAreaClass(TSubclassOf<UNavArea> NewAreaClass = nullptr);
+	NAVIGATIONSYSTEM_API void SetAreaClass(TSubclassOf<UNavArea> NewAreaClass = nullptr);
 
 	TSubclassOf<UNavArea> GetAreaClass() const { return AreaClass; }
 
-	virtual void GetNavigationData(FNavigationRelevantData& Data) const override;
-	virtual FBox GetNavigationBounds() const override;
-	virtual void RebuildNavigationData() override;
+	NAVIGATIONSYSTEM_API virtual void GetNavigationData(FNavigationRelevantData& Data) const override;
+	NAVIGATIONSYSTEM_API virtual FBox GetNavigationBounds() const override;
+	NAVIGATIONSYSTEM_API virtual void RebuildNavigationData() override;
 
 #if WITH_EDITOR
-	virtual void PostEditUndo() override;
-	virtual void PostEditChangeProperty(FPropertyChangedEvent& PropertyChangedEvent) override;
+	NAVIGATIONSYSTEM_API virtual void PostEditUndo() override;
+	NAVIGATIONSYSTEM_API virtual void PostEditChangeProperty(FPropertyChangedEvent& PropertyChangedEvent) override;
 #endif
 
 protected:
-	virtual void PostInitProperties() override;
-	virtual void BeginDestroy() override;
+	NAVIGATIONSYSTEM_API virtual void PostInitProperties() override;
+	NAVIGATIONSYSTEM_API virtual void BeginDestroy() override;
 
 #if WITH_EDITOR
-	void OnNavAreaRegistered(const UWorld& World, const UClass* NavAreaClass);
-	void OnNavAreaUnregistered(const UWorld& World, const UClass* NavAreaClass);
+	NAVIGATIONSYSTEM_API void OnNavAreaRegistered(const UWorld& World, const UClass* NavAreaClass);
+	NAVIGATIONSYSTEM_API void OnNavAreaUnregistered(const UWorld& World, const UClass* NavAreaClass);
 #endif
 };
diff --git a/Engine/Source/Runtime/NavigationSystem/Public/NavRelevantComponent.h b/Engine/Source/Runtime/NavigationSystem/Public/NavRelevantComponent.h
--- a/Engine/Source/Runtime/NavigationSystem/Public/NavRelevantComponent.h
+++ b/Engine/Source/Runtime/NavigationSystem/Public/NavRelevantComponent.h
@@ -8,33 +8,33 @@
 #include "AI/Navigation/NavRelevantInterface.h"
 #include "NavRelevantComponent.generated.h"
 
-UCLASS()
-class NAVIGATIONSYSTEM_API UNavRelevantComponent : public UActorComponent, public INavRelevantInterface
+UCLASS(MinimalAPI)
+class UNavRelevantComponent : public UActorComponent, public INavRelevantInterface
 {
 	GENERATED_UCLASS_BODY()
 
 	//~ Begin UActorComponent Interface
-	virtual void OnRegister() override;
-	virtual void OnUnregister() override;
+	NAVIGATIONSYSTEM_API virtual void OnRegister() override;
+	NAVIGATIONSYSTEM_API virtual void OnUnregister() override;
 	//~ End UActorComponent Interface
 
 	//~ Begin INavRelevantInterface Interface
-	virtual FBox GetNavigationBounds() const override;
-	virtual bool IsNavigationRelevant() const override;
-	virtual void UpdateNavigationBounds() override;
-	virtual UObject* GetNavigationParent() const override;
+	NAVIGATIONSYSTEM_API virtual FBox GetNavigationBounds() const override;
+	NAVIGATIONSYSTEM_API virtual bool IsNavigationRelevant() const override;
+	NAVIGATIONSYSTEM_API virtual void UpdateNavigationBounds() override;
+	NAVIGATIONSYSTEM_API virtual UObject* GetNavigationParent() const override;
 	//~ End INavRelevantInterface Interface
 	
-	virtual void CalcAndCacheBounds() const;
+	NAVIGATIONSYSTEM_API virtual void CalcAndCacheBounds() const;
 
 	UFUNCTION(BlueprintCallable, Category="AI|Navigation")
-	void SetNavigationRelevancy(bool bRelevant);
+	NAVIGATIONSYSTEM_API void SetNavigationRelevancy(bool bRelevant);
 
 	/** force relevancy and skip attaching navigation data to owner's root entry */
-	void ForceNavigationRelevancy(bool bForce);
+	NAVIGATIONSYSTEM_API void ForceNavigationRelevancy(bool bForce);
 
 	/** force refresh in navigation octree */
-	void RefreshNavigationModifiers();
+	NAVIGATIONSYSTEM_API void RefreshNavigationModifiers();
 	
 protected:
 
diff --git a/Engine/Source/Runtime/NavigationSystem/Public/NavSystemConfigOverride.h b/Engine/Source/Runtime/NavigationSystem/Public/NavSystemConfigOverride.h
--- a/Engine/Source/Runtime/NavigationSystem/Public/NavSystemConfigOverride.h
+++ b/Engine/Source/Runtime/NavigationSystem/Public/NavSystemConfigOverride.h
@@ -20,8 +20,8 @@ enum class ENavSystemOverridePolicy : uint8
 };
 
 
-UCLASS(hidecategories = (Input, Rendering, Actor, LOD, Cooking))
-class NAVIGATIONSYSTEM_API ANavSystemConfigOverride : public AActor
+UCLASS(hidecategories = (Input, Rendering, Actor, LOD, Cooking), MinimalAPI)
+class ANavSystemConfigOverride : public AActor
 {
 	GENERATED_BODY()
 
@@ -43,43 +43,43 @@ protected:
 	uint8 bLoadOnClient : 1;
 
 public:
-	ANavSystemConfigOverride(const FObjectInitializer& ObjectInitializer = FObjectInitializer::Get());
+	NAVIGATIONSYSTEM_API ANavSystemConfigOverride(const FObjectInitializer& ObjectInitializer = FObjectInitializer::Get());
 
 	//~ Begin UObject Interface
-	virtual void PostInitProperties() override;
+	NAVIGATIONSYSTEM_API virtual void PostInitProperties() override;
 #if WITH_EDITOR
-	virtual void PostEditChangeProperty(FPropertyChangedEvent& PropertyChangedEvent) override;
+	NAVIGATIONSYSTEM_API virtual void PostEditChangeProperty(FPropertyChangedEvent& PropertyChangedEvent) override;
 #endif // WITH_EDITOR
 	//~ End UObject Interface
 
 	//~ Begin AActor Interface
-	virtual void BeginPlay() override;
+	NAVIGATIONSYSTEM_API virtual void BeginPlay() override;
 #if WITH_EDITOR
-	virtual void PostRegisterAllComponents() override;
-	virtual void PostUnregisterAllComponents() override;
+	NAVIGATIONSYSTEM_API virtual void PostRegisterAllComponents() override;
+	NAVIGATIONSYSTEM_API virtual void PostUnregisterAllComponents() override;
 #endif
 	//~ End AActor Interface
 
 #if WITH_EDITOR
 	/** made an explicit function since rebuilding navigation system can be expensive */
 	UFUNCTION(Category = Navigation, meta = (CallInEditor = "true"))
-	void ApplyChanges();
+	NAVIGATIONSYSTEM_API void ApplyChanges();
 	//virtual void CheckForErrors() override;
 #endif
 
 protected:
 	/** Creates a new navigation system and plugs it into the world. If there's a
 	 *	nav system instance already in place it gets destroyed. */
-	virtual void OverrideNavSystem();
+	NAVIGATIONSYSTEM_API virtual void OverrideNavSystem();
 
 	/** Appends non-conflicting information (like supported agents) to a pre-existing 
 	 *	nav system instance */
-	virtual void AppendToNavSystem(UNavigationSystemBase& PrevNavSys);
+	NAVIGATIONSYSTEM_API virtual void AppendToNavSystem(UNavigationSystemBase& PrevNavSys);
 
 #if WITH_EDITOR
 	/** Called only in the editor mode*/
-	void InitializeForWorld(UNavigationSystemBase* NewNavSys, UWorld* World, const FNavigationSystemRunMode RunMode);
+	NAVIGATIONSYSTEM_API void InitializeForWorld(UNavigationSystemBase* NewNavSys, UWorld* World, const FNavigationSystemRunMode RunMode);
 #endif // WITH_EDITOR
 
-	void ApplyConfig();
+	NAVIGATIONSYSTEM_API void ApplyConfig();
 };
diff --git a/Engine/Source/Runtime/NavigationSystem/Public/NavigationData.h b/Engine/Source/Runtime/NavigationSystem/Public/NavigationData.h
--- a/Engine/Source/Runtime/NavigationSystem/Public/NavigationData.h
+++ b/Engine/Source/Runtime/NavigationSystem/Public/NavigationData.h
@@ -25,7 +25,7 @@ class UNavArea;
 class UPrimitiveComponent;
 
 USTRUCT()
-struct NAVIGATIONSYSTEM_API FSupportedAreaData
+struct FSupportedAreaData
 {
 	GENERATED_USTRUCT_BODY()
 
@@ -38,7 +38,7 @@ struct NAVIGATIONSYSTEM_API FSupportedAreaData
 	UPROPERTY(transient)
 	TObjectPtr<const UClass> AreaClass;
 
-	FSupportedAreaData(TSubclassOf<UNavArea> NavAreaClass = NULL, int32 InAreaID = INDEX_NONE);
+	NAVIGATIONSYSTEM_API FSupportedAreaData(TSubclassOf<UNavArea> NavAreaClass = NULL, int32 InAreaID = INDEX_NONE);
 };
 
 struct FNavPathRecalculationRequest
@@ -68,13 +68,13 @@ struct FPathFindingResult
 	FORCEINLINE bool IsPartial() const;
 };
 
-struct NAVIGATIONSYSTEM_API FNavigationPath : public TSharedFromThis<FNavigationPath, ESPMode::ThreadSafe>
+struct FNavigationPath : public TSharedFromThis<FNavigationPath, ESPMode::ThreadSafe>
 {
 	//DECLARE_DELEGATE_OneParam(FPathObserverDelegate, FNavigationPath*);
 	DECLARE_MULTICAST_DELEGATE_TwoParams(FPathObserverDelegate, FNavigationPath*, ENavPathEvent::Type);
 
-	FNavigationPath();
-	FNavigationPath(const TArray<FVector>& Points, AActor* Base = NULL);
+	NAVIGATIONSYSTEM_API FNavigationPath();
+	NAVIGATIONSYSTEM_API FNavigationPath(const TArray<FVector>& Points, AActor* Base = NULL);
 	virtual ~FNavigationPath()
 	{ }
 
@@ -141,7 +141,7 @@ struct NAVIGATIONSYSTEM_API FNavigationPath : public TSharedFromThis<FNavigation
 		bIsReady = true;
 	}
 
-	FORCEINLINE void SetNavigationDataUsed(const ANavigationData* const NewData);
+	NAVIGATIONSYSTEM_API FORCEINLINE void SetNavigationDataUsed(const ANavigationData* const NewData);
 
 	FORCEINLINE ANavigationData* GetNavigationDataUsed() const
 	{
@@ -222,33 +222,33 @@ struct NAVIGATIONSYSTEM_API FNavigationPath : public TSharedFromThis<FNavigation
 	FORCEINLINE double GetTimeStamp() const { return LastUpdateTimeStamp; }
 	FORCEINLINE void SetTimeStamp(double TimeStamp) { LastUpdateTimeStamp = TimeStamp; }
 
-	void Invalidate();
-	void RePathFailed();
+	NAVIGATIONSYSTEM_API void Invalidate();
+	NAVIGATIONSYSTEM_API void RePathFailed();
 
 	/** Resets all variables describing generated path before attempting new pathfinding call. 
 	  * This function will NOT reset setup variables like goal actor, filter, observer, etc */
-	virtual void ResetForRepath();
+	NAVIGATIONSYSTEM_API virtual void ResetForRepath();
 
 	UE_DEPRECATED(5.0, "Use version that takes LifeTime instead.")
-	virtual void DebugDraw(const ANavigationData* NavData, FColor PathColor, class UCanvas* Canvas, bool bPersistent, const uint32 NextPathPointIndex = 0) const;
-	virtual void DebugDraw(const ANavigationData* NavData, const FColor PathColor, class UCanvas* Canvas, const bool bPersistent, const float LifeTime, const uint32 NextPathPointIndex = 0) const;
+	NAVIGATIONSYSTEM_API virtual void DebugDraw(const ANavigationData* NavData, FColor PathColor, class UCanvas* Canvas, bool bPersistent, const uint32 NextPathPointIndex = 0) const;
+	NAVIGATIONSYSTEM_API virtual void DebugDraw(const ANavigationData* NavData, const FColor PathColor, class UCanvas* Canvas, const bool bPersistent, const float LifeTime, const uint32 NextPathPointIndex = 0) const;
 	
 #if ENABLE_VISUAL_LOG
-	virtual void DescribeSelfToVisLog(struct FVisualLogEntry* Snapshot) const;
-	virtual FString GetDescription() const;
+	NAVIGATIONSYSTEM_API virtual void DescribeSelfToVisLog(struct FVisualLogEntry* Snapshot) const;
+	NAVIGATIONSYSTEM_API virtual FString GetDescription() const;
 #endif // ENABLE_VISUAL_LOG
 
 	/** check if path contains specific custom nav link */
 	UE_DEPRECATED(5.3, "Use version that takes FNavLinkId instead. This function only returns false.")
 	virtual bool ContainsCustomLink(uint32 UniqueLinkId) const final {	return false; }
 
-	virtual bool ContainsCustomLink(FNavLinkId UniqueLinkId) const;
+	NAVIGATIONSYSTEM_API virtual bool ContainsCustomLink(FNavLinkId UniqueLinkId) const;
 
 	/** check if path contains any custom nav link */
-	virtual bool ContainsAnyCustomLink() const;
+	NAVIGATIONSYSTEM_API virtual bool ContainsAnyCustomLink() const;
 
 	/** check if path contains given node */
-	virtual bool ContainsNode(NavNodeRef NodeRef) const;
+	NAVIGATIONSYSTEM_API virtual bool ContainsNode(NavNodeRef NodeRef) const;
 
 	/** get cost of path, starting from given point */
 	virtual FVector::FReal GetCostFromIndex(int32 PathPointIndex) const
@@ -268,7 +268,7 @@ struct NAVIGATIONSYSTEM_API FNavigationPath : public TSharedFromThis<FNavigation
 	}
 
 	/** calculates total length of segments from NextPathPoint to the end of path, plus distance from CurrentPosition to NextPathPoint */
-	virtual FVector::FReal GetLengthFromPosition(FVector SegmentStart, uint32 NextPathPointIndex) const;
+	NAVIGATIONSYSTEM_API virtual FVector::FReal GetLengthFromPosition(FVector SegmentStart, uint32 NextPathPointIndex) const;
 
 	FORCEINLINE FVector::FReal GetLength() const
 	{
@@ -296,23 +296,23 @@ struct NAVIGATIONSYSTEM_API FNavigationPath : public TSharedFromThis<FNavigation
 	}
 
 	/** get based position of path point */
-	FBasedPosition GetPathPointLocation(uint32 Index) const;
+	NAVIGATIONSYSTEM_API FBasedPosition GetPathPointLocation(uint32 Index) const;
 
 	/** checks if given path, starting from StartingIndex, intersects with given AABB box */
-	virtual bool DoesIntersectBox(const FBox& Box, uint32 StartingIndex = 0, int32* IntersectingSegmentIndex = NULL, FVector* AgentExtent = NULL) const;
+	NAVIGATIONSYSTEM_API virtual bool DoesIntersectBox(const FBox& Box, uint32 StartingIndex = 0, int32* IntersectingSegmentIndex = NULL, FVector* AgentExtent = NULL) const;
 	/** checks if given path, starting from StartingIndex, intersects with given AABB box. This version uses AgentLocation as beginning of the path
 	 *	with segment between AgentLocation and path's StartingIndex-th node treated as first path segment to check */
-	virtual bool DoesIntersectBox(const FBox& Box, const FVector& AgentLocation, uint32 StartingIndex = 0, int32* IntersectingSegmentIndex = NULL, FVector* AgentExtent = NULL) const;
+	NAVIGATIONSYSTEM_API virtual bool DoesIntersectBox(const FBox& Box, const FVector& AgentLocation, uint32 StartingIndex = 0, int32* IntersectingSegmentIndex = NULL, FVector* AgentExtent = NULL) const;
 	/** retrieves normalized direction vector to given path segment
 	 *	for '0'-th segment returns same as for 1st segment 
 	 */
-	virtual FVector GetSegmentDirection(uint32 SegmentEndIndex) const;
+	NAVIGATIONSYSTEM_API virtual FVector GetSegmentDirection(uint32 SegmentEndIndex) const;
 
 private:
-	bool DoesPathIntersectBoxImplementation(const FBox& Box, const FVector& StartLocation, uint32 StartingIndex, int32* IntersectingSegmentIndex, FVector* AgentExtent) const;
+	NAVIGATIONSYSTEM_API bool DoesPathIntersectBoxImplementation(const FBox& Box, const FVector& StartLocation, uint32 StartingIndex, int32* IntersectingSegmentIndex, FVector* AgentExtent) const;
 
 	/** reset variables describing built path, leaves setup variables required for rebuilding */
-	void InternalResetNavigationPath();
+	NAVIGATIONSYSTEM_API void InternalResetNavigationPath();
 
 public:
 
@@ -330,19 +330,19 @@ public:
 	}
 
 	/** enables path observing specified AActor's location and update itself if actor changes location */
-	void SetGoalActorObservation(const AActor& ActorToObserve, float TetherDistance);
+	NAVIGATIONSYSTEM_API void SetGoalActorObservation(const AActor& ActorToObserve, float TetherDistance);
 	/** Modifies distance to the GoalActor at which we'll update the path */
 	void SetGoalActorTetherDistance(const float NewTetherDistace) { GoalActorLocationTetherDistanceSq = FMath::Square(NewTetherDistace); }
 	/** turns goal actor location's observation */
-	void DisableGoalActorObservation();
+	NAVIGATIONSYSTEM_API void DisableGoalActorObservation();
 	/** set's up the path to use SourceActor's location in case of recalculation */
-	void SetSourceActor(const AActor& InSourceActor);
+	NAVIGATIONSYSTEM_API void SetSourceActor(const AActor& InSourceActor);
 
 	const AActor* GetSourceActor() const { return SourceActor.Get(); }
 	const INavAgentInterface* GetSourceActorAsNavAgent() const { return SourceActorAsNavAgent; }
 
 	FVector GetLastRepathGoalLocation() const { return GoalActorLastLocation; }
-	void UpdateLastRepathGoalLocation();
+	NAVIGATIONSYSTEM_API void UpdateLastRepathGoalLocation();
 	
 	double GetLastUpdateTime() const { return LastUpdateTimeStamp; }
 	float GetGoalActorTetherDistance() const { return FMath::Sqrt(GoalActorLocationTetherDistanceSq); }
@@ -361,13 +361,13 @@ public:
 	void SetIgnoreInvalidation(bool bShouldIgnore) { bIgnoreInvalidation = bShouldIgnore; }
 	bool GetIgnoreInvalidation() const { return bIgnoreInvalidation; }
 
-	EPathObservationResult::Type TickPathObservation();
+	NAVIGATIONSYSTEM_API EPathObservationResult::Type TickPathObservation();
 
 	/** If GoalActor is set it retrieved its navigation location, if not retrieved last path point location */
-	FVector GetGoalLocation() const;
+	NAVIGATIONSYSTEM_API FVector GetGoalLocation() const;
 
 	/** retrieved location to start path finding from (in case of path recalculation) */
-	FVector GetPathFindingStartLocation() const;
+	NAVIGATIONSYSTEM_API FVector GetPathFindingStartLocation() const;
 
 	const AActor* GetGoalActor() const
 	{
@@ -416,7 +416,7 @@ protected:
 	FSharedConstNavQueryFilter Filter;
 
 	/** type of path */
-	static const FNavPathType Type;
+	static NAVIGATIONSYSTEM_API const FNavPathType Type;
 
 	FNavPathType PathType;
 
@@ -500,8 +500,8 @@ enum class ERuntimeGenerationType : uint8
  *	Represents abstract Navigation Data (sub-classed as NavMesh, NavGraph, etc)
  *	Used as a common interface for all navigation types handled by NavigationSystem
  */
-UCLASS(config=Engine, defaultconfig, NotBlueprintable, abstract)
-class NAVIGATIONSYSTEM_API ANavigationData : public AActor, public INavigationDataInterface
+UCLASS(config=Engine, defaultconfig, NotBlueprintable, abstract, MinimalAPI)
+class ANavigationData : public AActor, public INavigationDataInterface
 {
 	GENERATED_UCLASS_BODY()
 	
@@ -556,44 +556,44 @@ public:
 	//----------------------------------------------------------------------//
 
 	//~ Begin UObject/AActor Interface
-	virtual void PostInitProperties() override;
-	virtual void PostInitializeComponents() override;
-	virtual void PostLoad() override;
+	NAVIGATIONSYSTEM_API virtual void PostInitProperties() override;
+	NAVIGATIONSYSTEM_API virtual void PostInitializeComponents() override;
+	NAVIGATIONSYSTEM_API virtual void PostLoad() override;
 #if WITH_EDITOR
-	virtual void RerunConstructionScripts() override;
-	virtual void PostEditUndo() override;
+	NAVIGATIONSYSTEM_API virtual void RerunConstructionScripts() override;
+	NAVIGATIONSYSTEM_API virtual void PostEditUndo() override;
 	bool IsBuildingOnLoad() const { return bIsBuildingOnLoad; }
 	void SetIsBuildingOnLoad(bool bValue) { bIsBuildingOnLoad = bValue; }
 #endif // WITH_EDITOR
-	virtual void Destroyed() override;
-	virtual void EndPlay(const EEndPlayReason::Type EndPlayReason) override;
+	NAVIGATIONSYSTEM_API virtual void Destroyed() override;
+	NAVIGATIONSYSTEM_API virtual void EndPlay(const EEndPlayReason::Type EndPlayReason) override;
 	//~ End UObject Interface
 		
-	virtual void CleanUp();
-	virtual void ApplyWorldOffset(const FVector& InOffset, bool bWorldShift) override;
+	NAVIGATIONSYSTEM_API virtual void CleanUp();
+	NAVIGATIONSYSTEM_API virtual void ApplyWorldOffset(const FVector& InOffset, bool bWorldShift) override;
 
 protected:
-	void RequestRegistration();
+	NAVIGATIONSYSTEM_API void RequestRegistration();
 
 private:
 	/** Simply unregisters self from navigation system and calls CleanUp */
-	void UnregisterAndCleanUp();
+	NAVIGATIONSYSTEM_API void UnregisterAndCleanUp();
 
 public:
-	virtual void CleanUpAndMarkPendingKill();
+	NAVIGATIONSYSTEM_API virtual void CleanUpAndMarkPendingKill();
 
 	FORCEINLINE bool IsRegistered() const { return bRegistered; }
-	virtual void OnRegistered();
-	void OnUnregistered();
+	NAVIGATIONSYSTEM_API virtual void OnRegistered();
+	NAVIGATIONSYSTEM_API void OnUnregistered();
 	
 	FORCEINLINE uint16 GetNavDataUniqueID() const { return NavDataUniqueID; }
 
-	virtual void TickActor(float DeltaTime, enum ELevelTick TickType, FActorTickFunction& ThisTickFunction) override;
+	NAVIGATIONSYSTEM_API virtual void TickActor(float DeltaTime, enum ELevelTick TickType, FActorTickFunction& ThisTickFunction) override;
 
 	virtual bool NeedsRebuild() const { return false; }
-	virtual bool SupportsRuntimeGeneration() const;
-	virtual bool SupportsStreaming() const;
-	virtual void OnNavigationBoundsChanged();
+	NAVIGATIONSYSTEM_API virtual bool SupportsRuntimeGeneration() const;
+	NAVIGATIONSYSTEM_API virtual bool SupportsStreaming() const;
+	NAVIGATIONSYSTEM_API virtual void OnNavigationBoundsChanged();
 
 	virtual void FillNavigationDataChunkActor(const FBox& QueryBounds, class ANavigationDataChunkActor& DataChunkActor, FBox& OutTilesBounds) const {};
 	virtual void OnStreamingNavDataAdded(class ANavigationDataChunkActor& InActor) {};
@@ -617,7 +617,7 @@ public:
 	void SetSupportsDefaultAgent(bool bIsDefault) { bSupportsDefaultAgent = bIsDefault; SetNavRenderingEnabled(bIsDefault); }
 	bool IsSupportingDefaultAgent() const { return bSupportsDefaultAgent; }
 
-	virtual bool DoesSupportAgent(const FNavAgentProperties& AgentProps) const;
+	NAVIGATIONSYSTEM_API virtual bool DoesSupportAgent(const FNavAgentProperties& AgentProps) const;
 
 	virtual void RestrictBuildingToActiveTiles(bool InRestrictBuildingToActiveTiles) {}
 
@@ -630,32 +630,32 @@ protected:
 
 public:
 	/** Creates new generator in case navigation supports it */
-	virtual void ConditionalConstructGenerator();
+	NAVIGATIONSYSTEM_API virtual void ConditionalConstructGenerator();
 
 	/** Any loading before NavDataGenerator->RebuildAll() */
 	virtual void LoadBeforeGeneratorRebuild() {}
 	
 	/** Triggers rebuild in case navigation supports it */
-	virtual void RebuildAll();
+	NAVIGATIONSYSTEM_API virtual void RebuildAll();
 
 	/** Blocks until navigation build is complete  */
-	virtual void EnsureBuildCompletion();
+	NAVIGATIONSYSTEM_API virtual void EnsureBuildCompletion();
 
 	/** Cancels current build  */
-	virtual void CancelBuild();
+	NAVIGATIONSYSTEM_API virtual void CancelBuild();
 
 	/** Ticks navigation build
 	 *  If the generator is set to time sliced rebuild then this function will only get called when 
 	 *  there is sufficient time (effectively roughly once in n frames where n is the number of time sliced nav data generators currently building)
 	 */
-	virtual void TickAsyncBuild(float DeltaSeconds);
+	NAVIGATIONSYSTEM_API virtual void TickAsyncBuild(float DeltaSeconds);
 	
 	/** Retrieves navigation data generator */
 	FNavDataGenerator* GetGenerator() { return NavDataGenerator.Get(); }
 	const FNavDataGenerator* GetGenerator() const { return NavDataGenerator.Get(); }
 
 	/** Request navigation data update after changes in nav octree */
-	virtual void RebuildDirtyAreas(const TArray<FNavigationDirtyArea>& DirtyAreas);
+	NAVIGATIONSYSTEM_API virtual void RebuildDirtyAreas(const TArray<FNavigationDirtyArea>& DirtyAreas);
 
 	/** Configures this NavData instance's navigation generation to be suspended 
 	 *	or active. It's active by default. If Suspended then all calls to 
@@ -668,7 +668,7 @@ public:
 	 *	Note2: due to all areas dirtied during generation suspension ending up in 
 	 *	SuspendedDirtyAreas care needs to be taken to not use this feature for 
 	 *	extended periods of time - otherwise SuspendedDirtyAreas can get very large. */
-	virtual void SetRebuildingSuspended(const bool bNewSuspend);
+	NAVIGATIONSYSTEM_API virtual void SetRebuildingSuspended(const bool bNewSuspend);
 
 	/** Retrieves if this NavData instance's navigation generation is suspended */
 	virtual bool IsRebuildingSuspended() const { return bRebuildingSuspended; }
@@ -679,7 +679,7 @@ public:
 	/** releases navigation generator if any has been created */
 protected:
 	/** register self with navigation system as new NavAreaDefinition(s) observer */
-	void RegisterAsNavAreaClassObserver();
+	NAVIGATIONSYSTEM_API void RegisterAsNavAreaClassObserver();
 
 public:
 	/** 
@@ -716,7 +716,7 @@ public:
 
 protected:
 	/** removes from ActivePaths all paths that no longer have shared references (and are invalid in fact) */
-	void PurgeUnusedPaths();
+	NAVIGATIONSYSTEM_API void PurgeUnusedPaths();
 
 	void RegisterActivePath(FNavPathSharedPtr SharedPath)
 	{
@@ -727,25 +727,25 @@ protected:
 
 public:
 	/** Returns bounding box for the navmesh. */
-	virtual FBox GetBounds() const PURE_VIRTUAL(ANavigationData::GetBounds,return FBox(););
+	NAVIGATIONSYSTEM_API virtual FBox GetBounds() const PURE_VIRTUAL(ANavigationData::GetBounds,return FBox(););
 	
 	/** Returns list of navigable bounds. */
-	TArray<FBox> GetNavigableBounds() const;
+	NAVIGATIONSYSTEM_API TArray<FBox> GetNavigableBounds() const;
 	
 	/** Returns list of navigable bounds that belongs to specific level */
-	TArray<FBox> GetNavigableBoundsInLevel(ULevel* InLevel) const;
+	NAVIGATIONSYSTEM_API TArray<FBox> GetNavigableBoundsInLevel(ULevel* InLevel) const;
 	
 	//----------------------------------------------------------------------//
 	// Debug                                                                
 	//----------------------------------------------------------------------//
 	UE_DEPRECATED(5.0, "Use version that takes LifeTime instead.")
-	void DrawDebugPath(FNavigationPath* Path, const FColor PathColor, class UCanvas* Canvas, const bool bPersistent, const uint32 NextPathPointIndex) const;
-	void DrawDebugPath(FNavigationPath* Path, const FColor PathColor = FColor::White, class UCanvas* Canvas = nullptr, const bool bPersistent = true, const float LifeTime = -1.f, const uint32 NextPathPointIndex = 0) const;
+	NAVIGATIONSYSTEM_API void DrawDebugPath(FNavigationPath* Path, const FColor PathColor, class UCanvas* Canvas, const bool bPersistent, const uint32 NextPathPointIndex) const;
+	NAVIGATIONSYSTEM_API void DrawDebugPath(FNavigationPath* Path, const FColor PathColor = FColor::White, class UCanvas* Canvas = nullptr, const bool bPersistent = true, const float LifeTime = -1.f, const uint32 NextPathPointIndex = 0) const;
 
 	FORCEINLINE bool IsDrawingEnabled() const { return bEnableDrawing; }
 
 	/** @return Total mem counted, including super calls. */
-	virtual uint32 LogMemUsed() const;
+	NAVIGATIONSYSTEM_API virtual uint32 LogMemUsed() const;
 
 	//----------------------------------------------------------------------//
 	// Batch processing (important with async rebuilding)
@@ -837,13 +837,13 @@ public:
 	}
 
 	/** Raycasts batched for efficiency */
-	virtual void BatchRaycast(TArray<FNavigationRaycastWork>& Workload, FSharedConstNavQueryFilter QueryFilter, const UObject* Querier = NULL) const PURE_VIRTUAL(ANavigationData::BatchRaycast, );
+	NAVIGATIONSYSTEM_API virtual void BatchRaycast(TArray<FNavigationRaycastWork>& Workload, FSharedConstNavQueryFilter QueryFilter, const UObject* Querier = NULL) const PURE_VIRTUAL(ANavigationData::BatchRaycast, );
 
 	/**	Tries to move current nav location towards target constrained to navigable area. Faster than ProjectPointToNavmesh.
 	 *	@param OutLocation if successful this variable will be filed with result
 	 *	@return true if successful, false otherwise
 	 */
-	virtual bool FindMoveAlongSurface(const FNavLocation& StartLocation, const FVector& TargetPosition, FNavLocation& OutLocation, FSharedConstNavQueryFilter Filter = NULL, const UObject* Querier = NULL) const PURE_VIRTUAL(ANavigationData::FindMoveAlongSurface, return false;);
+	NAVIGATIONSYSTEM_API virtual bool FindMoveAlongSurface(const FNavLocation& StartLocation, const FVector& TargetPosition, FNavLocation& OutLocation, FSharedConstNavQueryFilter Filter = NULL, const UObject* Querier = NULL) const PURE_VIRTUAL(ANavigationData::FindMoveAlongSurface, return false;);
 
 	/**	Returns the navmesh edges that touch the convex polygon. The edges are not clipped by the polygon. 
 	 *	@param StartLocation a location on the navmesh where to start searching.
@@ -852,7 +852,7 @@ public:
 	 *	@param Filter Nav filter to use, or if nullptr, default filter is used. 
 	 *	@return true if successful, false otherwise
 	 */
-	virtual bool FindOverlappingEdges(const FNavLocation& StartLocation, TConstArrayView<FVector> ConvexPolygon, TArray<FVector>& OutEdges, FSharedConstNavQueryFilter Filter = NULL, const UObject* Querier = NULL) const PURE_VIRTUAL(ANavigationData::FindOverlappingEdges, return false;);
+	NAVIGATIONSYSTEM_API virtual bool FindOverlappingEdges(const FNavLocation& StartLocation, TConstArrayView<FVector> ConvexPolygon, TArray<FVector>& OutEdges, FSharedConstNavQueryFilter Filter = NULL, const UObject* Querier = NULL) const PURE_VIRTUAL(ANavigationData::FindOverlappingEdges, return false;);
 	
 	/**	Searches navmesh edges between the two path points, search up to the convex polygon described in SearchArea. The returned edges are not clipped to the search area polygon.
 	 *  @param Path Path where From and To belong to.
@@ -864,85 +864,85 @@ public:
 	 *	@param Filter Nav filter to use, or if nullptr, default filter is used. 
 	 *	@return true if successful, false otherwise
 	 */
-	virtual bool GetPathSegmentBoundaryEdges(const FNavigationPath& Path, const FNavPathPoint& StartPoint, const FNavPathPoint& EndPoint, const TConstArrayView<FVector> SearchArea, TArray<FVector>& OutEdges, const float MaxAreaEnterCost, FSharedConstNavQueryFilter Filter = NULL, const UObject* Querier = NULL) const PURE_VIRTUAL(ANavigationData::GetPathSegmentBoundaryEdges, return false;);
+	NAVIGATIONSYSTEM_API virtual bool GetPathSegmentBoundaryEdges(const FNavigationPath& Path, const FNavPathPoint& StartPoint, const FNavPathPoint& EndPoint, const TConstArrayView<FVector> SearchArea, TArray<FVector>& OutEdges, const float MaxAreaEnterCost, FSharedConstNavQueryFilter Filter = NULL, const UObject* Querier = NULL) const PURE_VIRTUAL(ANavigationData::GetPathSegmentBoundaryEdges, return false;);
 
-	virtual FNavLocation GetRandomPoint(FSharedConstNavQueryFilter Filter = NULL, const UObject* Querier = NULL) const PURE_VIRTUAL(ANavigationData::GetRandomPoint, return FNavLocation(););
+	NAVIGATIONSYSTEM_API virtual FNavLocation GetRandomPoint(FSharedConstNavQueryFilter Filter = NULL, const UObject* Querier = NULL) const PURE_VIRTUAL(ANavigationData::GetRandomPoint, return FNavLocation(););
 
 	/** finds a random location in Radius, reachable from Origin */
-	virtual bool GetRandomReachablePointInRadius(const FVector& Origin, float Radius, FNavLocation& OutResult, FSharedConstNavQueryFilter Filter = NULL, const UObject* Querier = NULL) const PURE_VIRTUAL(ANavigationData::GetRandomReachablePointInRadius, return false;);
+	NAVIGATIONSYSTEM_API virtual bool GetRandomReachablePointInRadius(const FVector& Origin, float Radius, FNavLocation& OutResult, FSharedConstNavQueryFilter Filter = NULL, const UObject* Querier = NULL) const PURE_VIRTUAL(ANavigationData::GetRandomReachablePointInRadius, return false;);
 
 	/** finds a random location in navigable space, in given Radius */
-	virtual bool GetRandomPointInNavigableRadius(const FVector& Origin, float Radius, FNavLocation& OutResult, FSharedConstNavQueryFilter Filter = NULL, const UObject* Querier = NULL) const PURE_VIRTUAL(ANavigationData::GetRandomPointInNavigableRadius, return false;);
+	NAVIGATIONSYSTEM_API virtual bool GetRandomPointInNavigableRadius(const FVector& Origin, float Radius, FNavLocation& OutResult, FSharedConstNavQueryFilter Filter = NULL, const UObject* Querier = NULL) const PURE_VIRTUAL(ANavigationData::GetRandomPointInNavigableRadius, return false;);
 	
 	/**	Tries to project given Point to this navigation type, within given Extent.
 	 *	@param OutLocation if successful this variable will be filed with result
 	 *	@return true if successful, false otherwise
 	 */
-	virtual bool ProjectPoint(const FVector& Point, FNavLocation& OutLocation, const FVector& Extent, FSharedConstNavQueryFilter Filter = NULL, const UObject* Querier = NULL) const PURE_VIRTUAL(ANavigationData::ProjectPoint, return false;);
+	NAVIGATIONSYSTEM_API virtual bool ProjectPoint(const FVector& Point, FNavLocation& OutLocation, const FVector& Extent, FSharedConstNavQueryFilter Filter = NULL, const UObject* Querier = NULL) const PURE_VIRTUAL(ANavigationData::ProjectPoint, return false;);
 
 	/**	batches ProjectPoint's work for efficiency */
-	virtual void BatchProjectPoints(TArray<FNavigationProjectionWork>& Workload, const FVector& Extent, FSharedConstNavQueryFilter Filter = NULL, const UObject* Querier = NULL) const PURE_VIRTUAL(ANavigationData::BatchProjectPoints, );
+	NAVIGATIONSYSTEM_API virtual void BatchProjectPoints(TArray<FNavigationProjectionWork>& Workload, const FVector& Extent, FSharedConstNavQueryFilter Filter = NULL, const UObject* Querier = NULL) const PURE_VIRTUAL(ANavigationData::BatchProjectPoints, );
 
 	/** Project batch of points using shared search filter. This version is not requiring user to pass in Extent, 
 	 *	and is instead relying on FNavigationProjectionWork.ProjectionLimit.
 	 *	@note function should assert if item's FNavigationProjectionWork.ProjectionLimit is invalid */
-	virtual void BatchProjectPoints(TArray<FNavigationProjectionWork>& Workload, FSharedConstNavQueryFilter Filter = NULL, const UObject* Querier = NULL) const PURE_VIRTUAL(ANavigationData::BatchProjectPoints, );
+	NAVIGATIONSYSTEM_API virtual void BatchProjectPoints(TArray<FNavigationProjectionWork>& Workload, FSharedConstNavQueryFilter Filter = NULL, const UObject* Querier = NULL) const PURE_VIRTUAL(ANavigationData::BatchProjectPoints, );
 
 	UE_DEPRECATED(5.2, "Use new version with FVector::FReal")
-	virtual ENavigationQueryResult::Type CalcPathCost(const FVector& PathStart, const FVector& PathEnd, float& OutPathCost, FSharedConstNavQueryFilter QueryFilter = NULL, const UObject* Querier = NULL) const final;
+	NAVIGATIONSYSTEM_API virtual ENavigationQueryResult::Type CalcPathCost(const FVector& PathStart, const FVector& PathEnd, float& OutPathCost, FSharedConstNavQueryFilter QueryFilter = NULL, const UObject* Querier = NULL) const final;
 
 	/** Calculates path from PathStart to PathEnd and retrieves its cost.
  *	@NOTE this function does not generate string pulled path so the result is an (over-estimated) approximation
  *	@NOTE potentially expensive, so use it with caution */
-	virtual ENavigationQueryResult::Type CalcPathCost(const FVector& PathStart, const FVector& PathEnd, FVector::FReal& OutPathCost, FSharedConstNavQueryFilter QueryFilter = NULL, const UObject* Querier = NULL) const PURE_VIRTUAL(ANavigationData::CalcPathCost, return ENavigationQueryResult::Invalid;);
+	NAVIGATIONSYSTEM_API virtual ENavigationQueryResult::Type CalcPathCost(const FVector& PathStart, const FVector& PathEnd, FVector::FReal& OutPathCost, FSharedConstNavQueryFilter QueryFilter = NULL, const UObject* Querier = NULL) const PURE_VIRTUAL(ANavigationData::CalcPathCost, return ENavigationQueryResult::Invalid;);
 
 	UE_DEPRECATED(5.2, "Use new version with FVector::FReal")
-	virtual ENavigationQueryResult::Type CalcPathLength(const FVector& PathStart, const FVector& PathEnd, float& OutPathLength, FSharedConstNavQueryFilter QueryFilter = NULL, const UObject* Querier = NULL) const final;
+	NAVIGATIONSYSTEM_API virtual ENavigationQueryResult::Type CalcPathLength(const FVector& PathStart, const FVector& PathEnd, float& OutPathLength, FSharedConstNavQueryFilter QueryFilter = NULL, const UObject* Querier = NULL) const final;
 
 	/** Calculates path from PathStart to PathEnd and retrieves its length.
 	 *	@NOTE this function does not generate string pulled path so the result is an (over-estimated) approximation
 	 *	@NOTE potentially expensive, so use it with caution */
-	virtual ENavigationQueryResult::Type CalcPathLength(const FVector& PathStart, const FVector& PathEnd, FVector::FReal& OutPathLength, FSharedConstNavQueryFilter QueryFilter = NULL, const UObject* Querier = NULL) const PURE_VIRTUAL(ANavigationData::CalcPathLength, return ENavigationQueryResult::Invalid;);
+	NAVIGATIONSYSTEM_API virtual ENavigationQueryResult::Type CalcPathLength(const FVector& PathStart, const FVector& PathEnd, FVector::FReal& OutPathLength, FSharedConstNavQueryFilter QueryFilter = NULL, const UObject* Querier = NULL) const PURE_VIRTUAL(ANavigationData::CalcPathLength, return ENavigationQueryResult::Invalid;);
 
 	UE_DEPRECATED(5.2, "Use new version with FVector::FReal")
-	virtual ENavigationQueryResult::Type CalcPathLengthAndCost(const FVector& PathStart, const FVector& PathEnd, float& OutPathLength, float& OutPathCost, FSharedConstNavQueryFilter QueryFilter = NULL, const UObject* Querier = NULL) const final;
+	NAVIGATIONSYSTEM_API virtual ENavigationQueryResult::Type CalcPathLengthAndCost(const FVector& PathStart, const FVector& PathEnd, float& OutPathLength, float& OutPathCost, FSharedConstNavQueryFilter QueryFilter = NULL, const UObject* Querier = NULL) const final;
 	/** Calculates path from PathStart to PathEnd and retrieves its length.
 	 *	@NOTE this function does not generate string pulled path so the result is an (over-estimated) approximation
 	 *	@NOTE potentially expensive, so use it with caution */
-	virtual ENavigationQueryResult::Type CalcPathLengthAndCost(const FVector& PathStart, const FVector& PathEnd, FVector::FReal& OutPathLength, FVector::FReal& OutPathCost, FSharedConstNavQueryFilter QueryFilter = NULL, const UObject* Querier = NULL) const PURE_VIRTUAL(ANavigationData::CalcPathLengthAndCost, return ENavigationQueryResult::Invalid;);
+	NAVIGATIONSYSTEM_API virtual ENavigationQueryResult::Type CalcPathLengthAndCost(const FVector& PathStart, const FVector& PathEnd, FVector::FReal& OutPathLength, FVector::FReal& OutPathCost, FSharedConstNavQueryFilter QueryFilter = NULL, const UObject* Querier = NULL) const PURE_VIRTUAL(ANavigationData::CalcPathLengthAndCost, return ENavigationQueryResult::Invalid;);
 
 	/** Checks if specified navigation node contains given location 
 	 *	@param Location is expressed in WorldSpace, navigation data is responsible for tansforming if need be */
-	virtual bool DoesNodeContainLocation(NavNodeRef NodeRef, const FVector& WorldSpaceLocation) const PURE_VIRTUAL(ANavigationData::DoesNodeContainLocation, return false;);
+	NAVIGATIONSYSTEM_API virtual bool DoesNodeContainLocation(NavNodeRef NodeRef, const FVector& WorldSpaceLocation) const PURE_VIRTUAL(ANavigationData::DoesNodeContainLocation, return false;);
 
-	double GetWorldTimeStamp() const;
+	NAVIGATIONSYSTEM_API double GetWorldTimeStamp() const;
 
 	//----------------------------------------------------------------------//
 	// Areas
 	//----------------------------------------------------------------------//
 
 	/** new area was registered in navigation system */
-	virtual void OnNavAreaAdded(const UClass* NavAreaClass, int32 AgentIndex);
+	NAVIGATIONSYSTEM_API virtual void OnNavAreaAdded(const UClass* NavAreaClass, int32 AgentIndex);
 	
 	/** area was removed from navigation system */
-	virtual void OnNavAreaRemoved(const UClass* NavAreaClass);
+	NAVIGATIONSYSTEM_API virtual void OnNavAreaRemoved(const UClass* NavAreaClass);
 
 	/** called after changes to registered area classes */
-	virtual void OnNavAreaChanged();
+	NAVIGATIONSYSTEM_API virtual void OnNavAreaChanged();
 
-	void OnNavAreaEvent(const UClass* NavAreaClass, ENavAreaEvent::Type Event);
+	NAVIGATIONSYSTEM_API void OnNavAreaEvent(const UClass* NavAreaClass, ENavAreaEvent::Type Event);
 
 	/** add all existing areas */
-	void ProcessNavAreas(const TSet<const UClass*>& AreaClasses, int32 AgentIndex);
+	NAVIGATIONSYSTEM_API void ProcessNavAreas(const TSet<const UClass*>& AreaClasses, int32 AgentIndex);
 
 	/** get class associated with AreaID */
-	const UClass* GetAreaClass(int32 AreaID) const;
+	NAVIGATIONSYSTEM_API const UClass* GetAreaClass(int32 AreaID) const;
 	
 	/** check if AreaID was assigned to class (class itself may not be loaded yet!) */
-	bool IsAreaAssigned(int32 AreaID) const;
+	NAVIGATIONSYSTEM_API bool IsAreaAssigned(int32 AreaID) const;
 
 	/** get ID assigned to AreaClas or -1 when not assigned */
-	int32 GetAreaID(const UClass* AreaClass) const;
+	NAVIGATIONSYSTEM_API int32 GetAreaID(const UClass* AreaClass) const;
 
 	/** get max areas supported by this navigation data */
 	virtual int32 GetMaxSupportedAreas() const { return MAX_int32; }
@@ -954,20 +954,20 @@ public:
 	// Custom navigation links
 	//----------------------------------------------------------------------//
 
-	virtual void UpdateCustomLink(const INavLinkCustomInterface* CustomLink);
+	NAVIGATIONSYSTEM_API virtual void UpdateCustomLink(const INavLinkCustomInterface* CustomLink);
 
 	//----------------------------------------------------------------------//
 	// Filters
 	//----------------------------------------------------------------------//
 
 	/** get cached query filter */
-	FSharedConstNavQueryFilter GetQueryFilter(TSubclassOf<UNavigationQueryFilter> FilterClass) const;
+	NAVIGATIONSYSTEM_API FSharedConstNavQueryFilter GetQueryFilter(TSubclassOf<UNavigationQueryFilter> FilterClass) const;
 
 	/** store cached query filter */
-	void StoreQueryFilter(TSubclassOf<UNavigationQueryFilter> FilterClass, FSharedConstNavQueryFilter NavFilter);
+	NAVIGATIONSYSTEM_API void StoreQueryFilter(TSubclassOf<UNavigationQueryFilter> FilterClass, FSharedConstNavQueryFilter NavFilter);
 
 	/** removes cached query filter */
-	void RemoveQueryFilter(TSubclassOf<UNavigationQueryFilter> FilterClass);
+	NAVIGATIONSYSTEM_API void RemoveQueryFilter(TSubclassOf<UNavigationQueryFilter> FilterClass);
 
 	//----------------------------------------------------------------------//
 	// all the rest                                                                
@@ -975,17 +975,17 @@ public:
 	virtual UPrimitiveComponent* ConstructRenderingComponent() { return NULL; }
 
 	/** updates state of rendering component */
-	void SetNavRenderingEnabled(bool bEnable);
+	NAVIGATIONSYSTEM_API void SetNavRenderingEnabled(bool bEnable);
 
 #if WITH_EDITOR
 	virtual bool CanChangeIsSpatiallyLoadedFlag() const override { return false; }
 #endif
 
 protected:
-	void InstantiateAndRegisterRenderingComponent();
+	NAVIGATIONSYSTEM_API void InstantiateAndRegisterRenderingComponent();
 
 	/** get ID to assign for newly added area */
-	virtual int32 GetNewAreaID(const UClass* AreaClass) const;
+	NAVIGATIONSYSTEM_API virtual int32 GetNewAreaID(const UClass* AreaClass) const;
 	
 protected:
 	/** Navigation data versioning. */
@@ -1064,7 +1064,7 @@ protected:
 private:
 	uint16 NavDataUniqueID;
 
-	static uint16 GetNextUniqueID();
+	static NAVIGATIONSYSTEM_API uint16 GetNextUniqueID();
 };
 
 struct FAsyncPathFindingQuery : public FPathFindingQuery
diff --git a/Engine/Source/Runtime/NavigationSystem/Public/NavigationDataHandler.h b/Engine/Source/Runtime/NavigationSystem/Public/NavigationDataHandler.h
--- a/Engine/Source/Runtime/NavigationSystem/Public/NavigationDataHandler.h
+++ b/Engine/Source/Runtime/NavigationSystem/Public/NavigationDataHandler.h
@@ -6,18 +6,18 @@
 
 struct FNavigationDirtyElement;
 
-struct NAVIGATIONSYSTEM_API FNavigationDataHandler
+struct FNavigationDataHandler
 {
 	FNavigationOctreeController& OctreeController;
 	FNavigationDirtyAreasController& DirtyAreasController;
 
-	FNavigationDataHandler(FNavigationOctreeController& InOctreeController, FNavigationDirtyAreasController& InDirtyAreasController);
+	NAVIGATIONSYSTEM_API FNavigationDataHandler(FNavigationOctreeController& InOctreeController, FNavigationDirtyAreasController& InDirtyAreasController);
 
-	void ConstructNavOctree(const FVector& Origin, const float Radius, const ENavDataGatheringModeConfig DataGatheringMode, const float GatheringNavModifiersWarningLimitTime);
+	NAVIGATIONSYSTEM_API void ConstructNavOctree(const FVector& Origin, const float Radius, const ENavDataGatheringModeConfig DataGatheringMode, const float GatheringNavModifiersWarningLimitTime);
 
-	void RemoveNavOctreeElementId(const FOctreeElementId2& ElementId, int32 UpdateFlags);
-	FSetElementId RegisterNavOctreeElement(UObject& ElementOwner, INavRelevantInterface& ElementInterface, int32 UpdateFlags);
-	void AddElementToNavOctree(const FNavigationDirtyElement& DirtyElement);
+	NAVIGATIONSYSTEM_API void RemoveNavOctreeElementId(const FOctreeElementId2& ElementId, int32 UpdateFlags);
+	NAVIGATIONSYSTEM_API FSetElementId RegisterNavOctreeElement(UObject& ElementOwner, INavRelevantInterface& ElementInterface, int32 UpdateFlags);
+	NAVIGATIONSYSTEM_API void AddElementToNavOctree(const FNavigationDirtyElement& DirtyElement);
 
 	/** Removes associated NavOctreeElement and invalidates associated pending updates. Also removes object from the list of children
 	* of the NavigationParent, if any.
@@ -27,15 +27,15 @@ struct NAVIGATIONSYSTEM_API FNavigationDataHandler
 	*
 	* @return True if associated NavOctreeElement has been removed or pending update has been invalidated; false otherwise.
 	*/
-	bool UnregisterNavOctreeElement(UObject& ElementOwner, INavRelevantInterface& ElementInterface, int32 UpdateFlags);
-	void UpdateNavOctreeElement(UObject& ElementOwner, INavRelevantInterface& ElementInterface, int32 UpdateFlags);
-	void UpdateNavOctreeParentChain(UObject& ElementOwner, bool bSkipElementOwnerUpdate = false);
-	bool UpdateNavOctreeElementBounds(UActorComponent& Comp, const FBox& NewBounds, const FBox& DirtyArea);	
-	void FindElementsInNavOctree(const FBox& QueryBox, const FNavigationOctreeFilter& Filter, TArray<FNavigationOctreeElement>& Elements);
-	bool ReplaceAreaInOctreeData(const UObject& Object, TSubclassOf<UNavArea> OldArea, TSubclassOf<UNavArea> NewArea, bool bReplaceChildClasses);
-	void AddLevelCollisionToOctree(ULevel& Level);
-	void RemoveLevelCollisionFromOctree(ULevel& Level);
-	void UpdateActorAndComponentsInNavOctree(AActor& Actor);
-	void ProcessPendingOctreeUpdates();
-	void DemandLazyDataGathering(FNavigationRelevantData& ElementData);
+	NAVIGATIONSYSTEM_API bool UnregisterNavOctreeElement(UObject& ElementOwner, INavRelevantInterface& ElementInterface, int32 UpdateFlags);
+	NAVIGATIONSYSTEM_API void UpdateNavOctreeElement(UObject& ElementOwner, INavRelevantInterface& ElementInterface, int32 UpdateFlags);
+	NAVIGATIONSYSTEM_API void UpdateNavOctreeParentChain(UObject& ElementOwner, bool bSkipElementOwnerUpdate = false);
+	NAVIGATIONSYSTEM_API bool UpdateNavOctreeElementBounds(UActorComponent& Comp, const FBox& NewBounds, const FBox& DirtyArea);	
+	NAVIGATIONSYSTEM_API void FindElementsInNavOctree(const FBox& QueryBox, const FNavigationOctreeFilter& Filter, TArray<FNavigationOctreeElement>& Elements);
+	NAVIGATIONSYSTEM_API bool ReplaceAreaInOctreeData(const UObject& Object, TSubclassOf<UNavArea> OldArea, TSubclassOf<UNavArea> NewArea, bool bReplaceChildClasses);
+	NAVIGATIONSYSTEM_API void AddLevelCollisionToOctree(ULevel& Level);
+	NAVIGATIONSYSTEM_API void RemoveLevelCollisionFromOctree(ULevel& Level);
+	NAVIGATIONSYSTEM_API void UpdateActorAndComponentsInNavOctree(AActor& Actor);
+	NAVIGATIONSYSTEM_API void ProcessPendingOctreeUpdates();
+	NAVIGATIONSYSTEM_API void DemandLazyDataGathering(FNavigationRelevantData& ElementData);
 };
diff --git a/Engine/Source/Runtime/NavigationSystem/Public/NavigationDirtyAreasController.h b/Engine/Source/Runtime/NavigationSystem/Public/NavigationDirtyAreasController.h
--- a/Engine/Source/Runtime/NavigationSystem/Public/NavigationDirtyAreasController.h
+++ b/Engine/Source/Runtime/NavigationSystem/Public/NavigationDirtyAreasController.h
@@ -18,7 +18,7 @@ NAVIGATIONSYSTEM_API DECLARE_LOG_CATEGORY_EXTERN(LogNavigationDirtyArea, Warning
 class ANavigationData;
 struct FNavigationDirtyElement;
 
-struct NAVIGATIONSYSTEM_API FNavigationDirtyAreasController
+struct FNavigationDirtyAreasController
 {
 	/** update frequency for dirty areas on navmesh */
 	float DirtyAreasUpdateFreq = 60.f;
@@ -42,18 +42,18 @@ private:
 	/** -1 by default, if set to a positive value dirty area with bounds size over that threshold will be logged */
 	float DirtyAreaWarningSizeThreshold = -1.f;
 
-	bool ShouldReportOversizedDirtyArea() const;
+	NAVIGATIONSYSTEM_API bool ShouldReportOversizedDirtyArea() const;
 #endif // !UE_BUILD_SHIPPING
 
 public:
-	FNavigationDirtyAreasController();
+	NAVIGATIONSYSTEM_API FNavigationDirtyAreasController();
 
-	void Reset();
+	NAVIGATIONSYSTEM_API void Reset();
 	
 	/** sets cumulative time to at least one cycle so next tick will rebuild dirty areas */
-	void ForceRebuildOnNextTick();
+	NAVIGATIONSYSTEM_API void ForceRebuildOnNextTick();
 
-	void Tick(float DeltaSeconds, const TArray<ANavigationData*>& NavDataSet, bool bForceRebuilding = false);
+	NAVIGATIONSYSTEM_API void Tick(float DeltaSeconds, const TArray<ANavigationData*>& NavDataSet, bool bForceRebuilding = false);
 
 	/** Add a dirty area to the queue based on the provided bounds and flags.
 	 * Bounds must be valid and non empty otherwise the request will be ignored and a warning reported.
@@ -64,20 +64,20 @@ public:
 	 *	@param DirtyElement Optional dirty element
 	 *	@param DebugReason Source of the new area
 	 */
-	void AddArea(const FBox& NewArea, const int32 Flags, const TFunction<UObject*()>& ObjectProviderFunc = nullptr,
+	NAVIGATIONSYSTEM_API void AddArea(const FBox& NewArea, const int32 Flags, const TFunction<UObject*()>& ObjectProviderFunc = nullptr,
 		const FNavigationDirtyElement* DirtyElement = nullptr, const FName& DebugReason = NAME_None);
 	
 	bool IsDirty() const { return GetNumDirtyAreas() > 0; }
 	int32 GetNumDirtyAreas() const { return DirtyAreas.Num(); }
 
-	void OnNavigationBuildLocked();
-	void OnNavigationBuildUnlocked();
+	NAVIGATIONSYSTEM_API void OnNavigationBuildLocked();
+	NAVIGATIONSYSTEM_API void OnNavigationBuildUnlocked();
 
-	void SetUseWorldPartitionedDynamicMode(bool bIsWPDynamic);
-	void SetCanReportOversizedDirtyArea(const bool bCanReport);
-	void SetDirtyAreaWarningSizeThreshold(const float Threshold);
+	NAVIGATIONSYSTEM_API void SetUseWorldPartitionedDynamicMode(bool bIsWPDynamic);
+	NAVIGATIONSYSTEM_API void SetCanReportOversizedDirtyArea(const bool bCanReport);
+	NAVIGATIONSYSTEM_API void SetDirtyAreaWarningSizeThreshold(const float Threshold);
 
 #if !UE_BUILD_SHIPPING
 	bool HadDirtyAreasReportedWhileAccumulationLocked() const { return bCanAccumulateDirtyAreas == false && bDirtyAreasReportedWhileAccumulationLocked; }
 #endif // UE_BUILD_SHIPPING
-};
\ No newline at end of file
+};
