[Engine]

diff --git a/Engine/Source/Runtime/NavigationSystem/Public/NavigationInvokerComponent.h b/Engine/Source/Runtime/NavigationSystem/Public/NavigationInvokerComponent.h
--- a/Engine/Source/Runtime/NavigationSystem/Public/NavigationInvokerComponent.h
+++ b/Engine/Source/Runtime/NavigationSystem/Public/NavigationInvokerComponent.h
@@ -12,8 +12,8 @@
 class UNavigationSystemV1;
 enum class ENavigationInvokerPriority : uint8;
 
-UCLASS(ClassGroup = (Navigation), meta = (BlueprintSpawnableComponent))
-class NAVIGATIONSYSTEM_API UNavigationInvokerComponent : public UActorComponent
+UCLASS(ClassGroup = (Navigation), meta = (BlueprintSpawnableComponent), MinimalAPI)
+class UNavigationInvokerComponent : public UActorComponent
 {
 	GENERATED_BODY()
 
@@ -34,18 +34,18 @@ protected:
 	ENavigationInvokerPriority Priority;
 
 public:
-	UNavigationInvokerComponent(const FObjectInitializer& ObjectInitializer = FObjectInitializer::Get());
+	NAVIGATIONSYSTEM_API UNavigationInvokerComponent(const FObjectInitializer& ObjectInitializer = FObjectInitializer::Get());
 
-	void RegisterWithNavigationSystem(UNavigationSystemV1& NavSys);
+	NAVIGATIONSYSTEM_API void RegisterWithNavigationSystem(UNavigationSystemV1& NavSys);
 
 	/** Sets generation/removal ranges. Doesn't force navigation system's update.
 	 *	Will get picked up the next time NavigationSystem::UpdateInvokers gets called */
-	void SetGenerationRadii(const float GenerationRadius, const float RemovalRadius);
+	NAVIGATIONSYSTEM_API void SetGenerationRadii(const float GenerationRadius, const float RemovalRadius);
 
 	float GetGenerationRadius() const { return TileGenerationRadius; }
 	float GetRemovalRadius() const { return TileRemovalRadius; }
 
-	virtual void Activate(bool bReset = false) override;
-	virtual void Deactivate() override;
-	virtual void PostInitProperties() override;
+	NAVIGATIONSYSTEM_API virtual void Activate(bool bReset = false) override;
+	NAVIGATIONSYSTEM_API virtual void Deactivate() override;
+	NAVIGATIONSYSTEM_API virtual void PostInitProperties() override;
 };
diff --git a/Engine/Source/Runtime/NavigationSystem/Public/NavigationOctree.h b/Engine/Source/Runtime/NavigationSystem/Public/NavigationOctree.h
--- a/Engine/Source/Runtime/NavigationSystem/Public/NavigationOctree.h
+++ b/Engine/Source/Runtime/NavigationSystem/Public/NavigationOctree.h
@@ -19,7 +19,7 @@ typedef FNavigationRelevantDataFilter FNavigationOctreeFilter;
 
 LLM_DECLARE_TAG(NavigationOctree);
 
-struct NAVIGATIONSYSTEM_API FNavigationOctreeElement
+struct FNavigationOctreeElement
 {
 	FBoxSphereBounds Bounds;
 	TSharedRef<FNavigationRelevantData, ESPMode::ThreadSafe> Data;
@@ -139,7 +139,7 @@ struct FNavigationOctreeSemantics
 	static void SetElementId(FOctree& OctreeOwner, const FNavigationOctreeElement& Element, FOctreeElementId2 Id);
 };
 
-class NAVIGATIONSYSTEM_API FNavigationOctree : public TOctree2<FNavigationOctreeElement, FNavigationOctreeSemantics>, public TSharedFromThis<FNavigationOctree, ESPMode::ThreadSafe>
+class FNavigationOctree : public TOctree2<FNavigationOctreeElement, FNavigationOctreeSemantics>, public TSharedFromThis<FNavigationOctree, ESPMode::ThreadSafe>
 {
 public:
 	DECLARE_DELEGATE_TwoParams(FNavigableGeometryComponentExportDelegate, UActorComponent*, FNavigationRelevantData&);
@@ -180,49 +180,49 @@ public:
 		INC_MEMORY_STAT_BY(STAT_Navigation_CollisionTreeMemory, OctreeSizeBytes);
 	}
 
-	FNavigationOctree(const FVector& Origin, FVector::FReal Radius);
-	virtual ~FNavigationOctree();
+	NAVIGATIONSYSTEM_API FNavigationOctree(const FVector& Origin, FVector::FReal Radius);
+	NAVIGATIONSYSTEM_API virtual ~FNavigationOctree();
 
 	/** Add new node and fill it with navigation export data */
-	void AddNode(UObject* ElementOb, INavRelevantInterface* NavElement, const FBox& Bounds, FNavigationOctreeElement& Data);
+	NAVIGATIONSYSTEM_API void AddNode(UObject* ElementOb, INavRelevantInterface* NavElement, const FBox& Bounds, FNavigationOctreeElement& Data);
 
 	/** Append new data to existing node */
-	void AppendToNode(const FOctreeElementId2& Id, INavRelevantInterface* NavElement, const FBox& Bounds, FNavigationOctreeElement& Data);
+	NAVIGATIONSYSTEM_API void AppendToNode(const FOctreeElementId2& Id, INavRelevantInterface* NavElement, const FBox& Bounds, FNavigationOctreeElement& Data);
 
 	/** Updates element bounds remove/add operation */
-	void UpdateNode(const FOctreeElementId2& Id, const FBox& NewBounds);
+	NAVIGATIONSYSTEM_API void UpdateNode(const FOctreeElementId2& Id, const FBox& NewBounds);
 
 	/** Remove node */
-	void RemoveNode(const FOctreeElementId2& Id);
+	NAVIGATIONSYSTEM_API void RemoveNode(const FOctreeElementId2& Id);
 
-	void SetNavigableGeometryStoringMode(ENavGeometryStoringMode NavGeometryMode);
+	NAVIGATIONSYSTEM_API void SetNavigableGeometryStoringMode(ENavGeometryStoringMode NavGeometryMode);
 
-	const FNavigationRelevantData* GetDataForID(const FOctreeElementId2& Id) const;
+	NAVIGATIONSYSTEM_API const FNavigationRelevantData* GetDataForID(const FOctreeElementId2& Id) const;
 
 	ENavGeometryStoringMode GetNavGeometryStoringMode() const
 	{
 		return bGatherGeometry ? StoreNavGeometry : SkipNavGeometry;
 	}
 
-	void SetDataGatheringMode(ENavDataGatheringModeConfig Mode);
+	NAVIGATIONSYSTEM_API void SetDataGatheringMode(ENavDataGatheringModeConfig Mode);
 	
 	// Lazy data gathering methods
-	bool IsLazyGathering(const INavRelevantInterface& ChildNavInterface) const;
-	void DemandLazyDataGathering(FNavigationRelevantData& ElementData);
-	void DemandChildLazyDataGathering(FNavigationRelevantData& ElementData, INavRelevantInterface& ChildNavInterface);
+	NAVIGATIONSYSTEM_API bool IsLazyGathering(const INavRelevantInterface& ChildNavInterface) const;
+	NAVIGATIONSYSTEM_API void DemandLazyDataGathering(FNavigationRelevantData& ElementData);
+	NAVIGATIONSYSTEM_API void DemandChildLazyDataGathering(FNavigationRelevantData& ElementData, INavRelevantInterface& ChildNavInterface);
 
 	FORCEINLINE static uint32 HashObject(const UObject& Object)
 	{
 		return Object.GetUniqueID();
 	}
 #if !UE_BUILD_SHIPPING	
-	void SetGatheringNavModifiersTimeLimitWarning(const float Threshold);
+	NAVIGATIONSYSTEM_API void SetGatheringNavModifiersTimeLimitWarning(const float Threshold);
 #endif // !UE_BUILD_SHIPPING	
 protected:
 	friend struct FNavigationOctreeController;
 	friend struct FNavigationOctreeSemantics;
 
-	void SetElementIdImpl(const uint32 OwnerUniqueId, FOctreeElementId2 Id);
+	NAVIGATIONSYSTEM_API void SetElementIdImpl(const uint32 OwnerUniqueId, FOctreeElementId2 Id);
 
 	TMap<uint32, FOctreeElementId2> ObjectToOctreeId;
 	ENavDataGatheringMode DefaultGeometryGatheringMode;
diff --git a/Engine/Source/Runtime/NavigationSystem/Public/NavigationOctreeController.h b/Engine/Source/Runtime/NavigationSystem/Public/NavigationOctreeController.h
--- a/Engine/Source/Runtime/NavigationSystem/Public/NavigationOctreeController.h
+++ b/Engine/Source/Runtime/NavigationSystem/Public/NavigationOctreeController.h
@@ -4,7 +4,7 @@
 #include "NavigationOctree.h"
 #include "AI/Navigation/NavigationDirtyElement.h"
 
-struct NAVIGATIONSYSTEM_API FNavigationOctreeController
+struct FNavigationOctreeController
 {
 	enum EOctreeUpdateMode
 	{
@@ -22,28 +22,28 @@ struct NAVIGATIONSYSTEM_API FNavigationOctreeController
 	/** if set, navoctree updates are ignored, use with caution! */
 	uint8 bNavOctreeLock : 1;
 
-	void SetNavigationOctreeLock(bool bLock);
-	bool HasPendingObjectNavOctreeId(UObject& Object) const;
-	void RemoveObjectsNavOctreeId(const UObject& Object);
-	void SetNavigableGeometryStoringMode(FNavigationOctree::ENavGeometryStoringMode NavGeometryMode);
-
-	void Reset();
-
-	const FNavigationOctree* GetOctree() const;
-	FNavigationOctree* GetMutableOctree();
-	const FOctreeElementId2* GetObjectsNavOctreeId(const UObject& Object) const;
-	bool GetNavOctreeElementData(const UObject& NodeOwner, int32& DirtyFlags, FBox& DirtyBounds);
-	const FNavigationRelevantData* GetDataForObject(const UObject& Object) const;
-	FNavigationRelevantData* GetMutableDataForObject(const UObject& Object);
-	bool HasObjectsNavOctreeId(const UObject& Object) const;
-	bool IsNavigationOctreeLocked() const;
+	NAVIGATIONSYSTEM_API void SetNavigationOctreeLock(bool bLock);
+	NAVIGATIONSYSTEM_API bool HasPendingObjectNavOctreeId(UObject& Object) const;
+	NAVIGATIONSYSTEM_API void RemoveObjectsNavOctreeId(const UObject& Object);
+	NAVIGATIONSYSTEM_API void SetNavigableGeometryStoringMode(FNavigationOctree::ENavGeometryStoringMode NavGeometryMode);
+
+	NAVIGATIONSYSTEM_API void Reset();
+
+	NAVIGATIONSYSTEM_API const FNavigationOctree* GetOctree() const;
+	NAVIGATIONSYSTEM_API FNavigationOctree* GetMutableOctree();
+	NAVIGATIONSYSTEM_API const FOctreeElementId2* GetObjectsNavOctreeId(const UObject& Object) const;
+	NAVIGATIONSYSTEM_API bool GetNavOctreeElementData(const UObject& NodeOwner, int32& DirtyFlags, FBox& DirtyBounds);
+	NAVIGATIONSYSTEM_API const FNavigationRelevantData* GetDataForObject(const UObject& Object) const;
+	NAVIGATIONSYSTEM_API FNavigationRelevantData* GetMutableDataForObject(const UObject& Object);
+	NAVIGATIONSYSTEM_API bool HasObjectsNavOctreeId(const UObject& Object) const;
+	NAVIGATIONSYSTEM_API bool IsNavigationOctreeLocked() const;
 	/** basically says if navoctree has been created already */
 	bool IsValid() const { return NavOctree.IsValid(); }
-	bool IsValidElement(const FOctreeElementId2& ElementId) const;
+	NAVIGATIONSYSTEM_API bool IsValidElement(const FOctreeElementId2& ElementId) const;
 	bool IsEmpty() const { return (IsValid() == false) || NavOctree->GetSizeBytes() == 0; }
 
 private:
-	static uint32 HashObject(const UObject& Object);
+	static NAVIGATIONSYSTEM_API uint32 HashObject(const UObject& Object);
 };
 
 //----------------------------------------------------------------------//
diff --git a/Engine/Source/Runtime/NavigationSystem/Public/NavigationPath.h b/Engine/Source/Runtime/NavigationSystem/Public/NavigationPath.h
--- a/Engine/Source/Runtime/NavigationSystem/Public/NavigationPath.h
+++ b/Engine/Source/Runtime/NavigationSystem/Public/NavigationPath.h
@@ -18,8 +18,8 @@ DECLARE_DYNAMIC_MULTICAST_DELEGATE_TwoParams(FOnNavigationPathUpdated, UNavigati
  *	UObject wrapper for FNavigationPath
  */
 
-UCLASS(BlueprintType)
-class NAVIGATIONSYSTEM_API UNavigationPath : public UObject
+UCLASS(BlueprintType, MinimalAPI)
+class UNavigationPath : public UObject
 {
 	GENERATED_UCLASS_BODY()
 
@@ -48,40 +48,40 @@ protected:
 public:
 
 	// UObject begin
-	virtual void BeginDestroy() override;
+	NAVIGATIONSYSTEM_API virtual void BeginDestroy() override;
 	// UObject end
 
 	UFUNCTION(BlueprintCallable, Category = "AI|Debug")
-	FString GetDebugString() const;
+	NAVIGATIONSYSTEM_API FString GetDebugString() const;
 
 	UFUNCTION(BlueprintCallable, Category = "AI|Debug")
-	void EnableDebugDrawing(bool bShouldDrawDebugData, FLinearColor PathColor = FLinearColor::White);
+	NAVIGATIONSYSTEM_API void EnableDebugDrawing(bool bShouldDrawDebugData, FLinearColor PathColor = FLinearColor::White);
 
 	/** if enabled path will request recalculation if it gets invalidated due to a change to underlying navigation */
 	UFUNCTION(BlueprintCallable, Category = "AI|Navigation")
-	void EnableRecalculationOnInvalidation(TEnumAsByte<ENavigationOptionFlag::Type> DoRecalculation);
+	NAVIGATIONSYSTEM_API void EnableRecalculationOnInvalidation(TEnumAsByte<ENavigationOptionFlag::Type> DoRecalculation);
 
 	UFUNCTION(BlueprintCallable, Category = "AI|Navigation")
-	double GetPathLength() const;
+	NAVIGATIONSYSTEM_API double GetPathLength() const;
 
 	UFUNCTION(BlueprintCallable, Category = "AI|Navigation")
-	double GetPathCost() const;
+	NAVIGATIONSYSTEM_API double GetPathCost() const;
 
 	UFUNCTION(BlueprintCallable, Category = "AI|Navigation")
-	bool IsPartial() const;
+	NAVIGATIONSYSTEM_API bool IsPartial() const;
 
 	UFUNCTION(BlueprintCallable, Category = "AI|Navigation")
-	bool IsValid() const;
+	NAVIGATIONSYSTEM_API bool IsValid() const;
 
 	UFUNCTION(BlueprintCallable, Category = "AI|Navigation")
-	bool IsStringPulled() const;
+	NAVIGATIONSYSTEM_API bool IsStringPulled() const;
 		
-	void SetPath(FNavPathSharedPtr NewSharedPath);
+	NAVIGATIONSYSTEM_API void SetPath(FNavPathSharedPtr NewSharedPath);
 	FNavPathSharedPtr GetPath() { return SharedPath; }
 
 protected:
-	void DrawDebug(UCanvas* Canvas, APlayerController*);
-	void OnPathEvent(FNavigationPath* Path, ENavPathEvent::Type PathEvent);
+	NAVIGATIONSYSTEM_API void DrawDebug(UCanvas* Canvas, APlayerController*);
+	NAVIGATIONSYSTEM_API void OnPathEvent(FNavigationPath* Path, ENavPathEvent::Type PathEvent);
 
-	void SetPathPointsFromPath(FNavigationPath& NativePath);
+	NAVIGATIONSYSTEM_API void SetPathPointsFromPath(FNavigationPath& NativePath);
 };
diff --git a/Engine/Source/Runtime/NavigationSystem/Public/NavigationSystem.h b/Engine/Source/Runtime/NavigationSystem/Public/NavigationSystem.h
--- a/Engine/Source/Runtime/NavigationSystem/Public/NavigationSystem.h
+++ b/Engine/Source/Runtime/NavigationSystem/Public/NavigationSystem.h
@@ -72,13 +72,13 @@ namespace FNavigationSystem
 	 */
 	typedef ANavigationData* (*FNavigationDataInstanceCreator)(UWorld*, const FNavDataConfig&);
 
-	struct NAVIGATIONSYSTEM_API FCustomLinkOwnerInfo
+	struct FCustomLinkOwnerInfo
 	{
 		FWeakObjectPtr LinkOwner;
 		INavLinkCustomInterface* LinkInterface;
 
 		FCustomLinkOwnerInfo() : LinkInterface(nullptr) {}
-		FCustomLinkOwnerInfo(INavLinkCustomInterface* Link);
+		NAVIGATIONSYSTEM_API FCustomLinkOwnerInfo(INavLinkCustomInterface* Link);
 
 		bool IsValid() const { return LinkOwner.IsValid(); }
 	};
@@ -121,17 +121,17 @@ TIME_SLICER.DebugSetSectionName(TIME_SLICE_FNAME);
 #define MARK_TIMESLICE_SECTION_DEBUG(TIME_SLICER, TIME_SLICE_FNAME) ;
 #endif
 
-class NAVIGATIONSYSTEM_API FNavRegenTimeSlicer
+class FNavRegenTimeSlicer
 {
 public:
 	/** Setup the initial values for a time slice. This can be called on an instance after TestTimeSliceFinished() has returned true and EndTimeSliceAndAdjustDuration() has been called */
-	void SetupTimeSlice(double SliceDuration);
+	NAVIGATIONSYSTEM_API void SetupTimeSlice(double SliceDuration);
 
 	/** 
 	 *  Starts the time slice, this can be called multiple times as long as EndTimeSliceAndAdjustDuration() is called between each call.
 	 *  StartTimeSlice should not be called after TestTimeSliceFinished() has returned true
 	 */
-	void StartTimeSlice();
+	NAVIGATIONSYSTEM_API void StartTimeSlice();
 
 	/** 
 	 *  Useful when multiple sections of code need to be timesliced per frame using the same time slice duration that do not necessarily occur concurrently.
@@ -140,9 +140,9 @@ public:
 	 *  This can be called multiple times as long as StartTimeSlice() is called before EndTimeSliceAndAdjustDuration().
 	 *  EndTimeSliceAndAdjustDuration can be called after TestTimeSliceFinished() has returned true in this case the RemainingDuration will just be zero
 	 */
-	void EndTimeSliceAndAdjustDuration();
+	NAVIGATIONSYSTEM_API void EndTimeSliceAndAdjustDuration();
 	double GetStartTime() const { return StartTime; }
-	bool TestTimeSliceFinished() const;
+	NAVIGATIONSYSTEM_API bool TestTimeSliceFinished() const;
 
 	//* Returns the cached result of calling TestTimeSliceFinished, false by default */
 	bool IsTimeSliceFinishedCached() const { return bTimeSliceFinishedCached; }
@@ -156,8 +156,8 @@ public:
 	 * @param LongTimeSliceFunction function called when a time slice is taking too long to process.
 	 * @param LongTimeSliceDuration when a time slice takes longer than this duration LongTimeSliceFunction will be called.
 	 */
-	void DebugSetLongTimeSliceData(TFunction<void(FName, double)> LongTimeSliceFunction, double LongTimeSliceDuration) const;
-	void DebugResetLongTimeSliceFunction() const;
+	NAVIGATIONSYSTEM_API void DebugSetLongTimeSliceData(TFunction<void(FName, double)> LongTimeSliceFunction, double LongTimeSliceDuration) const;
+	NAVIGATIONSYSTEM_API void DebugResetLongTimeSliceFunction() const;
 
 	/**
 	 *  Sets the debug name for a time sliced section of code.
@@ -204,10 +204,10 @@ struct FTileHistoryData
 };
 #endif // !UE_BUILD_SHIPPING
 
-class NAVIGATIONSYSTEM_API FNavRegenTimeSliceManager
+class FNavRegenTimeSliceManager
 {
 public:
-	FNavRegenTimeSliceManager();
+	NAVIGATIONSYSTEM_API FNavRegenTimeSliceManager();
 
 	void PushTileRegenTime(double NewTime) { MovingWindowTileRegenTime.PushValue(NewTime);  }
 
@@ -215,28 +215,28 @@ public:
 
 	double GetAverageDeltaTime() const { return MovingWindowDeltaTime.GetAverage(); }
 
-	void ResetTileWaitTimeArrays(const TArray<TObjectPtr<ANavigationData>>& NavDataSet);
-	void PushTileWaitTime(const int32 NavDataIndex, const double NewTime);
+	NAVIGATIONSYSTEM_API void ResetTileWaitTimeArrays(const TArray<TObjectPtr<ANavigationData>>& NavDataSet);
+	NAVIGATIONSYSTEM_API void PushTileWaitTime(const int32 NavDataIndex, const double NewTime);
 
 #if !UE_BUILD_SHIPPING
-	void ResetTileHistoryData(const TArray<TObjectPtr<ANavigationData>>& NavDataSet);
-	void PushTileHistoryData(const int32 NavDataIndex, const FTileHistoryData& TileData);
+	NAVIGATIONSYSTEM_API void ResetTileHistoryData(const TArray<TObjectPtr<ANavigationData>>& NavDataSet);
+	NAVIGATIONSYSTEM_API void PushTileHistoryData(const int32 NavDataIndex, const FTileHistoryData& TileData);
 #endif // UE_BUILD_SHIPPING
 
-	double GetAverageTileWaitTime(const int32 NavDataIndex) const;
-	void ResetTileWaitTime(const int32 NavDataIndex);
+	NAVIGATIONSYSTEM_API double GetAverageTileWaitTime(const int32 NavDataIndex) const;
+	NAVIGATIONSYSTEM_API void ResetTileWaitTime(const int32 NavDataIndex);
 
 	bool DoTimeSlicedUpdate() const { return bDoTimeSlicedUpdate; }
 
-	void CalcAverageDeltaTime(uint64 FrameNum);
+	NAVIGATIONSYSTEM_API void CalcAverageDeltaTime(uint64 FrameNum);
 
-	void CalcTimeSliceDuration(const TArray<TObjectPtr<ANavigationData>>& NavDataSet, int32 NumTilesToRegen, const TArray<double>& CurrentTileRegenDurations);
+	NAVIGATIONSYSTEM_API void CalcTimeSliceDuration(const TArray<TObjectPtr<ANavigationData>>& NavDataSet, int32 NumTilesToRegen, const TArray<double>& CurrentTileRegenDurations);
 
-	void SetMinTimeSliceDuration(double NewMinTimeSliceDuration);
+	NAVIGATIONSYSTEM_API void SetMinTimeSliceDuration(double NewMinTimeSliceDuration);
 
-	void SetMaxTimeSliceDuration(double NewMaxTimeSliceDuration);
+	NAVIGATIONSYSTEM_API void SetMaxTimeSliceDuration(double NewMaxTimeSliceDuration);
 
-	void SetMaxDesiredTileRegenDuration(float NewMaxDesiredTileRegenDuration);
+	NAVIGATIONSYSTEM_API void SetMaxDesiredTileRegenDuration(float NewMaxDesiredTileRegenDuration);
 
 	int32 GetNavDataIdx() const { return NavDataIdx;  }
 	void SetNavDataIdx(int32 InNavDataIdx) { NavDataIdx = InNavDataIdx; }
@@ -245,7 +245,7 @@ public:
 	const FNavRegenTimeSlicer& GetTimeSlicer() const { return TimeSlicer; }
 
 #if !UE_BUILD_SHIPPING	
-	void LogTileStatistics(const TArray<TObjectPtr<ANavigationData>>& NavDataSet) const;
+	NAVIGATIONSYSTEM_API void LogTileStatistics(const TArray<TObjectPtr<ANavigationData>>& NavDataSet) const;
 #endif // !UE_BUILD_SHIPPING
 	
 protected:
@@ -286,16 +286,16 @@ protected:
 	bool bDoTimeSlicedUpdate;
 };
 
-UCLASS(Within=World, config=Engine, defaultconfig)
-class NAVIGATIONSYSTEM_API UNavigationSystemV1 : public UNavigationSystemBase
+UCLASS(Within=World, config=Engine, defaultconfig, MinimalAPI)
+class UNavigationSystemV1 : public UNavigationSystemBase
 {
 	GENERATED_BODY()
 
 	friend UNavigationSystemModuleConfig;
 
 public:
-	UNavigationSystemV1(const FObjectInitializer& ObjectInitializer = FObjectInitializer::Get());
-	virtual ~UNavigationSystemV1();
+	NAVIGATIONSYSTEM_API UNavigationSystemV1(const FObjectInitializer& ObjectInitializer = FObjectInitializer::Get());
+	NAVIGATIONSYSTEM_API virtual ~UNavigationSystemV1();
 
 	UPROPERTY(Transient)
 	TObjectPtr<ANavigationData> MainNavData;
@@ -452,9 +452,9 @@ private:
 	TArray<FNavigationInvokerRaw> InvokerLocations;
 
 	double NextInvokersUpdateTime;
-	void UpdateInvokers();
+	NAVIGATIONSYSTEM_API void UpdateInvokers();
 
-	void DirtyTilesInBuildBounds();
+	NAVIGATIONSYSTEM_API void DirtyTilesInBuildBounds();
 
 public:
 	//----------------------------------------------------------------------//
@@ -462,70 +462,70 @@ public:
 	//----------------------------------------------------------------------//
 
 	UFUNCTION(BlueprintPure, Category = "AI|Navigation", meta = (WorldContext = "WorldContextObject"))
-	static UNavigationSystemV1* GetNavigationSystem(UObject* WorldContextObject);
+	static NAVIGATIONSYSTEM_API UNavigationSystemV1* GetNavigationSystem(UObject* WorldContextObject);
 	
 	/** Project a point onto the NavigationData */
 	UFUNCTION(BlueprintPure, Category = "AI|Navigation", meta = (WorldContext = "WorldContextObject", DisplayName = "ProjectPointToNavigation", ScriptName = "ProjectPointToNavigation"))
-	static bool K2_ProjectPointToNavigation(UObject* WorldContextObject, const FVector& Point, FVector& ProjectedLocation, ANavigationData* NavData, TSubclassOf<UNavigationQueryFilter> FilterClass, const FVector QueryExtent = FVector::ZeroVector);
+	static NAVIGATIONSYSTEM_API bool K2_ProjectPointToNavigation(UObject* WorldContextObject, const FVector& Point, FVector& ProjectedLocation, ANavigationData* NavData, TSubclassOf<UNavigationQueryFilter> FilterClass, const FVector QueryExtent = FVector::ZeroVector);
 
 	/** Generates a random location reachable from given Origin location.
 	 *	@return Return Value represents if the call was successful */
 	UFUNCTION(BlueprintPure, Category = "AI|Navigation", meta = (WorldContext = "WorldContextObject", DisplayName = "GetRandomReachablePointInRadius", ScriptName = "GetRandomReachablePointInRadius"))
-	static bool K2_GetRandomReachablePointInRadius(UObject* WorldContextObject, const FVector& Origin, FVector& RandomLocation, float Radius, ANavigationData* NavData = NULL, TSubclassOf<UNavigationQueryFilter> FilterClass = NULL);
+	static NAVIGATIONSYSTEM_API bool K2_GetRandomReachablePointInRadius(UObject* WorldContextObject, const FVector& Origin, FVector& RandomLocation, float Radius, ANavigationData* NavData = NULL, TSubclassOf<UNavigationQueryFilter> FilterClass = NULL);
 
 	/** Generates a random location in navigable space within given radius of Origin.
 	 *	@return Return Value represents if the call was successful */
 	UFUNCTION(BlueprintCallable, Category = "AI|Navigation", meta = (WorldContext = "WorldContextObject", DisplayName = "GetRandomLocationInNavigableRadius", ScriptName = "GetRandomLocationInNavigableRadius"))
-	static bool K2_GetRandomLocationInNavigableRadius(UObject* WorldContextObject, const FVector& Origin, FVector& RandomLocation, float Radius, ANavigationData* NavData = NULL, TSubclassOf<UNavigationQueryFilter> FilterClass = NULL);
+	static NAVIGATIONSYSTEM_API bool K2_GetRandomLocationInNavigableRadius(UObject* WorldContextObject, const FVector& Origin, FVector& RandomLocation, float Radius, ANavigationData* NavData = NULL, TSubclassOf<UNavigationQueryFilter> FilterClass = NULL);
 	
 	/** Potentially expensive. Use with caution. Consider using UPathFollowingComponent::GetRemainingPathCost instead */
 	UE_DEPRECATED(5.2, "Use new version with double")
-	static ENavigationQueryResult::Type GetPathCost(UObject* WorldContextObject, const FVector& PathStart, const FVector& PathEnd, float& PathCost, ANavigationData* NavData = NULL, TSubclassOf<UNavigationQueryFilter> FilterClass = NULL);
+	static NAVIGATIONSYSTEM_API ENavigationQueryResult::Type GetPathCost(UObject* WorldContextObject, const FVector& PathStart, const FVector& PathEnd, float& PathCost, ANavigationData* NavData = NULL, TSubclassOf<UNavigationQueryFilter> FilterClass = NULL);
 
 	/** Potentially expensive. Use with caution. Consider using UPathFollowingComponent::GetRemainingPathCost instead */
 	UFUNCTION(BlueprintPure, Category = "AI|Navigation", meta = (WorldContext = "WorldContextObject"))
-	static ENavigationQueryResult::Type GetPathCost(UObject* WorldContextObject, const FVector& PathStart, const FVector& PathEnd, double& PathCost, ANavigationData* NavData = NULL, TSubclassOf<UNavigationQueryFilter> FilterClass = NULL);
+	static NAVIGATIONSYSTEM_API ENavigationQueryResult::Type GetPathCost(UObject* WorldContextObject, const FVector& PathStart, const FVector& PathEnd, double& PathCost, ANavigationData* NavData = NULL, TSubclassOf<UNavigationQueryFilter> FilterClass = NULL);
 
 	/** Potentially expensive. Use with caution */
 	UE_DEPRECATED(5.2, "Use new version with double")
-	static ENavigationQueryResult::Type GetPathLength(UObject* WorldContextObject, const FVector& PathStart, const FVector& PathEnd, float& PathLength, ANavigationData* NavData = NULL, TSubclassOf<UNavigationQueryFilter> FilterClass = NULL);
+	static NAVIGATIONSYSTEM_API ENavigationQueryResult::Type GetPathLength(UObject* WorldContextObject, const FVector& PathStart, const FVector& PathEnd, float& PathLength, ANavigationData* NavData = NULL, TSubclassOf<UNavigationQueryFilter> FilterClass = NULL);
 
 	/** Potentially expensive. Use with caution */
 	UFUNCTION(BlueprintPure, Category = "AI|Navigation", meta = (WorldContext = "WorldContextObject"))
-	static ENavigationQueryResult::Type GetPathLength(UObject* WorldContextObject, const FVector& PathStart, const FVector& PathEnd, double& PathLength, ANavigationData* NavData = NULL, TSubclassOf<UNavigationQueryFilter> FilterClass = NULL);
+	static NAVIGATIONSYSTEM_API ENavigationQueryResult::Type GetPathLength(UObject* WorldContextObject, const FVector& PathStart, const FVector& PathEnd, double& PathLength, ANavigationData* NavData = NULL, TSubclassOf<UNavigationQueryFilter> FilterClass = NULL);
 
 	UFUNCTION(BlueprintPure, Category="AI|Navigation", meta=(WorldContext="WorldContextObject" ) )
-	static bool IsNavigationBeingBuilt(UObject* WorldContextObject);
+	static NAVIGATIONSYSTEM_API bool IsNavigationBeingBuilt(UObject* WorldContextObject);
 
 	UFUNCTION(BlueprintPure, Category = "AI|Navigation", meta = (WorldContext = "WorldContextObject"))
-	static bool IsNavigationBeingBuiltOrLocked(UObject* WorldContextObject);
+	static NAVIGATIONSYSTEM_API bool IsNavigationBeingBuiltOrLocked(UObject* WorldContextObject);
 
 	/** Finds path instantly, in a FindPath Synchronously. 
 	 *	@param PathfindingContext could be one of following: NavigationData (like Navmesh actor), Pawn or Controller. This parameter determines parameters of specific pathfinding query */
 	UFUNCTION(BlueprintCallable, Category = "AI|Navigation", meta = (WorldContext="WorldContextObject"))
-	static UNavigationPath* FindPathToLocationSynchronously(UObject* WorldContextObject, const FVector& PathStart, const FVector& PathEnd, AActor* PathfindingContext = NULL, TSubclassOf<UNavigationQueryFilter> FilterClass = NULL);
+	static NAVIGATIONSYSTEM_API UNavigationPath* FindPathToLocationSynchronously(UObject* WorldContextObject, const FVector& PathStart, const FVector& PathEnd, AActor* PathfindingContext = NULL, TSubclassOf<UNavigationQueryFilter> FilterClass = NULL);
 
 	/** Finds path instantly, in a FindPath Synchronously. Main advantage over FindPathToLocationSynchronously is that 
 	 *	the resulting path will automatically get updated if goal actor moves more than TetherDistance away from last path node
 	 *	@param PathfindingContext could be one of following: NavigationData (like Navmesh actor), Pawn or Controller. This parameter determines parameters of specific pathfinding query */
 	UFUNCTION(BlueprintCallable, Category = "AI|Navigation", meta = (WorldContext="WorldContextObject"))
-	static UNavigationPath* FindPathToActorSynchronously(UObject* WorldContextObject, const FVector& PathStart, AActor* GoalActor, float TetherDistance = 50.f, AActor* PathfindingContext = NULL, TSubclassOf<UNavigationQueryFilter> FilterClass = NULL);
+	static NAVIGATIONSYSTEM_API UNavigationPath* FindPathToActorSynchronously(UObject* WorldContextObject, const FVector& PathStart, AActor* GoalActor, float TetherDistance = 50.f, AActor* PathfindingContext = NULL, TSubclassOf<UNavigationQueryFilter> FilterClass = NULL);
 
 	/** Performs navigation raycast on NavigationData appropriate for given Querier.
 	 *	@param Querier if not passed default navigation data will be used
 	 *	@param HitLocation if line was obstructed this will be set to hit location. Otherwise it contains SegmentEnd
 	 *	@return true if line from RayStart to RayEnd was obstructed. Also, true when no navigation data present */
 	UFUNCTION(BlueprintCallable, Category="AI|Navigation", meta=(WorldContext="WorldContextObject" ))
-	static bool NavigationRaycast(UObject* WorldContextObject, const FVector& RayStart, const FVector& RayEnd, FVector& HitLocation, TSubclassOf<UNavigationQueryFilter> FilterClass = NULL, AController* Querier = NULL);
+	static NAVIGATIONSYSTEM_API bool NavigationRaycast(UObject* WorldContextObject, const FVector& RayStart, const FVector& RayEnd, FVector& HitLocation, TSubclassOf<UNavigationQueryFilter> FilterClass = NULL, AController* Querier = NULL);
 
 	/** will limit the number of simultaneously running navmesh tile generation jobs to specified number.
 	 *	@param MaxNumberOfJobs gets trimmed to be at least 1. You cannot use this function to pause navmesh generation */
 	UFUNCTION(BlueprintCallable, Category = "AI|Navigation")
-	void SetMaxSimultaneousTileGenerationJobsCount(int32 MaxNumberOfJobs);
+	NAVIGATIONSYSTEM_API void SetMaxSimultaneousTileGenerationJobsCount(int32 MaxNumberOfJobs);
 	
 	/** Brings limit of simultaneous navmesh tile generation jobs back to Project Setting's default value */
 	UFUNCTION(BlueprintCallable, Category = "AI|Navigation")
-	void ResetMaxSimultaneousTileGenerationJobsCount();
+	NAVIGATIONSYSTEM_API void ResetMaxSimultaneousTileGenerationJobsCount();
 
 	/** Registers given actor as a "navigation enforcer" which means navigation system will
 	 *	make sure navigation is being generated in specified radius around it.
@@ -533,25 +533,25 @@ public:
 	 *		to take advantage of this feature
 	 */
 	UFUNCTION(BlueprintCallable, Category = "AI|Navigation")
-	void RegisterNavigationInvoker(AActor* Invoker, float TileGenerationRadius = 3000, float TileRemovalRadius = 5000);
+	NAVIGATIONSYSTEM_API void RegisterNavigationInvoker(AActor* Invoker, float TileGenerationRadius = 3000, float TileRemovalRadius = 5000);
 
 	/** Removes given actor from the list of active navigation enforcers.
 	 *	@see RegisterNavigationInvoker for more details */
 	UFUNCTION(BlueprintCallable, Category = "AI|Navigation")
-	void UnregisterNavigationInvoker(AActor* Invoker);
+	NAVIGATIONSYSTEM_API void UnregisterNavigationInvoker(AActor* Invoker);
 
 	UFUNCTION(BlueprintCallable, Category = "AI|Navigation|Generation")
-	void SetGeometryGatheringMode(ENavDataGatheringModeConfig NewMode);
+	NAVIGATIONSYSTEM_API void SetGeometryGatheringMode(ENavDataGatheringModeConfig NewMode);
 
 	UFUNCTION(BlueprintCallable, Category = "AI|Navigation", meta=(DisplayName="ReplaceAreaInOctreeData"))
-	bool K2_ReplaceAreaInOctreeData(const UObject* Object, TSubclassOf<UNavArea> OldArea, TSubclassOf<UNavArea> NewArea);
+	NAVIGATIONSYSTEM_API bool K2_ReplaceAreaInOctreeData(const UObject* Object, TSubclassOf<UNavArea> OldArea, TSubclassOf<UNavArea> NewArea);
 
 	FORCEINLINE bool IsActiveTilesGenerationEnabled() const{ return bGenerateNavigationOnlyAroundNavigationInvokers; }
 	
 	/** delegate type for events that dirty the navigation data ( Params: const FBox& DirtyBounds ) */
 	DECLARE_MULTICAST_DELEGATE_OneParam(FOnNavigationDirty, const FBox&);
 	/** called after navigation influencing event takes place*/
-	static FOnNavigationDirty NavigationDirtyEvent;
+	static NAVIGATIONSYSTEM_API FOnNavigationDirty NavigationDirtyEvent;
 
 	enum ERegistrationResult
 	{
@@ -574,30 +574,30 @@ public:
 	};
 
 	//~ Begin UObject Interface
-	virtual void PostInitProperties() override;
-	static void AddReferencedObjects(UObject* InThis, FReferenceCollector& Collector);
+	NAVIGATIONSYSTEM_API virtual void PostInitProperties() override;
+	static NAVIGATIONSYSTEM_API void AddReferencedObjects(UObject* InThis, FReferenceCollector& Collector);
 #if WITH_EDITOR
-	virtual void PostEditChangeChainProperty(FPropertyChangedChainEvent& PropertyChangedEvent) override;
-	virtual void PostEditChangeProperty(struct FPropertyChangedEvent& PropertyChangedEvent) override;
+	NAVIGATIONSYSTEM_API virtual void PostEditChangeChainProperty(FPropertyChangedChainEvent& PropertyChangedEvent) override;
+	NAVIGATIONSYSTEM_API virtual void PostEditChangeProperty(struct FPropertyChangedEvent& PropertyChangedEvent) override;
 #endif // WITH_EDITOR
 	//~ End UObject Interface
 
-	virtual void Tick(float DeltaSeconds) override;	
+	NAVIGATIONSYSTEM_API virtual void Tick(float DeltaSeconds) override;	
 
 	UWorld* GetWorld() const override { return GetOuterUWorld(); }
 
 	UCrowdManagerBase* GetCrowdManager() const { return CrowdManager.Get(); }
 
 protected:
-	void CalcTimeSlicedUpdateData(TArray<double>& OutCurrentTimeSlicedBuildTaskDurations, TArray<bool>& OutIsTimeSlicingArray, bool& bOutAnyNonTimeSlicedGenerators, TArray<int32, TInlineAllocator<8>>& OutNumTimeSlicedRemainingBuildTasksArray);
+	NAVIGATIONSYSTEM_API void CalcTimeSlicedUpdateData(TArray<double>& OutCurrentTimeSlicedBuildTaskDurations, TArray<bool>& OutIsTimeSlicingArray, bool& bOutAnyNonTimeSlicedGenerators, TArray<int32, TInlineAllocator<8>>& OutNumTimeSlicedRemainingBuildTasksArray);
 
 	/** spawn new crowd manager */
-	virtual void CreateCrowdManager();
+	NAVIGATIONSYSTEM_API virtual void CreateCrowdManager();
 
 	/** Used to properly set navigation class for indicated agent and propagate information to other places
 	 *	(like project settings) that may need this information 
 	 */
-	void SetSupportedAgentsNavigationClass(int32 AgentIndex, TSubclassOf<ANavigationData> NavigationDataClass);
+	NAVIGATIONSYSTEM_API void SetSupportedAgentsNavigationClass(int32 AgentIndex, TSubclassOf<ANavigationData> NavigationDataClass);
 
 public:
 	//----------------------------------------------------------------------//
@@ -610,7 +610,7 @@ public:
 	 *	@param NavData optional navigation data that will be used instead of the one that would be deducted from AgentProperties
 	 *  @param Mode switch between normal and hierarchical path finding algorithms
 	 */
-	FPathFindingResult FindPathSync(const FNavAgentProperties& AgentProperties, FPathFindingQuery Query, EPathFindingMode::Type Mode = EPathFindingMode::Regular);
+	NAVIGATIONSYSTEM_API FPathFindingResult FindPathSync(const FNavAgentProperties& AgentProperties, FPathFindingQuery Query, EPathFindingMode::Type Mode = EPathFindingMode::Regular);
 
 	/** 
 	 *	Does a simple path finding from @StartLocation to @EndLocation on specified NavData. If none passed MainNavData will be used
@@ -618,7 +618,7 @@ public:
 	 *	@param NavData optional navigation data that will be used instead main navigation data
 	 *  @param Mode switch between normal and hierarchical path finding algorithms
 	 */
-	FPathFindingResult FindPathSync(FPathFindingQuery Query, EPathFindingMode::Type Mode = EPathFindingMode::Regular);
+	NAVIGATIONSYSTEM_API FPathFindingResult FindPathSync(FPathFindingQuery Query, EPathFindingMode::Type Mode = EPathFindingMode::Regular);
 
 	/** 
 	 *	Asynchronously looks for a path from @StartLocation to @EndLocation for agent with properties @AgentProperties. NavData actor appropriate for specified 
@@ -630,10 +630,10 @@ public:
 	 *  @param Mode switch between normal and hierarchical path finding algorithms
 	 *	@return request ID
 	 */
-	uint32 FindPathAsync(const FNavAgentProperties& AgentProperties, FPathFindingQuery Query, const FNavPathQueryDelegate& ResultDelegate, EPathFindingMode::Type Mode = EPathFindingMode::Regular);
+	NAVIGATIONSYSTEM_API uint32 FindPathAsync(const FNavAgentProperties& AgentProperties, FPathFindingQuery Query, const FNavPathQueryDelegate& ResultDelegate, EPathFindingMode::Type Mode = EPathFindingMode::Regular);
 
 	/** Removes query indicated by given ID from queue of path finding requests to process. */
-	void AbortAsyncFindPathRequest(uint32 AsynPathQueryID);
+	NAVIGATIONSYSTEM_API void AbortAsyncFindPathRequest(uint32 AsynPathQueryID);
 	
 	/** 
 	 *	Synchronously check if path between two points exists
@@ -641,52 +641,52 @@ public:
 	 *  @param Mode switch between normal and hierarchical path finding algorithms. @note Hierarchical mode ignores QueryFilter
 	 *	@return true if path exists
 	 */
-	bool TestPathSync(FPathFindingQuery Query, EPathFindingMode::Type Mode = EPathFindingMode::Regular, int32* NumVisitedNodes = NULL) const;
+	NAVIGATIONSYSTEM_API bool TestPathSync(FPathFindingQuery Query, EPathFindingMode::Type Mode = EPathFindingMode::Regular, int32* NumVisitedNodes = NULL) const;
 
 	/** Finds random point in navigable space
 	 *	@param ResultLocation Found point is put here
 	 *	@param NavData If NavData == NULL then MainNavData is used.
 	 *	@return true if any location found, false otherwise */
-	bool GetRandomPoint(FNavLocation& ResultLocation, ANavigationData* NavData = NULL, FSharedConstNavQueryFilter QueryFilter = NULL);
+	NAVIGATIONSYSTEM_API bool GetRandomPoint(FNavLocation& ResultLocation, ANavigationData* NavData = NULL, FSharedConstNavQueryFilter QueryFilter = NULL);
 
 	/** Finds random, reachable point in navigable space restricted to Radius around Origin
 	 *	@param ResultLocation Found point is put here
 	 *	@param NavData If NavData == NULL then MainNavData is used.
 	 *	@return true if any location found, false otherwise */
-	bool GetRandomReachablePointInRadius(const FVector& Origin, float Radius, FNavLocation& ResultLocation, ANavigationData* NavData = NULL, FSharedConstNavQueryFilter QueryFilter = NULL) const;
+	NAVIGATIONSYSTEM_API bool GetRandomReachablePointInRadius(const FVector& Origin, float Radius, FNavLocation& ResultLocation, ANavigationData* NavData = NULL, FSharedConstNavQueryFilter QueryFilter = NULL) const;
 
 	/** Finds random, point in navigable space restricted to Radius around Origin. Resulting location is not tested for reachability from the Origin
 	 *	@param ResultLocation Found point is put here
 	 *	@param NavData If NavData == NULL then MainNavData is used.
 	 *	@return true if any location found, false otherwise */
-	bool GetRandomPointInNavigableRadius(const FVector& Origin, float Radius, FNavLocation& ResultLocation, ANavigationData* NavData = NULL, FSharedConstNavQueryFilter QueryFilter = NULL) const;
+	NAVIGATIONSYSTEM_API bool GetRandomPointInNavigableRadius(const FVector& Origin, float Radius, FNavLocation& ResultLocation, ANavigationData* NavData = NULL, FSharedConstNavQueryFilter QueryFilter = NULL) const;
 	
 	/** Calculates a path from PathStart to PathEnd and retrieves its cost. 
 	 *	@NOTE potentially expensive, so use it with caution */
 	UE_DEPRECATED(5.2, "Use new version with FVector::FReal")
-	ENavigationQueryResult::Type GetPathCost(const FVector& PathStart, const FVector& PathEnd, float& PathCost, const ANavigationData* NavData = NULL, FSharedConstNavQueryFilter QueryFilter = NULL) const;
+	NAVIGATIONSYSTEM_API ENavigationQueryResult::Type GetPathCost(const FVector& PathStart, const FVector& PathEnd, float& PathCost, const ANavigationData* NavData = NULL, FSharedConstNavQueryFilter QueryFilter = NULL) const;
 
 	/** Calculates a path from PathStart to PathEnd and retrieves its cost.
 	 *	@NOTE potentially expensive, so use it with caution */
-	ENavigationQueryResult::Type GetPathCost(const FVector& PathStart, const FVector& PathEnd, FVector::FReal& PathCost, const ANavigationData* NavData = NULL, FSharedConstNavQueryFilter QueryFilter = NULL) const;
+	NAVIGATIONSYSTEM_API ENavigationQueryResult::Type GetPathCost(const FVector& PathStart, const FVector& PathEnd, FVector::FReal& PathCost, const ANavigationData* NavData = NULL, FSharedConstNavQueryFilter QueryFilter = NULL) const;
 
 	/** Calculates a path from PathStart to PathEnd and retrieves its overestimated length.
 	*	@NOTE potentially expensive, so use it with caution */
 	UE_DEPRECATED(5.2, "Use new version with FVector::FReal")
-	ENavigationQueryResult::Type GetPathLength(const FVector& PathStart, const FVector& PathEnd, float& PathLength, const ANavigationData* NavData = NULL, FSharedConstNavQueryFilter QueryFilter = NULL) const;
+	NAVIGATIONSYSTEM_API ENavigationQueryResult::Type GetPathLength(const FVector& PathStart, const FVector& PathEnd, float& PathLength, const ANavigationData* NavData = NULL, FSharedConstNavQueryFilter QueryFilter = NULL) const;
 
 	/** Calculates a path from PathStart to PathEnd and retrieves its overestimated length.
 	 *	@NOTE potentially expensive, so use it with caution */
-	ENavigationQueryResult::Type GetPathLength(const FVector& PathStart, const FVector& PathEnd, FVector::FReal& PathLength, const ANavigationData* NavData = NULL, FSharedConstNavQueryFilter QueryFilter = NULL) const;
+	NAVIGATIONSYSTEM_API ENavigationQueryResult::Type GetPathLength(const FVector& PathStart, const FVector& PathEnd, FVector::FReal& PathLength, const ANavigationData* NavData = NULL, FSharedConstNavQueryFilter QueryFilter = NULL) const;
 
 	/** Calculates a path from PathStart to PathEnd and retrieves its overestimated length and cost.
 	*	@NOTE potentially expensive, so use it with caution */
 	UE_DEPRECATED(5.2, "Use new version with FVector::FReal")
-	ENavigationQueryResult::Type GetPathLengthAndCost(const FVector& PathStart, const FVector& PathEnd, float& PathLength, float& PathCost, const ANavigationData* NavData = NULL, FSharedConstNavQueryFilter QueryFilter = NULL) const;
+	NAVIGATIONSYSTEM_API ENavigationQueryResult::Type GetPathLengthAndCost(const FVector& PathStart, const FVector& PathEnd, float& PathLength, float& PathCost, const ANavigationData* NavData = NULL, FSharedConstNavQueryFilter QueryFilter = NULL) const;
 
 	/** Calculates a path from PathStart to PathEnd and retrieves its overestimated length and cost.
 	 *	@NOTE potentially expensive, so use it with caution */
-	ENavigationQueryResult::Type GetPathLengthAndCost(const FVector& PathStart, const FVector& PathEnd, FVector::FReal& PathLength, FVector::FReal& PathCost, const ANavigationData* NavData = NULL, FSharedConstNavQueryFilter QueryFilter = NULL) const;
+	NAVIGATIONSYSTEM_API ENavigationQueryResult::Type GetPathLengthAndCost(const FVector& PathStart, const FVector& PathEnd, FVector::FReal& PathLength, FVector::FReal& PathCost, const ANavigationData* NavData = NULL, FSharedConstNavQueryFilter QueryFilter = NULL) const;
 
 	// @todo document
 	bool ProjectPointToNavigation(const FVector& Point, FNavLocation& OutLocation, const FVector& Extent = INVALID_NAVEXTENT, const FNavAgentProperties* AgentProperties = NULL, FSharedConstNavQueryFilter QueryFilter = NULL)
@@ -695,21 +695,21 @@ public:
 	}
 
 	// @todo document
-	bool ProjectPointToNavigation(const FVector& Point, FNavLocation& OutLocation, const FVector& Extent = INVALID_NAVEXTENT, const ANavigationData* NavData = NULL, FSharedConstNavQueryFilter QueryFilter = NULL) const;
+	NAVIGATIONSYSTEM_API bool ProjectPointToNavigation(const FVector& Point, FNavLocation& OutLocation, const FVector& Extent = INVALID_NAVEXTENT, const ANavigationData* NavData = NULL, FSharedConstNavQueryFilter QueryFilter = NULL) const;
 
 	/** 
 	 * Looks for NavData generated for specified movement properties and returns it. NULL if not found;
 	 */
-	virtual ANavigationData* GetNavDataForProps(const FNavAgentProperties& AgentProperties);
+	NAVIGATIONSYSTEM_API virtual ANavigationData* GetNavDataForProps(const FNavAgentProperties& AgentProperties);
 
 	/** 
 	 * Looks for NavData generated for specified movement properties and returns it. NULL if not found; Const version.
 	 */
-	virtual const ANavigationData* GetNavDataForProps(const FNavAgentProperties& AgentProperties) const;
+	NAVIGATIONSYSTEM_API virtual const ANavigationData* GetNavDataForProps(const FNavAgentProperties& AgentProperties) const;
 
 	/** Goes through all registered NavigationData instances and retrieves the one 
 	 *	supporting agent named AgentName */
-	virtual ANavigationData* GetNavDataForAgentName(const FName AgentName) const;
+	NAVIGATIONSYSTEM_API virtual ANavigationData* GetNavDataForAgentName(const FName AgentName) const;
 
 	/**
 	 * Looks up NavData appropriate for specified movement properties and returns it. NULL if not found;
@@ -717,24 +717,24 @@ public:
 	 *	(AgentLocation and Extent parameters not being used) but NaV2 will take advantage of it and all engine-level
 	 * AI navigation code is going to use call this flavor.
 	 */
-	virtual ANavigationData* GetNavDataForProps(const FNavAgentProperties& AgentProperties, const FVector& AgentLocation, const FVector& Extent = INVALID_NAVEXTENT) const;
+	NAVIGATIONSYSTEM_API virtual ANavigationData* GetNavDataForProps(const FNavAgentProperties& AgentProperties, const FVector& AgentLocation, const FVector& Extent = INVALID_NAVEXTENT) const;
 
 	/** Returns the world default navigation data instance. Creates one if it doesn't exist. */
-	ANavigationData* GetDefaultNavDataInstance(FNavigationSystem::ECreateIfMissing CreateNewIfNoneFound);
+	NAVIGATIONSYSTEM_API ANavigationData* GetDefaultNavDataInstance(FNavigationSystem::ECreateIfMissing CreateNewIfNoneFound);
 	/** Returns the world default navigation data instance. */
 	virtual INavigationDataInterface* GetMainNavData() const override { return Cast<INavigationDataInterface>(GetDefaultNavDataInstance()); }
 	ANavigationData& GetMainNavDataChecked() const { check(MainNavData); return *MainNavData; }
 
 	/** Set limiting bounds to be used when building navigation data. */
-	virtual void SetBuildBounds(const FBox& Bounds) override;
+	NAVIGATIONSYSTEM_API virtual void SetBuildBounds(const FBox& Bounds) override;
 
-	virtual FBox GetNavigableWorldBounds() const override;
+	NAVIGATIONSYSTEM_API virtual FBox GetNavigableWorldBounds() const override;
 
-	virtual bool ContainsNavData(const FBox& Bounds) const override;
-	virtual FBox ComputeNavDataBounds() const override;
-	virtual void AddNavigationDataChunk(class ANavigationDataChunkActor& DataChunkActor) override;
-	virtual void RemoveNavigationDataChunk(class ANavigationDataChunkActor& DataChunkActor) override;
-	virtual void FillNavigationDataChunkActor(const FBox& QueryBounds, class ANavigationDataChunkActor& DataChunkActor, FBox& OutTilesBounds) override;
+	NAVIGATIONSYSTEM_API virtual bool ContainsNavData(const FBox& Bounds) const override;
+	NAVIGATIONSYSTEM_API virtual FBox ComputeNavDataBounds() const override;
+	NAVIGATIONSYSTEM_API virtual void AddNavigationDataChunk(class ANavigationDataChunkActor& DataChunkActor) override;
+	NAVIGATIONSYSTEM_API virtual void RemoveNavigationDataChunk(class ANavigationDataChunkActor& DataChunkActor) override;
+	NAVIGATIONSYSTEM_API virtual void FillNavigationDataChunkActor(const FBox& QueryBounds, class ANavigationDataChunkActor& DataChunkActor, FBox& OutTilesBounds) override;
 
 	ANavigationData* GetDefaultNavDataInstance() const { return MainNavData; }
 
@@ -744,17 +744,17 @@ public:
 	 *	for supplied NavConfig. If Level == null and bSpawnNavDataInNavBoundsLevel == true
 	 *	then the first volume actor in RegisteredNavBounds will be used to source the level. 
 	 *	Otherwise the navdata instance will be spawned in NavigationSystem's world */
-	virtual ANavigationData* CreateNavigationDataInstanceInLevel(const FNavDataConfig& NavConfig, ULevel* SpawnLevel);
+	NAVIGATIONSYSTEM_API virtual ANavigationData* CreateNavigationDataInstanceInLevel(const FNavDataConfig& NavConfig, ULevel* SpawnLevel);
 
-	FSharedNavQueryFilter CreateDefaultQueryFilterCopy() const;
+	NAVIGATIONSYSTEM_API FSharedNavQueryFilter CreateDefaultQueryFilterCopy() const;
 
 	/** Super-hacky safety feature for threaded navmesh building. Will be gone once figure out why keeping TSharedPointer to Navigation Generator doesn't 
 	 *	guarantee its existence */
-	bool ShouldGeneratorRun(const FNavDataGenerator* Generator) const;
+	NAVIGATIONSYSTEM_API bool ShouldGeneratorRun(const FNavDataGenerator* Generator) const;
 
-	virtual bool IsNavigationBuilt(const AWorldSettings* Settings) const override;
+	NAVIGATIONSYSTEM_API virtual bool IsNavigationBuilt(const AWorldSettings* Settings) const override;
 
-	virtual bool IsThereAnywhereToBuildNavigation() const;
+	NAVIGATIONSYSTEM_API virtual bool IsThereAnywhereToBuildNavigation() const;
 
 	bool ShouldGenerateNavigationEverywhere() const { return bWholeWorldNavigable; }
 
@@ -762,56 +762,56 @@ public:
 	virtual bool ShouldLoadNavigationOnClient(ANavigationData* NavData = nullptr) const { return bAllowClientSideNavigation; }
 	virtual bool ShouldDiscardSubLevelNavData(ANavigationData* NavData = nullptr) const { return bShouldDiscardSubLevelNavData; }
 
-	FBox GetWorldBounds() const;
+	NAVIGATIONSYSTEM_API FBox GetWorldBounds() const;
 	
-	FBox GetLevelBounds(ULevel* InLevel) const;
+	NAVIGATIONSYSTEM_API FBox GetLevelBounds(ULevel* InLevel) const;
 
-	bool IsNavigationRelevant(const AActor* TestActor) const;
+	NAVIGATIONSYSTEM_API bool IsNavigationRelevant(const AActor* TestActor) const;
 
-	const TSet<FNavigationBounds>& GetNavigationBounds() const;
+	NAVIGATIONSYSTEM_API const TSet<FNavigationBounds>& GetNavigationBounds() const;
 
-	static const FNavDataConfig& GetDefaultSupportedAgent();
-	static const FNavDataConfig& GetBiggestSupportedAgent(const UWorld* World);
+	static NAVIGATIONSYSTEM_API const FNavDataConfig& GetDefaultSupportedAgent();
+	static NAVIGATIONSYSTEM_API const FNavDataConfig& GetBiggestSupportedAgent(const UWorld* World);
 	
 #if WITH_EDITOR
-	static double GetWorldPartitionNavigationDataBuilderOverlap(const UWorld& World);
+	static NAVIGATIONSYSTEM_API double GetWorldPartitionNavigationDataBuilderOverlap(const UWorld& World);
 #endif
 	
-	const FNavDataConfig& GetDefaultSupportedAgentConfig() const;
+	NAVIGATIONSYSTEM_API const FNavDataConfig& GetDefaultSupportedAgentConfig() const;
 	FORCEINLINE const TArray<FNavDataConfig>& GetSupportedAgents() const { return SupportedAgents; }
-	void OverrideSupportedAgents(const TArray<FNavDataConfig>& NewSupportedAgents);
-	void SetSupportedAgentsMask(const FNavAgentSelector& InSupportedAgentsMask);
+	NAVIGATIONSYSTEM_API void OverrideSupportedAgents(const TArray<FNavDataConfig>& NewSupportedAgents);
+	NAVIGATIONSYSTEM_API void SetSupportedAgentsMask(const FNavAgentSelector& InSupportedAgentsMask);
 	FNavAgentSelector GetSupportedAgentsMask() const { return SupportedAgentsMask; }
 
-	virtual void ApplyWorldOffset(const FVector& InOffset, bool bWorldShift) override;
+	NAVIGATIONSYSTEM_API virtual void ApplyWorldOffset(const FVector& InOffset, bool bWorldShift) override;
 
 	/** checks if navigation/navmesh is dirty and needs to be rebuilt */
-	bool IsNavigationDirty() const;
+	NAVIGATIONSYSTEM_API bool IsNavigationDirty() const;
 
 	/** checks if dirty navigation data can rebuild itself */
-	bool CanRebuildDirtyNavigation() const;
+	NAVIGATIONSYSTEM_API bool CanRebuildDirtyNavigation() const;
 
 	FORCEINLINE bool SupportsNavigationGeneration() const { return bSupportRebuilding; }
 
-	static bool DoesPathIntersectBox(const FNavigationPath* Path, const FBox& Box, uint32 StartingIndex = 0, FVector* AgentExtent = NULL);
-	static bool DoesPathIntersectBox(const FNavigationPath* Path, const FBox& Box, const FVector& AgentLocation, uint32 StartingIndex = 0, FVector* AgentExtent = NULL);
+	static NAVIGATIONSYSTEM_API bool DoesPathIntersectBox(const FNavigationPath* Path, const FBox& Box, uint32 StartingIndex = 0, FVector* AgentExtent = NULL);
+	static NAVIGATIONSYSTEM_API bool DoesPathIntersectBox(const FNavigationPath* Path, const FBox& Box, const FVector& AgentLocation, uint32 StartingIndex = 0, FVector* AgentExtent = NULL);
 	
 	//----------------------------------------------------------------------//
 	// Active tiles
 	//----------------------------------------------------------------------//
 	UE_DEPRECATED(5.2, "This function is deprecated. Please use the new RegisterInvoker method with the Agents parameter (FNavAgentSelector() can be used as default value to keep the same behavior)")
-	virtual void RegisterInvoker(AActor& Invoker, float TileGenerationRadius, float TileRemovalRadius);
+	NAVIGATIONSYSTEM_API virtual void RegisterInvoker(AActor& Invoker, float TileGenerationRadius, float TileRemovalRadius);
 
 	UE_DEPRECATED(5.3, "This function is deprecated. Please use the new RegisterInvoker method with invoker priority.")
-	virtual void RegisterInvoker(AActor& Invoker, float TileGenerationRadius, float TileRemovalRadius, const FNavAgentSelector& Agents);
+	NAVIGATIONSYSTEM_API virtual void RegisterInvoker(AActor& Invoker, float TileGenerationRadius, float TileRemovalRadius, const FNavAgentSelector& Agents);
 
-	virtual void RegisterInvoker(AActor& Invoker, float TileGenerationRadius, float TileRemovalRadius, const FNavAgentSelector& Agents, ENavigationInvokerPriority InPriority);
+	NAVIGATIONSYSTEM_API virtual void RegisterInvoker(AActor& Invoker, float TileGenerationRadius, float TileRemovalRadius, const FNavAgentSelector& Agents, ENavigationInvokerPriority InPriority);
 	
-	virtual void UnregisterInvoker(AActor& Invoker);
+	NAVIGATIONSYSTEM_API virtual void UnregisterInvoker(AActor& Invoker);
 
-	static void RegisterNavigationInvoker(AActor& Invoker, float TileGenerationRadius, float TileRemovalRadius, 
+	static NAVIGATIONSYSTEM_API void RegisterNavigationInvoker(AActor& Invoker, float TileGenerationRadius, float TileRemovalRadius, 
 		const FNavAgentSelector& Agents = FNavAgentSelector(), ENavigationInvokerPriority Priority = ENavigationInvokerPriority::Default);
-	static void UnregisterNavigationInvoker(AActor& Invoker);
+	static NAVIGATIONSYSTEM_API void UnregisterNavigationInvoker(AActor& Invoker);
 
 	const TArray<FNavigationInvokerRaw>& GetInvokerLocations() const { return InvokerLocations; }
 
@@ -820,72 +820,72 @@ public:
 	//----------------------------------------------------------------------//
 	
 	// @todo document
-	virtual void UnregisterNavData(ANavigationData* NavData);
+	NAVIGATIONSYSTEM_API virtual void UnregisterNavData(ANavigationData* NavData);
 
 	/** Traverses SupportedAgents and for all agents not supported (i.e. filtered
 	 *	out by SupportedAgentsMask) checks if there's a currently registered
 	 *	NavigationData instance for that agent, and if so it unregisters that agent */
-	virtual void UnregisterUnusedNavData();
+	NAVIGATIONSYSTEM_API virtual void UnregisterUnusedNavData();
 
 	/** Adds NavData to registration candidates queue - NavDataRegistrationQueue*/
-	virtual void RequestRegistrationDeferred(ANavigationData& NavData);
+	NAVIGATIONSYSTEM_API virtual void RequestRegistrationDeferred(ANavigationData& NavData);
 
 protected:
-	void ApplySupportedAgentsFilter();
+	NAVIGATIONSYSTEM_API void ApplySupportedAgentsFilter();
 
 	/** Processes all NavigationData instances in UWorld owning navigation system instance, and registers
 	 *	all previously unregistered */
-	void RegisterNavigationDataInstances();
+	NAVIGATIONSYSTEM_API void RegisterNavigationDataInstances();
 
 	/** called in places where we need to spawn the NavOctree, but is checking additional conditions if we really want to do that
 	 *	depending on navigation data setup among others 
 	 *	@return true if NavOctree instance has been created, or if one is already present */
-	virtual bool ConditionalPopulateNavOctree();
+	NAVIGATIONSYSTEM_API virtual bool ConditionalPopulateNavOctree();
 
 	/** called to instantiate NavigationSystem's NavOctree instance */
-	virtual void ConstructNavOctree();
+	NAVIGATIONSYSTEM_API virtual void ConstructNavOctree();
 
 	/** Processes registration of candidates queues via RequestRegistration and stored in NavDataRegistrationQueue */
-	virtual void ProcessRegistrationCandidates();
+	NAVIGATIONSYSTEM_API virtual void ProcessRegistrationCandidates();
 
 	/** registers CustomLinks awaiting registration in PendingCustomLinkRegistration */
-	void ProcessCustomLinkPendingRegistration();
+	NAVIGATIONSYSTEM_API void ProcessCustomLinkPendingRegistration();
 
 	/** used to apply updates of nav volumes in navigation system's tick */
-	virtual void PerformNavigationBoundsUpdate(const TArray<FNavigationBoundsUpdateRequest>& UpdateRequests);
+	NAVIGATIONSYSTEM_API virtual void PerformNavigationBoundsUpdate(const TArray<FNavigationBoundsUpdateRequest>& UpdateRequests);
 	
 	/** adds data to RegisteredNavBounds */
-	void AddNavigationBounds(const FNavigationBounds& NewBounds);
+	NAVIGATIONSYSTEM_API void AddNavigationBounds(const FNavigationBounds& NewBounds);
 
 	/** Searches for all valid navigation bounds in the world and stores them */
-	virtual void GatherNavigationBounds();
+	NAVIGATIONSYSTEM_API virtual void GatherNavigationBounds();
 
 	/** @return pointer to ANavigationData instance of given ID, or NULL if it was not found. Note it looks only through registered navigation data */
-	ANavigationData* GetNavDataWithID(const uint16 NavDataID) const;
+	NAVIGATIONSYSTEM_API ANavigationData* GetNavDataWithID(const uint16 NavDataID) const;
 
-	static void RegisterComponentToNavOctree(UActorComponent* Comp);
-	static void UnregisterComponentToNavOctree(UActorComponent* Comp);
+	static NAVIGATIONSYSTEM_API void RegisterComponentToNavOctree(UActorComponent* Comp);
+	static NAVIGATIONSYSTEM_API void UnregisterComponentToNavOctree(UActorComponent* Comp);
 
 public:
-	virtual void ReleaseInitialBuildingLock();
+	NAVIGATIONSYSTEM_API virtual void ReleaseInitialBuildingLock();
 
 	//----------------------------------------------------------------------//
 	// navigation octree related functions
 	//----------------------------------------------------------------------//
-	static void OnComponentRegistered(UActorComponent* Comp);
-	static void OnComponentUnregistered(UActorComponent* Comp);
-	static void RegisterComponent(UActorComponent* Comp);
-	static void UnregisterComponent(UActorComponent* Comp);
-	static void OnActorRegistered(AActor* Actor);
-	static void OnActorUnregistered(AActor* Actor);
+	static NAVIGATIONSYSTEM_API void OnComponentRegistered(UActorComponent* Comp);
+	static NAVIGATIONSYSTEM_API void OnComponentUnregistered(UActorComponent* Comp);
+	static NAVIGATIONSYSTEM_API void RegisterComponent(UActorComponent* Comp);
+	static NAVIGATIONSYSTEM_API void UnregisterComponent(UActorComponent* Comp);
+	static NAVIGATIONSYSTEM_API void OnActorRegistered(AActor* Actor);
+	static NAVIGATIONSYSTEM_API void OnActorUnregistered(AActor* Actor);
 
 	/** update navoctree entry for specified actor/component */
-	static void UpdateActorInNavOctree(AActor& Actor);
-	static void UpdateComponentInNavOctree(UActorComponent& Comp);
+	static NAVIGATIONSYSTEM_API void UpdateActorInNavOctree(AActor& Actor);
+	static NAVIGATIONSYSTEM_API void UpdateComponentInNavOctree(UActorComponent& Comp);
 	/** update all navoctree entries for actor and its components */
-	static void UpdateActorAndComponentsInNavOctree(AActor& Actor, bool bUpdateAttachedActors = true);
+	static NAVIGATIONSYSTEM_API void UpdateActorAndComponentsInNavOctree(AActor& Actor, bool bUpdateAttachedActors = true);
 	/** update all navoctree entries for actor and its non scene components after root movement */
-	static void UpdateNavOctreeAfterMove(USceneComponent* Comp);
+	static NAVIGATIONSYSTEM_API void UpdateNavOctreeAfterMove(USceneComponent* Comp);
 
 protected:
 	/** A helper function that gathers all actors attached to RootActor and fetches 
@@ -894,22 +894,22 @@ protected:
 	 *		When done it is guaranteed to contain unique, non-null ptrs. It 
 	 *		will not include RootActor.
 	 *	@return the number of unique attached actors */
-	static int32 GetAllAttachedActors(const AActor& RootActor, TArray<AActor*>& OutAttachedActors);
+	static NAVIGATIONSYSTEM_API int32 GetAllAttachedActors(const AActor& RootActor, TArray<AActor*>& OutAttachedActors);
 
 	/** updates navoctree information on actors attached to RootActor */
-	static void UpdateAttachedActorsInNavOctree(AActor& RootActor);
+	static NAVIGATIONSYSTEM_API void UpdateAttachedActorsInNavOctree(AActor& RootActor);
 
 public:
 	/** removes all navoctree entries for actor and its components */
-	static void ClearNavOctreeAll(AActor* Actor);
+	static NAVIGATIONSYSTEM_API void ClearNavOctreeAll(AActor* Actor);
 	/** updates bounds of all components implementing INavRelevantInterface */
-	static void UpdateNavOctreeBounds(AActor* Actor);
+	static NAVIGATIONSYSTEM_API void UpdateNavOctreeBounds(AActor* Actor);
 
-	void AddDirtyArea(const FBox& NewArea, int32 Flags, const FName& DebugReason = NAME_None);
-	void AddDirtyArea(const FBox& NewArea, int32 Flags, const TFunction<UObject*()>& ObjectProviderFunc, const FName& DebugReason = NAME_None);
-	void AddDirtyAreas(const TArray<FBox>& NewAreas, int32 Flags, const FName& DebugReason = NAME_None);
-	bool HasDirtyAreasQueued() const;
-	int32 GetNumDirtyAreas() const;
+	NAVIGATIONSYSTEM_API void AddDirtyArea(const FBox& NewArea, int32 Flags, const FName& DebugReason = NAME_None);
+	NAVIGATIONSYSTEM_API void AddDirtyArea(const FBox& NewArea, int32 Flags, const TFunction<UObject*()>& ObjectProviderFunc, const FName& DebugReason = NAME_None);
+	NAVIGATIONSYSTEM_API void AddDirtyAreas(const TArray<FBox>& NewAreas, int32 Flags, const FName& DebugReason = NAME_None);
+	NAVIGATIONSYSTEM_API bool HasDirtyAreasQueued() const;
+	NAVIGATIONSYSTEM_API int32 GetNumDirtyAreas() const;
 	float GetDirtyAreaWarningSizeThreshold() const { return DirtyAreaWarningSizeThreshold; }
 
 	const FNavigationOctree* GetNavOctree() const { return DefaultOctreeController.GetOctree(); }
@@ -923,86 +923,86 @@ public:
 	FORCEINLINE bool HasPendingObjectNavOctreeId(UObject* Object) const { return Object && DefaultOctreeController.HasPendingObjectNavOctreeId(*Object); }
 	FORCEINLINE void RemoveObjectsNavOctreeId(const UObject& Object) { DefaultOctreeController.RemoveObjectsNavOctreeId(Object); }
 
-	void RemoveNavOctreeElementId(const FOctreeElementId2& ElementId, int32 UpdateFlags);
+	NAVIGATIONSYSTEM_API void RemoveNavOctreeElementId(const FOctreeElementId2& ElementId, int32 UpdateFlags);
 
-	const FNavigationRelevantData* GetDataForObject(const UObject& Object) const;
-	FNavigationRelevantData* GetMutableDataForObject(const UObject& Object);
+	NAVIGATIONSYSTEM_API const FNavigationRelevantData* GetDataForObject(const UObject& Object) const;
+	NAVIGATIONSYSTEM_API FNavigationRelevantData* GetMutableDataForObject(const UObject& Object);
 
 	/** find all elements in navigation octree within given box (intersection) */
-	void FindElementsInNavOctree(const FBox& QueryBox, const FNavigationOctreeFilter& Filter, TArray<FNavigationOctreeElement>& Elements);
+	NAVIGATIONSYSTEM_API void FindElementsInNavOctree(const FBox& QueryBox, const FNavigationOctreeFilter& Filter, TArray<FNavigationOctreeElement>& Elements);
 
 	/** update single element in navoctree */
-	void UpdateNavOctreeElement(UObject* ElementOwner, INavRelevantInterface* ElementInterface, int32 UpdateFlags);
+	NAVIGATIONSYSTEM_API void UpdateNavOctreeElement(UObject* ElementOwner, INavRelevantInterface* ElementInterface, int32 UpdateFlags);
 
 	/** force updating parent node and all its children */
-	void UpdateNavOctreeParentChain(UObject* ElementOwner, bool bSkipElementOwnerUpdate = false);
+	NAVIGATIONSYSTEM_API void UpdateNavOctreeParentChain(UObject* ElementOwner, bool bSkipElementOwnerUpdate = false);
 
 	/** update component bounds in navigation octree and mark only specified area as dirty, doesn't re-export component geometry */
-	bool UpdateNavOctreeElementBounds(UActorComponent* Comp, const FBox& NewBounds, const FBox& DirtyArea);
+	NAVIGATIONSYSTEM_API bool UpdateNavOctreeElementBounds(UActorComponent* Comp, const FBox& NewBounds, const FBox& DirtyArea);
 
 	/** fetched Object's data from the octree and replaces occurences of OldArea with NewArea */
-	bool ReplaceAreaInOctreeData(const UObject& Object, TSubclassOf<UNavArea> OldArea, TSubclassOf<UNavArea> NewArea, bool bReplaceChildClasses = false);
+	NAVIGATIONSYSTEM_API bool ReplaceAreaInOctreeData(const UObject& Object, TSubclassOf<UNavArea> OldArea, TSubclassOf<UNavArea> NewArea, bool bReplaceChildClasses = false);
 
 	//----------------------------------------------------------------------//
 	// Custom navigation links
 	//----------------------------------------------------------------------//
-	virtual void RegisterCustomLink(INavLinkCustomInterface& CustomLink);
-	void UnregisterCustomLink(INavLinkCustomInterface& CustomLink);
+	NAVIGATIONSYSTEM_API virtual void RegisterCustomLink(INavLinkCustomInterface& CustomLink);
+	NAVIGATIONSYSTEM_API void UnregisterCustomLink(INavLinkCustomInterface& CustomLink);
 	
-	static void RequestCustomLinkRegistering(INavLinkCustomInterface& CustomLink, UObject* OwnerOb);
-	static void RequestCustomLinkUnregistering(INavLinkCustomInterface& CustomLink, UObject* ObjectOb);
+	static NAVIGATIONSYSTEM_API void RequestCustomLinkRegistering(INavLinkCustomInterface& CustomLink, UObject* OwnerOb);
+	static NAVIGATIONSYSTEM_API void RequestCustomLinkUnregistering(INavLinkCustomInterface& CustomLink, UObject* ObjectOb);
 
 	UE_DEPRECATED(5.3, "LinkIds are now based on FNavLinkId call the version of this function that takes FNavLinkId. This function only returns nullptr.")
 	INavLinkCustomInterface* GetCustomLink(uint32 UniqueLinkId) const { return nullptr; };
 
 	/** find custom link by unique ID */
-	INavLinkCustomInterface* GetCustomLink(FNavLinkId UniqueLinkId) const;
+	NAVIGATIONSYSTEM_API INavLinkCustomInterface* GetCustomLink(FNavLinkId UniqueLinkId) const;
 
 	/** updates custom link for all active navigation data instances */
-	void UpdateCustomLink(const INavLinkCustomInterface* CustomLink);
+	NAVIGATIONSYSTEM_API void UpdateCustomLink(const INavLinkCustomInterface* CustomLink);
 
 	/** Return a Bounding Box containing the navlink points */
-	static FBox ComputeCustomLinkBounds(const INavLinkCustomInterface& CustomLink);
+	static NAVIGATIONSYSTEM_API FBox ComputeCustomLinkBounds(const INavLinkCustomInterface& CustomLink);
 
 	//----------------------------------------------------------------------//
 	// Areas
 	//----------------------------------------------------------------------//
-	static void RequestAreaRegistering(UClass* NavAreaClass);
-	static void RequestAreaUnregistering(UClass* NavAreaClass);
+	static NAVIGATIONSYSTEM_API void RequestAreaRegistering(UClass* NavAreaClass);
+	static NAVIGATIONSYSTEM_API void RequestAreaUnregistering(UClass* NavAreaClass);
 
 	/** find index in SupportedAgents array for given navigation data */
-	int32 GetSupportedAgentIndex(const ANavigationData* NavData) const;
+	NAVIGATIONSYSTEM_API int32 GetSupportedAgentIndex(const ANavigationData* NavData) const;
 
 	/** find index in SupportedAgents array for agent type */
-	int32 GetSupportedAgentIndex(const FNavAgentProperties& NavAgent) const;
+	NAVIGATIONSYSTEM_API int32 GetSupportedAgentIndex(const FNavAgentProperties& NavAgent) const;
 
 	//----------------------------------------------------------------------//
 	// Filters
 	//----------------------------------------------------------------------//
 	
 	/** prepare descriptions of navigation flags in UNavigationQueryFilter class: using enum */
-	void DescribeFilterFlags(UEnum* FlagsEnum) const;
+	NAVIGATIONSYSTEM_API void DescribeFilterFlags(UEnum* FlagsEnum) const;
 
 	/** prepare descriptions of navigation flags in UNavigationQueryFilter class: using array */
-	void DescribeFilterFlags(const TArray<FString>& FlagsDesc) const;
+	NAVIGATIONSYSTEM_API void DescribeFilterFlags(const TArray<FString>& FlagsDesc) const;
 
 	/** removes cached filters from currently registered navigation data */
-	void ResetCachedFilter(TSubclassOf<UNavigationQueryFilter> FilterClass);
+	NAVIGATIONSYSTEM_API void ResetCachedFilter(TSubclassOf<UNavigationQueryFilter> FilterClass);
 
 	//----------------------------------------------------------------------//
 	// building
 	//----------------------------------------------------------------------//
 	
 	/** Triggers navigation building on all eligible navigation data. */
-	virtual void Build();
+	NAVIGATIONSYSTEM_API virtual void Build();
 
 	/** Cancels all currently running navigation builds */
-	virtual void CancelBuild();
+	NAVIGATIONSYSTEM_API virtual void CancelBuild();
 
 	// @todo document
-	void OnPIEStart();
+	NAVIGATIONSYSTEM_API void OnPIEStart();
 	// @todo document
-	void OnPIEEnd();
+	NAVIGATIONSYSTEM_API void OnPIEEnd();
 	
 	// @todo document
 	FORCEINLINE bool IsNavigationBuildingLocked(uint8 Flags = ~0) const { return (NavBuildingLockFlags & Flags) != 0; }
@@ -1018,46 +1018,46 @@ public:
 
 	// @todo document
 	UFUNCTION(BlueprintCallable, Category = "AI|Navigation")
-	void OnNavigationBoundsUpdated(ANavMeshBoundsVolume* NavVolume);
-	virtual void OnNavigationBoundsAdded(ANavMeshBoundsVolume* NavVolume);
-	virtual void OnNavigationBoundsRemoved(ANavMeshBoundsVolume* NavVolume);
+	NAVIGATIONSYSTEM_API void OnNavigationBoundsUpdated(ANavMeshBoundsVolume* NavVolume);
+	NAVIGATIONSYSTEM_API virtual void OnNavigationBoundsAdded(ANavMeshBoundsVolume* NavVolume);
+	NAVIGATIONSYSTEM_API virtual void OnNavigationBoundsRemoved(ANavMeshBoundsVolume* NavVolume);
 
 	/** determines whether any generator is performing navigation building actions at the moment, dirty areas are also checked */
-	bool IsNavigationBuildInProgress();
+	NAVIGATIONSYSTEM_API bool IsNavigationBuildInProgress();
 	
-	virtual void OnNavigationGenerationFinished(ANavigationData& NavData);
+	NAVIGATIONSYSTEM_API virtual void OnNavigationGenerationFinished(ANavigationData& NavData);
 
 	/** Used to display "navigation building in progress" counter */
-	int32 GetNumRemainingBuildTasks() const;
+	NAVIGATIONSYSTEM_API int32 GetNumRemainingBuildTasks() const;
 
 	/** Number of currently running tasks */
-	int32 GetNumRunningBuildTasks() const;
+	NAVIGATIONSYSTEM_API int32 GetNumRunningBuildTasks() const;
 
 protected:
 	/** Sets up SuportedAgents and NavigationDataCreators. Override it to add additional setup, but make sure to call Super implementation */
-	virtual void DoInitialSetup();
+	NAVIGATIONSYSTEM_API virtual void DoInitialSetup();
 
 	/** Find or create abstract nav data */
-	virtual void UpdateAbstractNavData();
+	NAVIGATIONSYSTEM_API virtual void UpdateAbstractNavData();
 	
 	/** Called during ConditionalPopulateNavOctree and gives subclassess a chance
 	 *	to influence what gets added */
-	virtual void AddLevelToOctree(ULevel& Level);
+	NAVIGATIONSYSTEM_API virtual void AddLevelToOctree(ULevel& Level);
 
 	/** Called as part of UWorld::BeginTearingDown */
-	virtual void OnBeginTearingDown(UWorld* World);
+	NAVIGATIONSYSTEM_API virtual void OnBeginTearingDown(UWorld* World);
 	
 public:
 	/** Called upon UWorld destruction to release what needs to be released */
-	virtual void CleanUp(const FNavigationSystem::ECleanupMode Mode = FNavigationSystem::ECleanupMode::CleanupUnsafe) override;
+	NAVIGATIONSYSTEM_API virtual void CleanUp(const FNavigationSystem::ECleanupMode Mode = FNavigationSystem::ECleanupMode::CleanupUnsafe) override;
 
 	/** 
 	 *	Called when owner-UWorld initializes actors
 	 */
-	virtual void OnInitializeActors() override;
+	NAVIGATIONSYSTEM_API virtual void OnInitializeActors() override;
 
 	/** */
-	virtual void OnWorldInitDone(FNavigationSystemRunMode Mode);
+	NAVIGATIONSYSTEM_API virtual void OnWorldInitDone(FNavigationSystemRunMode Mode);
 
 	/** Returns true if world has been initialized. */
 	FORCEINLINE virtual bool IsWorldInitDone() const override { return bWorldInitDone; }
@@ -1066,7 +1066,7 @@ public:
 	FORCEINLINE FNavigationSystemRunMode GetRunMode() const { return OperationMode; }
 
 	/** adds BSP collisions of currently streamed in levels to octree */
-	void InitializeLevelCollisions();
+	NAVIGATIONSYSTEM_API void InitializeLevelCollisions();
 
 	enum class ELockRemovalRebuildAction
 	{
@@ -1074,10 +1074,10 @@ public:
 		RebuildIfNotInEditor,
 		NoRebuild
 	};
-	void AddNavigationBuildLock(uint8 Flags);
-	void RemoveNavigationBuildLock(uint8 Flags, const ELockRemovalRebuildAction RebuildAction = ELockRemovalRebuildAction::Rebuild);
+	NAVIGATIONSYSTEM_API void AddNavigationBuildLock(uint8 Flags);
+	NAVIGATIONSYSTEM_API void RemoveNavigationBuildLock(uint8 Flags, const ELockRemovalRebuildAction RebuildAction = ELockRemovalRebuildAction::Rebuild);
 
-	void SetNavigationOctreeLock(bool bLock);
+	NAVIGATIONSYSTEM_API void SetNavigationOctreeLock(bool bLock);
 
 	/** checks if auto-rebuilding navigation data is enabled. Defaults to bNavigationAutoUpdateEnabled
 	*	value, but can be overridden per nav sys instance */
@@ -1085,7 +1085,7 @@ public:
 
 #if WITH_EDITOR
 	/** allow editor to toggle whether seamless navigation building is enabled */
-	static void SetNavigationAutoUpdateEnabled(bool bNewEnable, UNavigationSystemBase* InNavigationSystem);
+	static NAVIGATIONSYSTEM_API void SetNavigationAutoUpdateEnabled(bool bNewEnable, UNavigationSystemBase* InNavigationSystem);
 
 	FORCEINLINE bool IsNavigationRegisterLocked() const { return NavUpdateLockFlags != 0; }
 	FORCEINLINE bool IsNavigationUnregisterLocked() const { return NavUpdateLockFlags && !(NavUpdateLockFlags & ENavigationLockReason::AllowUnregister); }
@@ -1093,21 +1093,21 @@ public:
 	FORCEINLINE void AddNavigationUpdateLock(uint8 Flags) { NavUpdateLockFlags |= Flags; }
 	FORCEINLINE void RemoveNavigationUpdateLock(uint8 Flags) { NavUpdateLockFlags &= ~Flags; }
 
-	void UpdateLevelCollision(ULevel* InLevel);
+	NAVIGATIONSYSTEM_API void UpdateLevelCollision(ULevel* InLevel);
 
 #endif // WITH_EDITOR
 
 	FORCEINLINE bool IsSetUpForLazyGeometryExporting() const { return bGenerateNavigationOnlyAroundNavigationInvokers; }
 
-	static UNavigationSystemV1* CreateNavigationSystem(UWorld* WorldOwner);
+	static NAVIGATIONSYSTEM_API UNavigationSystemV1* CreateNavigationSystem(UWorld* WorldOwner);
 
-	static UNavigationSystemV1* GetCurrent(UWorld* World);
-	static UNavigationSystemV1* GetCurrent(UObject* WorldContextObject);
+	static NAVIGATIONSYSTEM_API UNavigationSystemV1* GetCurrent(UWorld* World);
+	static NAVIGATIONSYSTEM_API UNavigationSystemV1* GetCurrent(UObject* WorldContextObject);
 	
-	virtual void InitializeForWorld(UWorld& World, FNavigationSystemRunMode Mode) override;
+	NAVIGATIONSYSTEM_API virtual void InitializeForWorld(UWorld& World, FNavigationSystemRunMode Mode) override;
 
 	// Fetch the array of all nav-agent properties.
-	void GetNavAgentPropertiesArray(TArray<FNavAgentProperties>& OutNavAgentProperties) const;
+	NAVIGATIONSYSTEM_API void GetNavAgentPropertiesArray(TArray<FNavAgentProperties>& OutNavAgentProperties) const;
 
 	static FORCEINLINE bool ShouldUpdateNavOctreeOnComponentChange()
 	{
@@ -1134,20 +1134,20 @@ public:
 	 *	your game. bEnableStatic = false will reset the mechanism.
 	 *	Note: this is not a runtime switch. It's highly advisable not to call it manually and 
 	 *	use UNavigationSystemModuleConfig.bStrictlyStatic instead */
-	static void ConfigureAsStatic(bool bEnableStatic = true);
+	static NAVIGATIONSYSTEM_API void ConfigureAsStatic(bool bEnableStatic = true);
 
-	static void SetUpdateNavOctreeOnComponentChange(bool bNewUpdateOnComponentChange);
+	static NAVIGATIONSYSTEM_API void SetUpdateNavOctreeOnComponentChange(bool bNewUpdateOnComponentChange);
 
 	/** 
 	 * Exec command handlers
 	 */
-	bool HandleCycleNavDrawnCommand( const TCHAR* Cmd, FOutputDevice& Ar );
-	bool HandleCountNavMemCommand();
+	NAVIGATIONSYSTEM_API bool HandleCycleNavDrawnCommand( const TCHAR* Cmd, FOutputDevice& Ar );
+	NAVIGATIONSYSTEM_API bool HandleCountNavMemCommand();
 	
 	//----------------------------------------------------------------------//
 	// debug
 	//----------------------------------------------------------------------//
-	void CycleNavigationDataDrawn();
+	NAVIGATIONSYSTEM_API void CycleNavigationDataDrawn();
 
 	FNavRegenTimeSliceManager& GetMutableNavRegenTimeSliceManager() { return NavRegenTimeSliceManager; }
 
@@ -1186,7 +1186,7 @@ protected:
 
 	// async queries
 	FCriticalSection NavDataRegistrationSection;
-	static FCriticalSection CustomLinkRegistrationSection;
+	static NAVIGATIONSYSTEM_API FCriticalSection CustomLinkRegistrationSection;
 	
 #if WITH_EDITOR
 	uint8 NavUpdateLockFlags;
@@ -1212,69 +1212,69 @@ protected:
 	
 #if !UE_BUILD_SHIPPING
 	/** self-registering exec command to handle nav sys console commands */
-	static FNavigationSystemExec ExecHandler;
+	static NAVIGATIONSYSTEM_API FNavigationSystemExec ExecHandler;
 #endif // !UE_BUILD_SHIPPING
 	
 	/** whether seamless navigation building is enabled */
-	static bool bNavigationAutoUpdateEnabled;
-	static bool bUpdateNavOctreeOnComponentChange;
-	static bool bStaticRuntimeNavigation;
-	static bool bIsPIEActive;
+	static NAVIGATIONSYSTEM_API bool bNavigationAutoUpdateEnabled;
+	static NAVIGATIONSYSTEM_API bool bUpdateNavOctreeOnComponentChange;
+	static NAVIGATIONSYSTEM_API bool bStaticRuntimeNavigation;
+	static NAVIGATIONSYSTEM_API bool bIsPIEActive;
 
-	static TMap<INavLinkCustomInterface*, FWeakObjectPtr> PendingCustomLinkRegistration;
+	static NAVIGATIONSYSTEM_API TMap<INavLinkCustomInterface*, FWeakObjectPtr> PendingCustomLinkRegistration;
 	TSet<TObjectPtr<const UClass>> NavAreaClasses;
 
 	FNavRegenTimeSliceManager NavRegenTimeSliceManager;
 
 	/** delegate handler for PostLoadMap event */
-	void OnPostLoadMap(UWorld* LoadedWorld);
+	NAVIGATIONSYSTEM_API void OnPostLoadMap(UWorld* LoadedWorld);
 #if WITH_EDITOR
 	/** delegate handler for ActorMoved events */
-	void OnActorMoved(AActor* Actor);
+	NAVIGATIONSYSTEM_API void OnActorMoved(AActor* Actor);
 #endif
 	/** delegate handler called when navigation is dirtied*/
-	void OnNavigationDirtied(const FBox& Bounds);
+	NAVIGATIONSYSTEM_API void OnNavigationDirtied(const FBox& Bounds);
 
 	FDelegateHandle ReloadCompleteDelegateHandle;
 
 	/** called to notify NavigaitonSystem about finished reload */
-	virtual void OnReloadComplete(EReloadCompleteReason Reason);
+	NAVIGATIONSYSTEM_API virtual void OnReloadComplete(EReloadCompleteReason Reason);
 
 	/** Registers given navigation data with this Navigation System.
 	 *	@return RegistrationSuccessful if registration was successful, other results mean it failed
 	 *	@see ERegistrationResult
 	 */
-	virtual ERegistrationResult RegisterNavData(ANavigationData* NavData);
+	NAVIGATIONSYSTEM_API virtual ERegistrationResult RegisterNavData(ANavigationData* NavData);
 
 	/** tries to register navigation area */
-	void RegisterNavAreaClass(UClass* NavAreaClass);
+	NAVIGATIONSYSTEM_API void RegisterNavAreaClass(UClass* NavAreaClass);
 
 	/** tries to unregister navigation area */
-	void UnregisterNavAreaClass(UClass* NavAreaClass);
+	NAVIGATIONSYSTEM_API void UnregisterNavAreaClass(UClass* NavAreaClass);
 
-	void OnNavigationAreaEvent(UClass* AreaClass, ENavAreaEvent::Type Event);
+	NAVIGATIONSYSTEM_API void OnNavigationAreaEvent(UClass* AreaClass, ENavAreaEvent::Type Event);
 	
- 	FSetElementId RegisterNavOctreeElement(UObject* ElementOwner, INavRelevantInterface* ElementInterface, int32 UpdateFlags);
-	void UnregisterNavOctreeElement(UObject* ElementOwner, INavRelevantInterface* ElementInterface, int32 UpdateFlags);
+ 	NAVIGATIONSYSTEM_API FSetElementId RegisterNavOctreeElement(UObject* ElementOwner, INavRelevantInterface* ElementInterface, int32 UpdateFlags);
+	NAVIGATIONSYSTEM_API void UnregisterNavOctreeElement(UObject* ElementOwner, INavRelevantInterface* ElementInterface, int32 UpdateFlags);
 	
 	/** read element data from navigation octree */
-	bool GetNavOctreeElementData(const UObject& NodeOwner, int32& DirtyFlags, FBox& DirtyBounds);
+	NAVIGATIONSYSTEM_API bool GetNavOctreeElementData(const UObject& NodeOwner, int32& DirtyFlags, FBox& DirtyBounds);
 	//bool GetNavOctreeElementData(UObject* NodeOwner, int32& DirtyFlags, FBox& DirtyBounds);
 
 	/** Adds given element to NavOctree. No check for owner's validity are performed, 
 	 *	nor its presence in NavOctree - function assumes callee responsibility 
 	 *	in this regard **/
-	void AddElementToNavOctree(const FNavigationDirtyElement& DirtyElement);
+	NAVIGATIONSYSTEM_API void AddElementToNavOctree(const FNavigationDirtyElement& DirtyElement);
 
-	void SetCrowdManager(UCrowdManagerBase* NewCrowdManager);
+	NAVIGATIONSYSTEM_API void SetCrowdManager(UCrowdManagerBase* NewCrowdManager);
 
 	/** Add BSP collision data to navigation octree */
-	void AddLevelCollisionToOctree(ULevel* Level);
+	NAVIGATIONSYSTEM_API void AddLevelCollisionToOctree(ULevel* Level);
 	
 	/** Remove BSP collision data from navigation octree */
-	void RemoveLevelCollisionFromOctree(ULevel* Level);
+	NAVIGATIONSYSTEM_API void RemoveLevelCollisionFromOctree(ULevel* Level);
 
-	virtual void SpawnMissingNavigationData();
+	NAVIGATIONSYSTEM_API virtual void SpawnMissingNavigationData();
 
 	/** 
 	 * Fills a mask indicating which navigation data associated to the supported agent mask are already instantiated.
@@ -1282,91 +1282,91 @@ protected:
 	 * @param InLevel If specified will be used to search for navigation data; will use the owning world otherwise.
 	 * @return Number of instantiated navigation data.
 	 */
-	uint8 FillInstantiatedDataMask(TBitArray<>& OutInstantiatedMask, ULevel* InLevel = nullptr);
+	NAVIGATIONSYSTEM_API uint8 FillInstantiatedDataMask(TBitArray<>& OutInstantiatedMask, ULevel* InLevel = nullptr);
 
 	/**
 	 * Spawns missing navigation data.
 	 * @param InInstantiatedMask The mask representing already instantiated navigation data.
 	 * @param InLevel Level in which the new data must be added. See CreateNavigationDataInstanceInLevel doc.
 	 */
-	void SpawnMissingNavigationDataInLevel(const TBitArray<>& InInstantiatedMask, ULevel* InLevel = nullptr);
+	NAVIGATIONSYSTEM_API void SpawnMissingNavigationDataInLevel(const TBitArray<>& InInstantiatedMask, ULevel* InLevel = nullptr);
 
 public:
-	void DemandLazyDataGathering(FNavigationRelevantData& ElementData);
+	NAVIGATIONSYSTEM_API void DemandLazyDataGathering(FNavigationRelevantData& ElementData);
 
 protected:
-	virtual void RebuildDirtyAreas(float DeltaSeconds);
+	NAVIGATIONSYSTEM_API virtual void RebuildDirtyAreas(float DeltaSeconds);
 
 	// adds navigation bounds update request to a pending list
-	void AddNavigationBoundsUpdateRequest(const FNavigationBoundsUpdateRequest& UpdateRequest);
+	NAVIGATIONSYSTEM_API void AddNavigationBoundsUpdateRequest(const FNavigationBoundsUpdateRequest& UpdateRequest);
 
 	/** Triggers navigation building on all eligible navigation data. */
-	virtual void RebuildAll(bool bIsLoadTime = false);
+	NAVIGATIONSYSTEM_API virtual void RebuildAll(bool bIsLoadTime = false);
 		 
 	/** Handler for FWorldDelegates::LevelAddedToWorld event */
-	 void OnLevelAddedToWorld(ULevel* InLevel, UWorld* InWorld);
+	 NAVIGATIONSYSTEM_API void OnLevelAddedToWorld(ULevel* InLevel, UWorld* InWorld);
 	 
 	/** Handler for FWorldDelegates::LevelRemovedFromWorld event */
-	void OnLevelRemovedFromWorld(ULevel* InLevel, UWorld* InWorld);
+	NAVIGATIONSYSTEM_API void OnLevelRemovedFromWorld(ULevel* InLevel, UWorld* InWorld);
 
 	/** Handler for FWorldDelegates::OnWorldPostActorTick event */
 	void OnWorldPostActorTick(UWorld* World, ELevelTick TickType, float DeltaTime) { PostponeAsyncQueries(); }
 
 	/** Adds given request to requests queue. Note it's to be called only on game thread only */
-	void AddAsyncQuery(const FAsyncPathFindingQuery& Query);
+	NAVIGATIONSYSTEM_API void AddAsyncQuery(const FAsyncPathFindingQuery& Query);
 		 
 	/** spawns a non-game-thread task to process requests given in PathFindingQueries.
 	 *	In the process PathFindingQueries gets copied. */
-	void TriggerAsyncQueries(TArray<FAsyncPathFindingQuery>& PathFindingQueries);
+	NAVIGATIONSYSTEM_API void TriggerAsyncQueries(TArray<FAsyncPathFindingQuery>& PathFindingQueries);
 
 	/** Processes pathfinding requests given in PathFindingQueries.*/
-	void PerformAsyncQueries(TArray<FAsyncPathFindingQuery> PathFindingQueries);
+	NAVIGATIONSYSTEM_API void PerformAsyncQueries(TArray<FAsyncPathFindingQuery> PathFindingQueries);
 
 	/** Broadcasts completion delegate for all completed async pathfinding requests. */
-	void DispatchAsyncQueriesResults(const TArray<FAsyncPathFindingQuery>& PathFindingQueries) const;
+	NAVIGATIONSYSTEM_API void DispatchAsyncQueriesResults(const TArray<FAsyncPathFindingQuery>& PathFindingQueries) const;
 
 	/**
 	 * Requests the async pathfinding task to abort and waits for it to complete
 	 * before resuming the main thread. Pathfind task will postpone remaining queries to next frame.
 	 */
-	void PostponeAsyncQueries();
+	NAVIGATIONSYSTEM_API void PostponeAsyncQueries();
 
 	/** */
-	virtual void DestroyNavOctree();
+	NAVIGATIONSYSTEM_API virtual void DestroyNavOctree();
 
 	/** Whether Navigation system needs to populate nav octree. 
 	 *  Depends on runtime generation settings of each navigation data, always true in the editor
 	 */
-	bool RequiresNavOctree() const;
+	NAVIGATIONSYSTEM_API bool RequiresNavOctree() const;
 
 	/** Return "Strongest" runtime generation type required by registered navigation data objects
 	 *  Depends on runtime generation settings of each navigation data, always ERuntimeGenerationType::Dynamic in the editor world
 	 */
-	ERuntimeGenerationType GetRuntimeGenerationType() const;
+	NAVIGATIONSYSTEM_API ERuntimeGenerationType GetRuntimeGenerationType() const;
 
-	void LogNavDataRegistrationResult(ERegistrationResult);
+	NAVIGATIONSYSTEM_API void LogNavDataRegistrationResult(ERegistrationResult);
 
 	/** Whether Navigation System is allowed to rebuild the navmesh
 	 *  Depends on runtime generation settings of each navigation data, always true in the editor
 	 */
-	bool IsAllowedToRebuild() const;
+	NAVIGATIONSYSTEM_API bool IsAllowedToRebuild() const;
 
 	/** Handle forwarding the information where needed when the setting is changed*/
-	void OnGenerateNavigationOnlyAroundNavigationInvokersChanged();
+	NAVIGATIONSYSTEM_API void OnGenerateNavigationOnlyAroundNavigationInvokersChanged();
 	
 	//----------------------------------------------------------------------//
 	// new stuff
 	//----------------------------------------------------------------------//
 public:
-	void VerifyNavigationRenderingComponents(const bool bShow);
+	NAVIGATIONSYSTEM_API void VerifyNavigationRenderingComponents(const bool bShow);
 	/** @param if InLevel is given then only navigation bounds from that level will be considered*/
-	virtual int GetNavigationBoundsForNavData(const ANavigationData& NavData, TArray<FBox>& OutBounds, ULevel* InLevel = nullptr) const;
-	static INavigationDataInterface* GetNavDataForActor(const AActor& Actor);
-	virtual void Configure(const UNavigationSystemConfig& Config) override;
-	virtual void AppendConfig(const UNavigationSystemConfig& NewConfig) override;
+	NAVIGATIONSYSTEM_API virtual int GetNavigationBoundsForNavData(const ANavigationData& NavData, TArray<FBox>& OutBounds, ULevel* InLevel = nullptr) const;
+	static NAVIGATIONSYSTEM_API INavigationDataInterface* GetNavDataForActor(const AActor& Actor);
+	NAVIGATIONSYSTEM_API virtual void Configure(const UNavigationSystemConfig& Config) override;
+	NAVIGATIONSYSTEM_API virtual void AppendConfig(const UNavigationSystemConfig& NewConfig) override;
 	
 #if !UE_BUILD_SHIPPING
-	void GetOnScreenMessages(TMultiMap<FCoreDelegates::EOnScreenMessageSeverity, FText>& OutMessages);
+	NAVIGATIONSYSTEM_API void GetOnScreenMessages(TMultiMap<FCoreDelegates::EOnScreenMessageSeverity, FText>& OutMessages);
 #endif // !UE_BUILD_SHIPPING
 
 	//----------------------------------------------------------------------//
@@ -1377,14 +1377,14 @@ public:
 	// Note2: originally deprecated as 4.22, bumped up to 5.1 as per comment above.
 	UE_DEPRECATED(5.1, "This version is deprecated.  Please use GetRandomLocationInNavigableRadius instead")
 	UFUNCTION(BlueprintPure, Category = "AI|Navigation", meta = (WorldContext = "WorldContextObject", DisplayName = "GetRandomPointInNavigableRadius", ScriptName = "GetRandomPointInNavigableRadius", DeprecatedFunction, DeprecationMessage = "GetRandomPointInNavigableRadius is deprecated. Use GetRandomLocationInNavigableRadius instead"))
-	static bool K2_GetRandomPointInNavigableRadius(UObject* WorldContextObject, const FVector& Origin, FVector& RandomLocation, float Radius, ANavigationData* NavData = NULL, TSubclassOf<UNavigationQueryFilter> FilterClass = NULL);
+	static NAVIGATIONSYSTEM_API bool K2_GetRandomPointInNavigableRadius(UObject* WorldContextObject, const FVector& Origin, FVector& RandomLocation, float Radius, ANavigationData* NavData = NULL, TSubclassOf<UNavigationQueryFilter> FilterClass = NULL);
 };
 
 //----------------------------------------------------------------------//
 // UNavigationSystemModuleConfig 
 //----------------------------------------------------------------------//
-UCLASS()
-class NAVIGATIONSYSTEM_API UNavigationSystemModuleConfig : public UNavigationSystemConfig
+UCLASS(MinimalAPI)
+class UNavigationSystemModuleConfig : public UNavigationSystemConfig
 {
 	GENERATED_BODY()
 
@@ -1403,19 +1403,19 @@ protected:
 	uint32 bSpawnNavDataInNavBoundsLevel : 1;
 
 public:
-	UNavigationSystemModuleConfig(const FObjectInitializer& ObjectInitializer = FObjectInitializer::Get());
+	NAVIGATIONSYSTEM_API UNavigationSystemModuleConfig(const FObjectInitializer& ObjectInitializer = FObjectInitializer::Get());
 
-	virtual void PostInitProperties() override;
+	NAVIGATIONSYSTEM_API virtual void PostInitProperties() override;
 
-	virtual UNavigationSystemBase* CreateAndConfigureNavigationSystem(UWorld& World) const override;
+	NAVIGATIONSYSTEM_API virtual UNavigationSystemBase* CreateAndConfigureNavigationSystem(UWorld& World) const override;
 
 #if WITH_EDITOR
-	virtual void PostEditChangeProperty(FPropertyChangedEvent& PropertyChangedEvent) override;
+	NAVIGATIONSYSTEM_API virtual void PostEditChangeProperty(FPropertyChangedEvent& PropertyChangedEvent) override;
 #endif // WITH_EDITOR
 
 protected:
 #if WITH_EDITOR
 	friend UNavigationSystemV1;
 #endif // WITH_EDITOR
-	void UpdateWithNavSysCDO(const UNavigationSystemV1& NavSysCDO);
+	NAVIGATIONSYSTEM_API void UpdateWithNavSysCDO(const UNavigationSystemV1& NavSysCDO);
 };
diff --git a/Engine/Source/Runtime/NavigationSystem/Public/NavigationSystemTypes.h b/Engine/Source/Runtime/NavigationSystem/Public/NavigationSystemTypes.h
--- a/Engine/Source/Runtime/NavigationSystem/Public/NavigationSystemTypes.h
+++ b/Engine/Source/Runtime/NavigationSystem/Public/NavigationSystemTypes.h
@@ -21,7 +21,7 @@ class FNavigationOctree;
 class UNavigationPath;
 class ANavigationData;
 
-struct NAVIGATIONSYSTEM_API FPathFindingQueryData
+struct FPathFindingQueryData
 {
 	TWeakObjectPtr<const UObject> Owner;
 	FVector StartLocation;
@@ -46,17 +46,17 @@ struct NAVIGATIONSYSTEM_API FPathFindingQueryData
 		Owner(InOwner), StartLocation(InStartLocation), EndLocation(InEndLocation), QueryFilter(InQueryFilter), CostLimit(InCostLimit), NavDataFlags(InNavDataFlags), bAllowPartialPaths(bInAllowPartialPaths), bRequireNavigableEndLocation(bInRequireNavigableEndLocation) {}
 };
 
-struct NAVIGATIONSYSTEM_API FPathFindingQuery : public FPathFindingQueryData
+struct FPathFindingQuery : public FPathFindingQueryData
 {
 	TWeakObjectPtr<const ANavigationData> NavData;
 	FNavPathSharedPtr PathInstanceToFill;
 	FNavAgentProperties NavAgentProperties;
 
 	FPathFindingQuery() : FPathFindingQueryData() {}
-	FPathFindingQuery(const UObject* InOwner, const ANavigationData& InNavData, const FVector& Start, const FVector& End, FSharedConstNavQueryFilter SourceQueryFilter = NULL, FNavPathSharedPtr InPathInstanceToFill = NULL, const FVector::FReal CostLimit = TNumericLimits<FVector::FReal>::Max(), const bool bInRequireNavigableEndLocation = true);
-	FPathFindingQuery(const INavAgentInterface& InNavAgent, const ANavigationData& InNavData, const FVector& Start, const FVector& End, FSharedConstNavQueryFilter SourceQueryFilter = NULL, FNavPathSharedPtr InPathInstanceToFill = NULL, const FVector::FReal CostLimit = TNumericLimits<FVector::FReal>::Max(), const bool bInRequireNavigableEndLocation = true);
+	NAVIGATIONSYSTEM_API FPathFindingQuery(const UObject* InOwner, const ANavigationData& InNavData, const FVector& Start, const FVector& End, FSharedConstNavQueryFilter SourceQueryFilter = NULL, FNavPathSharedPtr InPathInstanceToFill = NULL, const FVector::FReal CostLimit = TNumericLimits<FVector::FReal>::Max(), const bool bInRequireNavigableEndLocation = true);
+	NAVIGATIONSYSTEM_API FPathFindingQuery(const INavAgentInterface& InNavAgent, const ANavigationData& InNavData, const FVector& Start, const FVector& End, FSharedConstNavQueryFilter SourceQueryFilter = NULL, FNavPathSharedPtr InPathInstanceToFill = NULL, const FVector::FReal CostLimit = TNumericLimits<FVector::FReal>::Max(), const bool bInRequireNavigableEndLocation = true);
 
-	explicit FPathFindingQuery(FNavPathSharedRef PathToRecalculate, const ANavigationData* NavDataOverride = NULL);
+	NAVIGATIONSYSTEM_API explicit FPathFindingQuery(FNavPathSharedRef PathToRecalculate, const ANavigationData* NavDataOverride = NULL);
 
 	FPathFindingQuery& SetPathInstanceToUpdate(FNavPathSharedPtr InPathInstanceToFill) { PathInstanceToFill = InPathInstanceToFill; return *this; }
 	FPathFindingQuery& SetAllowPartialPaths(const bool bAllow) { bAllowPartialPaths = bAllow; return *this; }
@@ -66,7 +66,7 @@ struct NAVIGATIONSYSTEM_API FPathFindingQuery : public FPathFindingQueryData
 	/** utility function to compute a cost limit using an Euclidean heuristic, an heuristic scale and a cost limit factor
 	*	CostLimitFactor: multiplier used to compute the cost limit value from the initial heuristic
 	*	MinimumCostLimit: minimum clamping value used to prevent low cost limit for short path query */
-	static FVector::FReal ComputeCostLimitFromHeuristic(const FVector& StartPos, const FVector& EndPos, const FVector::FReal HeuristicScale, const FVector::FReal CostLimitFactor, const FVector::FReal MinimumCostLimit);
+	static NAVIGATIONSYSTEM_API FVector::FReal ComputeCostLimitFromHeuristic(const FVector& StartPos, const FVector& EndPos, const FVector::FReal HeuristicScale, const FVector::FReal CostLimitFactor, const FVector::FReal MinimumCostLimit);
 };
 
 namespace EPathFindingMode
@@ -82,7 +82,7 @@ namespace EPathFindingMode
 //// Custom path following data
 //
 ///** Custom data passed to movement requests. */
-struct NAVIGATIONSYSTEM_API FMoveRequestCustomData
+struct FMoveRequestCustomData
 {
 };
 
@@ -151,14 +151,14 @@ namespace FNavigationSystem
 
 namespace NavigationHelper
 {
-	struct NAVIGATIONSYSTEM_API FNavLinkOwnerData
+	struct FNavLinkOwnerData
 	{
 		const AActor* Actor;
 		FTransform LinkToWorld;
 
 		FNavLinkOwnerData() : Actor(nullptr) {}
-		FNavLinkOwnerData(const AActor& InActor);
-		FNavLinkOwnerData(const USceneComponent& InComponent);
+		NAVIGATIONSYSTEM_API FNavLinkOwnerData(const AActor& InActor);
+		NAVIGATIONSYSTEM_API FNavLinkOwnerData(const USceneComponent& InComponent);
 	};
 
 	DECLARE_DELEGATE_ThreeParams(FNavLinkProcessorDelegate, FCompositeNavModifier*, const AActor*, const TArray<FNavigationLink>&);
@@ -186,4 +186,4 @@ namespace NavigationHelper
 	NAVIGATIONSYSTEM_API void DefaultNavLinkSegmentProcessorImpl(FCompositeNavModifier* OUT CompositeModifier, const FNavLinkOwnerData& OwnerData, const TArray<FNavigationSegmentLink>& IN NavLinks);
 
 	NAVIGATIONSYSTEM_API bool IsBodyNavigationRelevant(const UBodySetup& IN BodySetup);
-}
\ No newline at end of file
+}
diff --git a/Engine/Source/Runtime/NavigationSystem/Public/NavigationTestingActor.h b/Engine/Source/Runtime/NavigationSystem/Public/NavigationTestingActor.h
--- a/Engine/Source/Runtime/NavigationSystem/Public/NavigationTestingActor.h
+++ b/Engine/Source/Runtime/NavigationSystem/Public/NavigationTestingActor.h
@@ -40,8 +40,8 @@ namespace ENavCostDisplay
 	};
 }
 
-UCLASS(hidecategories=(Object, Actor, Input, Rendering, Replication, HLOD, Cooking), showcategories=("Input|MouseInput", "Input|TouchInput"), Blueprintable)
-class NAVIGATIONSYSTEM_API ANavigationTestingActor : public AActor, public INavAgentInterface, public INavPathObserverInterface
+UCLASS(hidecategories=(Object, Actor, Input, Rendering, Replication, HLOD, Cooking), showcategories=("Input|MouseInput", "Input|TouchInput"), Blueprintable, MinimalAPI)
+class ANavigationTestingActor : public AActor, public INavAgentInterface, public INavPathObserverInterface
 {
 	GENERATED_UCLASS_BODY()
 
@@ -182,22 +182,22 @@ public:
 	FNavigationPath::FPathObserverDelegate::FDelegate PathObserver;
 
 	/** Dtor */
-	virtual ~ANavigationTestingActor();
+	NAVIGATIONSYSTEM_API virtual ~ANavigationTestingActor();
 
-	virtual void BeginDestroy() override;
+	NAVIGATIONSYSTEM_API virtual void BeginDestroy() override;
 
 #if WITH_EDITOR
-	virtual void PreEditChange(FProperty* PropertyThatWillChange) override;
-	virtual void PostEditChangeProperty(FPropertyChangedEvent& PropertyChangedEvent) override;
-	virtual void PostEditMove(bool bFinished) override;
+	NAVIGATIONSYSTEM_API virtual void PreEditChange(FProperty* PropertyThatWillChange) override;
+	NAVIGATIONSYSTEM_API virtual void PostEditChangeProperty(FPropertyChangedEvent& PropertyChangedEvent) override;
+	NAVIGATIONSYSTEM_API virtual void PostEditMove(bool bFinished) override;
 	
-	virtual void PostLoad() override;
-	void TickMe();
+	NAVIGATIONSYSTEM_API virtual void PostLoad() override;
+	NAVIGATIONSYSTEM_API void TickMe();
 #endif // WITH_EDITOR
 
 	//~ Begin INavAgentInterface Interface
 	virtual const FNavAgentProperties& GetNavAgentPropertiesRef() const override { return NavAgentProps; }
-	virtual FVector GetNavAgentLocation() const override;
+	NAVIGATIONSYSTEM_API virtual FVector GetNavAgentLocation() const override;
 	virtual void GetMoveGoalReachTest(const AActor* MovingActor, const FVector& MoveOffset, FVector& GoalOffset, float& GoalRadius, float& GoalHalfHeight) const override {}
 	//~ End INavAgentInterface Interface
 
@@ -207,20 +207,20 @@ public:
 	virtual void OnPathFailed(class INavigationPathGenerator* PathGenerator) override {};
 	//~ End INavPathObserverInterface Interface	
 
-	void UpdateNavData();
-	void UpdatePathfinding();
-	void GatherDetailedData(ANavigationTestingActor* Goal);
-	virtual void SearchPathTo(ANavigationTestingActor* Goal);
+	NAVIGATIONSYSTEM_API void UpdateNavData();
+	NAVIGATIONSYSTEM_API void UpdatePathfinding();
+	NAVIGATIONSYSTEM_API void GatherDetailedData(ANavigationTestingActor* Goal);
+	NAVIGATIONSYSTEM_API virtual void SearchPathTo(ANavigationTestingActor* Goal);
 
 	/*	Called when given path becomes invalid (via @see PathObserverDelegate)
 	 *	NOTE: InvalidatedPath doesn't have to be instance's current Path
 	 */
-	void OnPathEvent(FNavigationPath* InvalidatedPath, ENavPathEvent::Type Event);
+	NAVIGATIONSYSTEM_API void OnPathEvent(FNavigationPath* InvalidatedPath, ENavPathEvent::Type Event);
 
 	// Virtual method to override if you want to customize the query being 
 	// constructed for the path find (e.g. change the filter or add 
 	// constraints/goal evaluators).
-	virtual FPathFindingQuery BuildPathFindingQuery(const ANavigationTestingActor* Goal) const;
+	NAVIGATIONSYSTEM_API virtual FPathFindingQuery BuildPathFindingQuery(const ANavigationTestingActor* Goal) const;
 
 	/** Returns CapsuleComponent subobject **/
 	class UCapsuleComponent* GetCapsuleComponent() const { return CapsuleComponent; }
@@ -230,5 +230,5 @@ public:
 #endif
 
 protected:
-	FVector FindClosestWallLocation() const;
+	NAVIGATIONSYSTEM_API FVector FindClosestWallLocation() const;
 };
