[Navmesh] Fix voxel filtering issue when using high AgentMaxSlope values in combination with high CellSize values.

diff --git a/Engine/Source/Runtime/NavigationSystem/Private/NavMesh/RecastNavMeshGenerator.cpp b/Engine/Source/Runtime/NavigationSystem/Private/NavMesh/RecastNavMeshGenerator.cpp
--- a/Engine/Source/Runtime/NavigationSystem/Private/NavMesh/RecastNavMeshGenerator.cpp
+++ b/Engine/Source/Runtime/NavigationSystem/Private/NavMesh/RecastNavMeshGenerator.cpp
@@ -251,6 +251,16 @@ struct FRecastGeometryExport : public FNavigableGeometryExport
 	virtual void SetNavDataPerInstanceTransformDelegate(const FNavDataPerInstanceTransformDelegate& InDelegate) override;
 };
 
+bool FRecastBuildConfig::IsUsingCoarseCellSize() const
+{
+	const rcReal MaxStepFromUniformSlope = cs * FMath::Tan(FMath::DegreesToRadians(walkableSlopeAngle));
+
+	// When filtering occurs, in rcFilterLedgeSpansImp, there can be a 2 steps vertical distance between neighbors ((asmax - asmin) > walkableClimb).
+	// This is why we compare with 2 times the max step.
+	const bool bIsUsingCoarseCellSize = (walkableClimb * ch) <= (2 * MaxStepFromUniformSlope);
+	return bIsUsingCoarseCellSize;
+}
+
 FRecastVoxelCache::FRecastVoxelCache(const uint8* Memory)
 {
 	uint8* BytesArr = (uint8*)Memory;
@@ -2963,7 +2973,7 @@ void FRecastTileGenerator::GenerateRecastFilter(FNavMeshBuildContext& BuildConte
 	{
 		SCOPE_CYCLE_COUNTER(STAT_Navigation_FilterLedgeSpans)
 
-		const bool bFilterNeighborSlope = TileConfig.walkableRadius > 1 || !UE::NavMesh::Private::bKeepSteepSlopeForSingleVoxelAgent;
+		const bool bFilterNeighborSlope = (TileConfig.walkableRadius > 1 || !UE::NavMesh::Private::bKeepSteepSlopeForSingleVoxelAgent) && !TileConfig.IsUsingCoarseCellSize();
 		rcFilterLedgeSpans(&BuildContext, TileConfig.walkableHeight, TileConfig.walkableClimb, bFilterNeighborSlope, *RasterContext.SolidHF);
 	}
 	if (!TileConfig.bMarkLowHeightAreas)
@@ -3010,7 +3020,7 @@ ETimeSliceWorkResult FRecastTileGenerator::GenerateRecastFilterTimeSliced(FNavMe
 
 		bool DoIter = true;
 
-		const bool bFilterNeighborSlope = TileConfig.walkableRadius > 1 || !UE::NavMesh::Private::bKeepSteepSlopeForSingleVoxelAgent;
+		const bool bFilterNeighborSlope = (TileConfig.walkableRadius > 1 || !UE::NavMesh::Private::bKeepSteepSlopeForSingleVoxelAgent) && !TileConfig.IsUsingCoarseCellSize();
 
 		do
 		{
diff --git a/Engine/Source/Runtime/NavigationSystem/Public/NavMesh/RecastNavMeshGenerator.h b/Engine/Source/Runtime/NavigationSystem/Public/NavMesh/RecastNavMeshGenerator.h
--- a/Engine/Source/Runtime/NavigationSystem/Public/NavMesh/RecastNavMeshGenerator.h
+++ b/Engine/Source/Runtime/NavigationSystem/Public/NavMesh/RecastNavMeshGenerator.h
@@ -104,6 +104,9 @@ struct FRecastBuildConfig : public rcConfig
 	}
 
 	rcReal GetTileSizeUU() const { return tileSize * cs; }
+
+	/** Detects if we are using big cell size values in relation to the walkableClimb and walkableSlopeAngle. */
+	bool IsUsingCoarseCellSize() const;
 };
 
 PRAGMA_ENABLE_DEPRECATION_WARNINGS
