Fixed unitialized FBox in QABoxRenderComponent

diff --git a/Engine/Source/Runtime/NavigationSystem/Private/NavMesh/RecastNavMesh.cpp b/Engine/Source/Runtime/NavigationSystem/Private/NavMesh/RecastNavMesh.cpp
--- a/Engine/Source/Runtime/NavigationSystem/Private/NavMesh/RecastNavMesh.cpp
+++ b/Engine/Source/Runtime/NavigationSystem/Private/NavMesh/RecastNavMesh.cpp
@@ -1183,7 +1183,7 @@ FBox ARecastNavMesh::GetNavMeshBounds() const
 
 FBox ARecastNavMesh::GetNavMeshTileBounds(int32 TileIndex) const
 {
-	FBox Bounds;
+	FBox Bounds(ForceInit);
 	if (RecastNavMeshImpl)
 	{
 		Bounds = RecastNavMeshImpl->GetNavMeshTileBounds(TileIndex);
