[Navmesh] Cancel async building tasks when navmesh generator is destroyed to prevent leaking any build tasks.

diff --git a/Engine/Source/Runtime/NavigationSystem/Private/NavMesh/RecastNavMeshGenerator.cpp b/Engine/Source/Runtime/NavigationSystem/Private/NavMesh/RecastNavMeshGenerator.cpp
--- a/Engine/Source/Runtime/NavigationSystem/Private/NavMesh/RecastNavMeshGenerator.cpp
+++ b/Engine/Source/Runtime/NavigationSystem/Private/NavMesh/RecastNavMeshGenerator.cpp
@@ -4627,6 +4627,8 @@ PRAGMA_ENABLE_DEPRECATION_WARNINGS
 
 FRecastNavMeshGenerator::~FRecastNavMeshGenerator()
 {
+	UE_CLOG(RunningDirtyTiles.Num() > 0, LogNavigation, Log, TEXT("Discarding %d build tasks"), RunningDirtyTiles.Num());
+	CancelBuild();
 	DEC_DWORD_STAT_BY( STAT_NavigationMemory, sizeof(*this) );
 }
 
