[Navmesh] Validating AgentMaxStepHeight in relation to AgentMaxSlope and CellSize

diff --git a/Engine/Source/Runtime/NavigationSystem/Private/NavMesh/RecastNavMeshGenerator.cpp b/Engine/Source/Runtime/NavigationSystem/Private/NavMesh/RecastNavMeshGenerator.cpp
--- a/Engine/Source/Runtime/NavigationSystem/Private/NavMesh/RecastNavMeshGenerator.cpp
+++ b/Engine/Source/Runtime/NavigationSystem/Private/NavMesh/RecastNavMeshGenerator.cpp
@@ -4686,6 +4686,31 @@ void FRecastNavMeshGenerator::ConfigureBuildProperties(FRecastBuildConfig& OutCo
 	OutConfig.walkableClimb = FMath::CeilToInt(AgentMaxClimb / CellHeight);
 	OutConfig.walkableRadius = FMath::CeilToInt(AgentRadius / CellSize);
 
+	// For each navmesh resolutions, validate that AgentMaxStepHeight is high enough for the AgentMaxSlope angle
+	for (int32 Index = 0; Index < (uint8)ENavigationDataResolution::MAX; Index++)
+	{
+		const ENavigationDataResolution Resolution = (ENavigationDataResolution)Index;
+		
+		const float MaxStepHeight = DestNavMesh->GetAgentMaxStepHeight(Resolution);
+		const float TempCellHeight = DestNavMesh->GetCellHeight(Resolution);
+		const int WalkableClimbVx = FMath::CeilToInt(MaxStepHeight / TempCellHeight);
+
+		// Compute the required climb to prevent direct neighbor filtering in rcFilterLedgeSpansImp (minh < -walkableClimb).
+		// See comment: "The current span is close to a ledge if the drop to any neighbour span is less than the walkableClimb."
+		const float RequiredClimb = DestNavMesh->GetCellSize(Resolution) * FMath::Tan(FMath::DegreesToRadians(AgentMaxSlope));
+		const int RequiredClimbVx = FMath::CeilToInt(RequiredClimb / TempCellHeight);
+		
+		if (WalkableClimbVx < RequiredClimbVx)
+		{
+			// This is a log since we need to let the user decide which one of the parameters needs to be changed (if any).
+			UE_LOG(LogNavigationDataBuild, Log, TEXT("%s: AgentMaxStepHeight (%f) for resolution %i is not high enough in steep slopes (AgentMaxSlope is %f). "
+				"Use AgentMaxStepHeight bigger than %f or a smaller AgentMaxSlope to avoid undesirable navmesh holes in steep slopes. "
+				"This can also be avoided by using smaller CellSize and CellHeight."),
+				*GetNameSafe(DestNavMesh), MaxStepHeight,
+				*UEnum::GetDisplayValueAsText(Resolution).ToString(), AgentMaxSlope, (RequiredClimbVx-1)*TempCellHeight);	
+		}
+	}
+	
 	// store original sizes
 	OutConfig.AgentHeight = AgentHeight;
 	OutConfig.AgentMaxClimb = AgentMaxClimb;
