Fix LogNavigationHistory log

diff --git a/Engine/Source/Runtime/NavigationSystem/Private/NavigationSystem.cpp b/Engine/Source/Runtime/NavigationSystem/Private/NavigationSystem.cpp
--- a/Engine/Source/Runtime/NavigationSystem/Private/NavigationSystem.cpp
+++ b/Engine/Source/Runtime/NavigationSystem/Private/NavigationSystem.cpp
@@ -515,7 +515,7 @@ void FNavRegenTimeSliceManager::LogTileStatistics(const TArray<TObjectPtr<ANavig
 					const double MedianWaitTimeMs = HistoryData[MedianIndex].TileWaitTime * 1000.f;
 					const double HighWaitTimeMs = HistoryData[HighIndex].TileWaitTime * 1000.f;
 					
-					UE_LOG(LogNavigation, Warning, TEXT("%-30s Median tile stats: regen time: %2.2f ms, regen frames %lld, wait time: %4.f ms (high regen time: %2.2f ms, high wait time: %4.f ms."),
+					UE_LOG(LogNavigationHistory, Log, TEXT("%-35s Median tile stats: regen time: %2.2f ms, regen frames %lld, wait time: %4.f ms (high regen time: %2.2f ms, high wait time: %4.f ms."),
 						*GetNameSafe(NavDataSet[NavDataIndex]), MedianRegenTimeMs, MedianRegenFrames, MedianWaitTimeMs, HighRegenTimeMs, HighWaitTimeMs);
 				}
 			}
