Addition of InvokersMaximumDistanceFromSeed allowing to restrict invokers updating far from seed points (player positions by default).

diff --git a/Engine/Source/Runtime/NavigationSystem/Private/NavigationSystem.cpp b/Engine/Source/Runtime/NavigationSystem/Private/NavigationSystem.cpp
--- a/Engine/Source/Runtime/NavigationSystem/Private/NavigationSystem.cpp
+++ b/Engine/Source/Runtime/NavigationSystem/Private/NavigationSystem.cpp
@@ -25,6 +25,7 @@
 #include "UObject/Package.h"
 #include "Components/PrimitiveComponent.h"
 #include "UObject/UObjectThreadContext.h"
+#include "GameFramework/Pawn.h"
 
 #if WITH_RECAST
 #include "NavMesh/RecastNavMesh.h"
@@ -515,8 +516,8 @@ void FNavRegenTimeSliceManager::LogTileStatistics(const TArray<TObjectPtr<ANavig
 					const double MedianWaitTimeMs = HistoryData[MedianIndex].TileWaitTime * 1000.f;
 					const double HighWaitTimeMs = HistoryData[HighIndex].TileWaitTime * 1000.f;
 					
-					UE_LOG(LogNavigationHistory, Log, TEXT("%-35s Median tile stats: regen time: %2.2f ms, regen frames %lld, wait time: %4.f ms (high regen time: %2.2f ms, high wait time: %4.f ms."),
-						*GetNameSafe(NavDataSet[NavDataIndex]), MedianRegenTimeMs, MedianRegenFrames, MedianWaitTimeMs, HighRegenTimeMs, HighWaitTimeMs);
+					UE_LOG(LogNavigationHistory, Log, TEXT("%-35s Median tile stats: regen time: %2.2f ms, regen frames %lld, wait time: %4.f ms (high regen time: %2.2f ms, high wait time: %4.f ms) regen count: %i"),
+						*GetNameSafe(NavDataSet[NavDataIndex]), MedianRegenTimeMs, MedianRegenFrames, MedianWaitTimeMs, HighRegenTimeMs, HighWaitTimeMs, HistoryData.Num());
 				}
 			}
 		}
@@ -3713,6 +3714,19 @@ void UNavigationSystemV1::GatherNavigationBounds()
 	}
 }
 
+void UNavigationSystemV1::GetInvokerSeedLocations(const UWorld& InWorld, TArray<FVector2D, TInlineAllocator<32>>& OutSeedLocations)
+{
+	for (FConstPlayerControllerIterator PlayerIt = InWorld.GetPlayerControllerIterator(); PlayerIt; ++PlayerIt)
+	{
+		const APlayerController* PlayerController = PlayerIt->Get();
+		if (PlayerController && PlayerController->GetPawn())
+		{
+			const FVector2D SeedLoc(PlayerController->GetPawn()->GetActorLocation());
+			OutSeedLocations.Add(SeedLoc);
+		}
+	}
+}
+
 void UNavigationSystemV1::Build()
 {
 	TRACE_CPUPROFILER_EVENT_SCOPE(UNavigationSystemV1::Build);
@@ -4980,7 +4994,7 @@ void UNavigationSystemV1::UnregisterInvoker(AActor& Invoker)
 
 void UNavigationSystemV1::UpdateInvokers()
 {
-	UWorld* World = GetWorld();
+	const UWorld* World = GetWorld();
 	const double CurrentTime = World->GetTimeSeconds();
 	if (CurrentTime >= NextInvokersUpdateTime)
 	{
@@ -4990,13 +5004,20 @@ void UNavigationSystemV1::UpdateInvokers()
 		{
 			QUICK_SCOPE_CYCLE_COUNTER(STAT_NavSys_Clusterize);
 
+			const bool bCheckMaximumDistanceFromSeeds = InvokersMaximumDistanceFromSeed != -1;
+			TArray<FVector2D, TInlineAllocator<32>> SeedLocations;
+			if (bCheckMaximumDistanceFromSeeds)
+			{
+				GetInvokerSeedLocations(*World, SeedLocations);
+			}
+			
 			const double StartTime = FPlatformTime::Seconds();
 
 			InvokerLocations.Reserve(Invokers.Num());
 
 			for (auto ItemIterator = Invokers.CreateIterator(); ItemIterator; ++ItemIterator)
 			{
-				AActor* Actor = ItemIterator->Value.Actor.Get();
+				const AActor* Actor = ItemIterator->Value.Actor.Get();
 				if (Actor != nullptr
 #if WITH_EDITOR
 					// Would like to ignore objects in transactional buffer here, but there's no flag for it
@@ -5004,8 +5025,41 @@ void UNavigationSystemV1::UpdateInvokers()
 #endif //WITH_EDITOR
 					)
 				{
-					InvokerLocations.Add(FNavigationInvokerRaw(Actor->GetActorLocation(), ItemIterator->Value.GenerationRadius, ItemIterator->Value.RemovalRadius,
-						ItemIterator->Value.SupportedAgents, ItemIterator->Value.Priority));
+					const FVector ActorLocation = Actor->GetActorLocation();
+					const float GenerationRadius = ItemIterator->Value.GenerationRadius;
+					bool bKeep = !bCheckMaximumDistanceFromSeeds;
+
+					double ClosestDistanceSq = DBL_MAX;
+					if (bCheckMaximumDistanceFromSeeds)
+					{
+						const double CheckDistanceSq = FMath::Square(InvokersMaximumDistanceFromSeed + GenerationRadius);
+
+						// Check if the invoker is close enough
+						for (const FVector2D SeedLocation : SeedLocations)
+						{
+							const double InvokerDistanceToSeedSq = FVector2D::DistSquared(SeedLocation, FVector2D(ActorLocation));
+							if (InvokerDistanceToSeedSq <= CheckDistanceSq)
+							{
+								bKeep = true;
+								break;
+							}
+							else
+							{
+								ClosestDistanceSq = FMath::Min(InvokerDistanceToSeedSq, ClosestDistanceSq);
+							}
+						}	
+					}
+
+					if (bKeep)
+					{
+						InvokerLocations.Add(FNavigationInvokerRaw(ActorLocation, GenerationRadius, ItemIterator->Value.RemovalRadius,
+							ItemIterator->Value.SupportedAgents, ItemIterator->Value.Priority));
+					}
+					else
+					{
+						UE_LOG(LogNavInvokers, Verbose, TEXT("Invoker %s ignored because it's too far from any seed location. Closest seed at %.0f."),
+							*Actor->GetName(), FMath::Sqrt(ClosestDistanceSq));
+					}
 				}
 				else
 				{
@@ -5053,9 +5107,9 @@ void UNavigationSystemV1::UpdateInvokers()
 			{
 				const int32 NavDataSupportedAgentIndex = GetSupportedAgentIndex(NavData);	
 
-				for (auto ItemIterator = Invokers.CreateIterator(); ItemIterator; ++ItemIterator)
+				for (auto ItemIterator = InvokerLocations.CreateIterator(); ItemIterator; ++ItemIterator)
 				{
-					const FNavAgentSelector& InvokerSupportedAgents = ItemIterator->Value.SupportedAgents;
+					const FNavAgentSelector& InvokerSupportedAgents = ItemIterator->SupportedAgents;
 					if (InvokerSupportedAgents.Contains(NavDataSupportedAgentIndex))
 					{
 						InvokerCounts[NavDataIndex]++;
@@ -5065,6 +5119,8 @@ void UNavigationSystemV1::UpdateInvokers()
 				const FString StatName = FString::Printf(TEXT("InvokerCount_%s"), *NavData->GetName()); 
 				FCsvProfiler::RecordCustomStat(*StatName, CSV_CATEGORY_INDEX(NavInvokers), InvokerCounts[NavDataIndex], ECsvCustomStatOp::Set);
 			}
+
+			FCsvProfiler::RecordCustomStat(TEXT("InvokersFarAway"), CSV_CATEGORY_INDEX(NavInvokers), Invokers.Num() - InvokerLocations.Num(), ECsvCustomStatOp::Set);
 		}		
 	}
 #endif // CSV_PROFILER
diff --git a/Engine/Source/Runtime/NavigationSystem/Public/NavigationSystem.h b/Engine/Source/Runtime/NavigationSystem/Public/NavigationSystem.h
--- a/Engine/Source/Runtime/NavigationSystem/Public/NavigationSystem.h
+++ b/Engine/Source/Runtime/NavigationSystem/Public/NavigationSystem.h
@@ -22,6 +22,7 @@
 #include "NavigationDirtyAreasController.h"
 #include "Math/MovingWindowAverageFast.h"
 #include "AI/Navigation/NavigationBounds.h"
+#include "Containers/ContainerAllocationPolicies.h"
 #if UE_ENABLE_INCLUDE_ORDER_DEPRECATED_IN_5_2
 #if WITH_EDITOR
 #include "UnrealEdMisc.h"
@@ -389,6 +390,10 @@ protected:
 	UPROPERTY(EditAnywhere, Category = "Navigation Enforcing", meta = (ClampMin = "0.1", UIMin = "0.1", EditCondition = "bGenerateNavigationOnlyAroundNavigationInvokers"), config)
 	float ActiveTilesUpdateInterval;
 
+	/** When in use, invokers farther away from any invoker seed will be ignored (set to -1 to disable). */
+	UPROPERTY(EditAnywhere, Category = "Navigation Enforcing", meta = (EditCondition = "bGenerateNavigationOnlyAroundNavigationInvokers"), config)
+	double InvokersMaximumDistanceFromSeed = -1;
+	
 	/** Sets how navigation data should be gathered when building collision information */
 	UPROPERTY(EditDefaultsOnly, Category = "NavigationSystem", config)
 	ENavDataGatheringModeConfig DataGatheringMode;
@@ -860,6 +865,9 @@ protected:
 	/** Searches for all valid navigation bounds in the world and stores them */
 	NAVIGATIONSYSTEM_API virtual void GatherNavigationBounds();
 
+	/** Get seed locations for invokers, @see InvokersMaximumDistanceFromSeed */ 
+	NAVIGATIONSYSTEM_API virtual void GetInvokerSeedLocations(const UWorld& InWorld, TArray<FVector2D, TInlineAllocator<32>>& OutSeedLocations);
+
 	/** @return pointer to ANavigationData instance of given ID, or NULL if it was not found. Note it looks only through registered navigation data */
 	NAVIGATIONSYSTEM_API ANavigationData* GetNavDataWithID(const uint16 NavDataID) const;
 
