Make sure InvokersMaximumDistanceFromSeed is only active in game mode (ignore it in editor)

diff --git a/Engine/Source/Runtime/NavigationSystem/Private/NavigationSystem.cpp b/Engine/Source/Runtime/NavigationSystem/Private/NavigationSystem.cpp
--- a/Engine/Source/Runtime/NavigationSystem/Private/NavigationSystem.cpp
+++ b/Engine/Source/Runtime/NavigationSystem/Private/NavigationSystem.cpp
@@ -5004,7 +5004,7 @@ void UNavigationSystemV1::UpdateInvokers()
 		{
 			QUICK_SCOPE_CYCLE_COUNTER(STAT_NavSys_Clusterize);
 
-			const bool bCheckMaximumDistanceFromSeeds = InvokersMaximumDistanceFromSeed != -1;
+			const bool bCheckMaximumDistanceFromSeeds = (InvokersMaximumDistanceFromSeed != -1) && World->IsGameWorld();
 			TArray<FVector2D, TInlineAllocator<32>> SeedLocations;
 			if (bCheckMaximumDistanceFromSeeds)
 			{
