Fix no-Recast builds [Github Pull Request: 9579]

diff --git a/Engine/Source/Runtime/AIModule/Classes/Navigation/RecastGraphAStar.h b/Engine/Source/Runtime/AIModule/Classes/Navigation/RecastGraphAStar.h
--- a/Engine/Source/Runtime/AIModule/Classes/Navigation/RecastGraphAStar.h
+++ b/Engine/Source/Runtime/AIModule/Classes/Navigation/RecastGraphAStar.h
@@ -32,7 +32,8 @@ struct dtMeshTile;
 struct dtLink;
 class dtNavMesh;
 
-struct FRecastNeighbour
+#if WITH_RECAST
+struct AIMODULE_API FRecastNeighbour
 {
 	friend FRecastGraphWrapper;
 	friend FRecastGraphAStarFilter;
@@ -65,6 +66,7 @@ struct FRecastAStarResult : public dtQueryResult
 
 	AIMODULE_API dtPolyRef SetPathInfo(const int32 Index, const FRecastAStarSearchNode& SearchNode);
 };
+#endif // WITH_RECAST
 
 USTRUCT()
 struct FRecastGraphWrapper
@@ -74,6 +76,7 @@ struct FRecastGraphWrapper
 public:
 	FRecastGraphWrapper() {}
 
+#if WITH_RECAST
 	/** Initialization of the wrapper from the RecastNavMesh pointer */
 	AIMODULE_API void Initialize(const ARecastNavMesh* InRecastNavMeshActor);
 
@@ -101,18 +104,22 @@ protected:
 	FORCEINLINE const ARecastNavMesh* GetRecastNavMeshActor() const { checkSlow(RecastNavMeshActor);  return RecastNavMeshActor; }
 	FORCEINLINE const dtNavMesh* GetDetourNavMesh() const { checkSlow(DetourNavMesh); return DetourNavMesh; }
 	AIMODULE_API void BindFilter(FRecastGraphAStarFilter& AStarFilter);
+#endif // WITH_RECAST
 
 private:
 	UPROPERTY(Transient)
 	TObjectPtr<const ARecastNavMesh> RecastNavMeshActor = nullptr;
 
+#if WITH_RECAST
 	const dtNavMesh* DetourNavMesh = nullptr;
 
 	dtNavMeshQuery RecastQuery;
 
 	mutable unsigned int CachedNextLink = DT_NULL_LINK;
+#endif // WITH_RECAST
 };
 
+#if WITH_RECAST
 struct FRecastAStarSearchNode : public FGraphAStarDefaultNode<FRecastGraphWrapper>
 {
 	typedef FGraphAStarDefaultNode<FRecastGraphWrapper> Super;
@@ -245,3 +252,4 @@ private:
 	uint32 MaxSearchNodes;
 	FVector::FReal CostLimit;
 };
+#endif // WITH_RECAST
diff --git a/Engine/Source/Runtime/AIModule/Private/GameplayDebugger/GameplayDebuggerCategory_AI.cpp b/Engine/Source/Runtime/AIModule/Private/GameplayDebugger/GameplayDebuggerCategory_AI.cpp
--- a/Engine/Source/Runtime/AIModule/Private/GameplayDebugger/GameplayDebuggerCategory_AI.cpp
+++ b/Engine/Source/Runtime/AIModule/Private/GameplayDebugger/GameplayDebuggerCategory_AI.cpp
@@ -272,6 +272,7 @@ void FGameplayDebuggerCategory_AI::CollectPathData(AAIController* DebugAI)
 			{
 				LastPathUpdateTime = CurrentPath->GetLastUpdateTime();
 
+#if WITH_RECAST
 				const FNavMeshPath* NavMeshPath = CurrentPath->CastPath<FNavMeshPath>();
 				const ARecastNavMesh* NavData = Cast<const ARecastNavMesh>(CurrentPath->GetNavigationDataUsed());
 				if (NavMeshPath && NavData)
@@ -287,6 +288,7 @@ void FGameplayDebuggerCategory_AI::CollectPathData(AAIController* DebugAI)
 						PathDataPack.PathCorridor.Add(PolyData);
 					}
 				}
+#endif // WITH_RECAST
 
 				for (int32 Idx = 0; Idx < CurrentPath->GetPathPoints().Num(); Idx++)
 				{
diff --git a/Engine/Source/Runtime/AIModule/Private/Navigation/CrowdManager.cpp b/Engine/Source/Runtime/AIModule/Private/Navigation/CrowdManager.cpp
--- a/Engine/Source/Runtime/AIModule/Private/Navigation/CrowdManager.cpp
+++ b/Engine/Source/Runtime/AIModule/Private/Navigation/CrowdManager.cpp
@@ -221,7 +221,9 @@ void UCrowdManager::BeginDestroy()
 	delete TickHelper;
 #endif
 
+#if WITH_RECAST
 	DestroyCrowdManager();
+#endif // WITH_RECAST
 	Super::BeginDestroy();
 }
 
diff --git a/Engine/Source/Runtime/AIModule/Private/Navigation/RecastGraphAStar.cpp b/Engine/Source/Runtime/AIModule/Private/Navigation/RecastGraphAStar.cpp
--- a/Engine/Source/Runtime/AIModule/Private/Navigation/RecastGraphAStar.cpp
+++ b/Engine/Source/Runtime/AIModule/Private/Navigation/RecastGraphAStar.cpp
@@ -8,6 +8,8 @@
 
 #include UE_INLINE_GENERATED_CPP_BY_NAME(RecastGraphAStar)
 
+#if WITH_RECAST
+
 void FRecastGraphWrapper::Initialize(const ARecastNavMesh* InRecastNavMeshActor)
 {
 	RecastNavMeshActor = InRecastNavMeshActor;
@@ -91,3 +93,4 @@ FRecastGraphAStarFilter::FRecastGraphAStarFilter(FRecastGraphWrapper& InRecastGr
 	InRecastGraphWrapper.BindFilter(*this);
 }
 
+#endif // WITH_RECAST
diff --git a/Engine/Source/Runtime/NavigationSystem/Private/NavLinkRenderingComponent.cpp b/Engine/Source/Runtime/NavigationSystem/Private/NavLinkRenderingComponent.cpp
--- a/Engine/Source/Runtime/NavigationSystem/Private/NavLinkRenderingComponent.cpp
+++ b/Engine/Source/Runtime/NavigationSystem/Private/NavLinkRenderingComponent.cpp
@@ -181,11 +181,13 @@ void FNavLinkRenderingProxy::GetDynamicMeshElements(const TArray<const FSceneVie
 				if (NavMesh != NULL)
 				{
 					AgentMask = NavMesh->IsDrawingEnabled() ? AgentMask | (1 << DataIndex) : AgentMask;
+#if WITH_RECAST
 					const float AgentMaxStepHeight = NavMesh->GetAgentMaxStepHeight(ENavigationDataResolution::Default);
 					if (AgentMaxStepHeight > 0 && NavMesh->IsDrawingEnabled())
 					{
 						StepHeights.Add(AgentMaxStepHeight);
 					}
+#endif // WITH_RECAST
 				}
 			}
 		}
diff --git a/Engine/Source/Runtime/NavigationSystem/Private/NavMesh/NavMeshRenderingComponent.cpp b/Engine/Source/Runtime/NavigationSystem/Private/NavMesh/NavMeshRenderingComponent.cpp
--- a/Engine/Source/Runtime/NavigationSystem/Private/NavMesh/NavMeshRenderingComponent.cpp
+++ b/Engine/Source/Runtime/NavigationSystem/Private/NavMesh/NavMeshRenderingComponent.cpp
@@ -1809,9 +1809,11 @@ void UNavMeshRenderingComponent::OnUnregister()
 
 void UNavMeshRenderingComponent::GatherData(const ARecastNavMesh& NavMesh, FNavMeshSceneProxyData& OutProxyData) const
 {
+#if WITH_RECAST
 	const int32 DetailFlags = FNavMeshRenderingHelpers::GetDetailFlags(&NavMesh);
 	const TArray<int32> EmptyTileSet;
 	OutProxyData.GatherData(&NavMesh, DetailFlags, EmptyTileSet);
+#endif // WITH_RECAST
 }
 
 #if UE_ENABLE_DEBUG_DRAWING
diff --git a/Engine/Source/Runtime/NavigationSystem/Private/NavMesh/RecastNavMesh.cpp b/Engine/Source/Runtime/NavigationSystem/Private/NavMesh/RecastNavMesh.cpp
--- a/Engine/Source/Runtime/NavigationSystem/Private/NavMesh/RecastNavMesh.cpp
+++ b/Engine/Source/Runtime/NavigationSystem/Private/NavMesh/RecastNavMesh.cpp
@@ -84,6 +84,7 @@ namespace UE::NavMesh::Private
 		return FMath::Clamp(CellSize, ArbitraryMinCellSize, ArbitraryMaxCellSize);
 	}
 
+#if WITH_RECAST
 	FNavTileRef GetTileRefFromPolyRef(const dtNavMesh& DetourMesh, const NavNodeRef PolyRef)
 	{
 		unsigned int Salt = 0;
@@ -93,10 +94,12 @@ namespace UE::NavMesh::Private
 		const dtTileRef TileRef = DetourMesh.encodePolyId(Salt, TileIndex, 0);
 		return FNavTileRef(TileRef);
 	}
+#endif // WITH_RECAST
 } // namespace UE::NavMesh::Private
 
 FDetourTileLayout::FDetourTileLayout(const dtMeshTile& tile)
 {
+#if WITH_RECAST
 	const dtMeshHeader* header = tile.header;
 
 	if (header && (header->version == DT_NAVMESH_VERSION))
@@ -123,6 +126,7 @@ FDetourTileLayout::FDetourTileLayout(const dtMeshTile& tile)
 
 		InitFromSizeInfo(SizeInfo);
 	}
+#endif // WITH_RECAST
 }
 
 FDetourTileLayout::FDetourTileLayout(const FDetourTileSizeInfo& SizeInfo)
@@ -132,6 +136,7 @@ FDetourTileLayout::FDetourTileLayout(const FDetourTileSizeInfo& SizeInfo)
 
 void FDetourTileLayout::InitFromSizeInfo(const FDetourTileSizeInfo& SizeInfo)
 {
+#if WITH_RECAST
 	// Patch header pointers.
 	HeaderSize = dtAlign(sizeof(dtMeshHeader));
 	VertsSize = dtAlign(sizeof(dtReal) * 3 * SizeInfo.VertCount);
@@ -154,6 +159,7 @@ void FDetourTileLayout::InitFromSizeInfo(const FDetourTileSizeInfo& SizeInfo)
 
 	TileSize = HeaderSize + VertsSize + PolysSize + LinksSize + DetailMeshesSize + DetailVertsSize + DetailTrisSize
 		+ BvTreeSize + OffMeshConsSize + OffMeshSegsSize + ClustersSize + PolyClustersSize;
+#endif // WITH_RECAST
 }
 
 
@@ -269,6 +275,13 @@ ARecastNavMesh::~ARecastNavMesh()
 {
 }
 
+#if WITH_EDITOR
+bool ARecastNavMesh::CanEditChange(const FProperty* InPropery) const
+{
+	return false;
+}
+#endif // WITH_EDITOR
+
 #else // WITH_RECAST
 
 #include "Detour/DetourNavMesh.h"
@@ -396,6 +409,8 @@ namespace FNavMeshConfig
 	}
 }
 
+#endif // WITH_RECAST
+
 // Deprecated
 FRecastNavMeshGenerationProperties::FRecastNavMeshGenerationProperties()
 {
@@ -425,6 +440,8 @@ FRecastNavMeshGenerationProperties::FRecastNavMeshGenerationProperties()
 	bIsWorldPartitioned = false;
 }
 
+#if WITH_RECAST
+
 // Deprecated
 FRecastNavMeshGenerationProperties::FRecastNavMeshGenerationProperties(const ARecastNavMesh& RecastNavMesh)
 {
@@ -455,6 +472,8 @@ FRecastNavMeshGenerationProperties::FRecastNavMeshGenerationProperties(const ARe
 	bIsWorldPartitioned = RecastNavMesh.bIsWorldPartitioned;
 }
 
+#endif // WITH_RECAST
+
 FRecastNavMeshTileGenerationDebug::FRecastNavMeshTileGenerationDebug()
 {
 	bEnabled = false;
@@ -477,6 +496,8 @@ FRecastNavMeshTileGenerationDebug::FRecastNavMeshTileGenerationDebug()
 	bTileCacheDetailMesh = false;
 }
 
+#if WITH_RECAST
+
 ARecastNavMesh::FNavPolyFlags ARecastNavMesh::NavLinkFlag = ARecastNavMesh::FNavPolyFlags(0);
 
 ARecastNavMesh::ARecastNavMesh(const FObjectInitializer& ObjectInitializer)
@@ -3783,11 +3804,14 @@ const dtNavMesh* ARecastNavMesh::GetRecastMesh() const
 //----------------------------------------------------------------------//
 bool ARecastNavMesh::K2_ReplaceAreaInTileBounds(FBox Bounds, TSubclassOf<UNavArea> OldArea, TSubclassOf<UNavArea> NewArea, bool ReplaceLinks)
 {
-	bool bReplaced = ReplaceAreaInTileBounds(Bounds, OldArea, NewArea, ReplaceLinks) > 0;
+	bool bReplaced = false;
+#if WITH_RECAST
+	bReplaced = ReplaceAreaInTileBounds(Bounds, OldArea, NewArea, ReplaceLinks) > 0;
 	if (bReplaced)
 	{
 		RequestDrawingUpdate();
 	}
+#endif // WITH_RECAST
 	return bReplaced;
 }
 
diff --git a/Engine/Source/Runtime/NavigationSystem/Private/NavMesh/RecastNavMeshDataChunk.cpp b/Engine/Source/Runtime/NavigationSystem/Private/NavMesh/RecastNavMeshDataChunk.cpp
--- a/Engine/Source/Runtime/NavigationSystem/Private/NavMesh/RecastNavMeshDataChunk.cpp
+++ b/Engine/Source/Runtime/NavigationSystem/Private/NavMesh/RecastNavMeshDataChunk.cpp
@@ -7,9 +7,12 @@
 #include "NavMesh/PImplRecastNavMesh.h"
 #include "NavMesh/RecastHelpers.h"
 #include "NavMesh/RecastVersion.h"
-#include "Detour/DetourNavMeshBuilder.h"
 #include "NavMesh/RecastNavMeshGenerator.h"
 
+#if WITH_RECAST
+#include "Detour/DetourNavMeshBuilder.h"
+#endif // WITH_RECAST
+
 #include UE_INLINE_GENERATED_CPP_BY_NAME(RecastNavMeshDataChunk)
 
 //----------------------------------------------------------------------//
@@ -71,11 +74,13 @@ namespace UE::NavMesh::Private
 {
 	bool IsUsingActiveTileGeneration(const ARecastNavMesh& NavMesh)
 	{
+#if WITH_RECAST
 		const UNavigationSystemV1* NavSys = FNavigationSystem::GetCurrent<UNavigationSystemV1>(NavMesh.GetWorld());
 		if (NavSys)
 		{
 			return NavMesh.IsUsingActiveTilesGeneration(*NavSys);
 		}
+#endif // WITH_RECAST
 		return false;
 	}
 } // namespace UE::NavMesh::Private
@@ -197,43 +202,52 @@ void URecastNavMeshDataChunk::SerializeRecastData(FArchive& Ar, int32 NavMeshVer
 // Deprecated
 TArray<uint32> URecastNavMeshDataChunk::AttachTiles(FPImplRecastNavMesh& NavMeshImpl)
 {
+	TArray<uint32> TileIds;
+#if WITH_RECAST
 	check(NavMeshImpl.NavMeshOwner);
 	const TArray<FNavTileRef> TileRefs = AttachTiles(*NavMeshImpl.NavMeshOwner);
-	TArray<uint32> TileIds;
 	FNavTileRef::DeprecatedGetTileIdsFromNavTileRefs(&NavMeshImpl, TileRefs, TileIds);
+#endif // WITH_RECAST
 	return TileIds;
 }
 
 // Deprecated
 TArray<uint32> URecastNavMeshDataChunk::AttachTiles(FPImplRecastNavMesh& NavMeshImpl, const bool bKeepCopyOfData, const bool bKeepCopyOfCacheData)
 {
+	TArray<uint32> TileIds;
+#if WITH_RECAST
 	check(NavMeshImpl.NavMeshOwner);
 	const TArray<FNavTileRef> TileRefs = AttachTiles(*NavMeshImpl.NavMeshOwner, bKeepCopyOfData, bKeepCopyOfCacheData);
-	TArray<uint32> TileIds;
 	FNavTileRef::DeprecatedGetTileIdsFromNavTileRefs(&NavMeshImpl, TileRefs, TileIds);
+#endif // WITH_RECAST
 	return TileIds;
 }
 
 // Deprecated
 TArray<uint32> URecastNavMeshDataChunk::DetachTiles(FPImplRecastNavMesh& NavMeshImpl)
 {
+	TArray<uint32> TileIds;
+#if WITH_RECAST
 	check(NavMeshImpl.NavMeshOwner);
 	const TArray<FNavTileRef> TileRefs = DetachTiles(*NavMeshImpl.NavMeshOwner);
-	TArray<uint32> TileIds;
 	FNavTileRef::DeprecatedGetTileIdsFromNavTileRefs(&NavMeshImpl, TileRefs, TileIds);
+#endif // WITH_RECAST
 	return TileIds;
 }
 
 // Deprecated
 TArray<uint32> URecastNavMeshDataChunk::DetachTiles(FPImplRecastNavMesh& NavMeshImpl, const bool bTakeDataOwnership, const bool bTakeCacheDataOwnership)
 {
+	TArray<uint32> TileIds;
+#if WITH_RECAST
 	check(NavMeshImpl.NavMeshOwner);
 	const TArray<FNavTileRef> TileRefs = DetachTiles(*NavMeshImpl.NavMeshOwner, bTakeDataOwnership, bTakeCacheDataOwnership);
-	TArray<uint32> TileIds;
 	FNavTileRef::DeprecatedGetTileIdsFromNavTileRefs(&NavMeshImpl, TileRefs, TileIds);
+#endif // WITH_RECAST
 	return TileIds;
 }
 
+#if WITH_RECAST
 TArray<FNavTileRef> URecastNavMeshDataChunk::AttachTiles(ARecastNavMesh& NavMesh)
 {
 	check(NavMesh.GetWorld());
@@ -253,7 +267,6 @@ TArray<FNavTileRef> URecastNavMeshDataChunk::AttachTiles(ARecastNavMesh& NavMesh
 	TArray<FNavTileRef> Result;
 	Result.Reserve(Tiles.Num());
 
-#if WITH_RECAST	
 	dtNavMesh* DetourNavMesh = NavMesh.GetRecastMesh();
 
 	if (DetourNavMesh != nullptr)
@@ -362,7 +375,6 @@ TArray<FNavTileRef> URecastNavMeshDataChunk::AttachTiles(ARecastNavMesh& NavMesh
 			}
 		}
 	}
-#endif// WITH_RECAST
 
 	UE_LOG(LogNavigation, Verbose, TEXT("Attached %d tiles to NavMesh - %s"), Result.Num(), *NavigationDataName.ToString());
 	return Result;
@@ -387,7 +399,6 @@ TArray<FNavTileRef> URecastNavMeshDataChunk::DetachTiles(ARecastNavMesh& NavMesh
 	TArray<FNavTileRef> Result;
 	Result.Reserve(Tiles.Num());
 
-#if WITH_RECAST
 	dtNavMesh* DetourNavMesh = NavMesh.GetRecastMesh();
 
 	if (DetourNavMesh != nullptr)
@@ -474,11 +485,11 @@ TArray<FNavTileRef> URecastNavMeshDataChunk::DetachTiles(ARecastNavMesh& NavMesh
 			TileData.Layer = 0;
 		}
 	}
-#endif// WITH_RECAST
 
 	UE_LOG(LogNavigation, Verbose, TEXT("Detached %d tiles from NavMesh - %s"), Result.Num(), *NavigationDataName.ToString());
 	return Result;
 }
+#endif // WITH_RECAST
 
 void URecastNavMeshDataChunk::MoveTiles(FPImplRecastNavMesh& NavMeshImpl, const FIntPoint& Offset, const FVector::FReal RotationDeg, const FVector2D& RotationCenter)
 {
@@ -543,6 +554,7 @@ void URecastNavMeshDataChunk::GetTiles(const FPImplRecastNavMesh* NavMeshImpl, c
 {
 	Tiles.Empty(TileIndices.Num());
 
+#if WITH_RECAST
 	const dtNavMesh* NavMesh = NavMeshImpl->DetourNavMesh;
 	
 	for (int32 TileIdx : TileIndices)
@@ -581,11 +593,13 @@ void URecastNavMeshDataChunk::GetTiles(const FPImplRecastNavMesh* NavMeshImpl, c
 			Tiles.Add(RecastTileData);
 		}
 	}
+#endif // WITH_RECAST
 }
 
 void URecastNavMeshDataChunk::GetTilesBounds(const FPImplRecastNavMesh& NavMeshImpl, const TArray<int32>& TileIndices, FBox& OutBounds) const
 {
 	OutBounds.Init();
+#if WITH_RECAST
 	const dtNavMesh* NavMesh = NavMeshImpl.DetourNavMesh;
 
 	for (const int32 TileIdx : TileIndices)
@@ -596,4 +610,5 @@ void URecastNavMeshDataChunk::GetTilesBounds(const FPImplRecastNavMesh& NavMeshI
 			OutBounds += Recast2UnrealBox(Tile->header->bmin, Tile->header->bmax);
 		}
 	}
-}
+#endif // WITH_RECAST
+}
\ No newline at end of file
diff --git a/Engine/Source/Runtime/NavigationSystem/Private/NavigationSystem.cpp b/Engine/Source/Runtime/NavigationSystem/Private/NavigationSystem.cpp
--- a/Engine/Source/Runtime/NavigationSystem/Private/NavigationSystem.cpp
+++ b/Engine/Source/Runtime/NavigationSystem/Private/NavigationSystem.cpp
@@ -26,6 +26,7 @@
 #include "Components/PrimitiveComponent.h"
 #include "UObject/UObjectThreadContext.h"
 #include "GameFramework/Pawn.h"
+#include "AI/NavDataGenerator.h"
 
 #if WITH_RECAST
 #include "NavMesh/RecastNavMesh.h"
@@ -398,7 +399,7 @@ FNavRegenTimeSliceManager::FNavRegenTimeSliceManager()
 	, MaxDesiredTileRegenDuration(0.7)
 	, TimeLastCall(-1.f)
 	, NavDataIdx(0)
-#if TIME_SLICE_NAV_REGEN
+#if WITH_RECAST && TIME_SLICE_NAV_REGEN
 	, bDoTimeSlicedUpdate(true)
 #else
 	, bDoTimeSlicedUpdate(false)
@@ -728,8 +729,10 @@ UNavigationSystemV1::UNavigationSystemV1(const FObjectInitializer& ObjectInitial
 		SetDefaultWalkableArea(UNavArea_Default::StaticClass());
 		SetDefaultObstacleArea(UNavArea_Obstacle::StaticClass());
 		
+#if WITH_RECAST
 		const FTransform RecastToUnrealTransfrom(Recast2UnrealMatrix());
 		SetCoordTransform(ENavigationCoordSystem::Navigation, ENavigationCoordSystem::Unreal, RecastToUnrealTransfrom);
+#endif // WITH_RECAST
 	}
 }
 
@@ -1194,12 +1197,14 @@ void UNavigationSystemV1::OnWorldInitDone(FNavigationSystemRunMode Mode)
 	{
 		if (NavData)
 		{
+#if WITH_RECAST
 			const ARecastNavMesh* RecastNavMesh = Cast<ARecastNavMesh>(NavData);
 			if (RecastNavMesh && RecastNavMesh->bIsWorldPartitioned && NavData->GetRuntimeGenerationMode() > ERuntimeGenerationType::Static)
 			{
 				DefaultDirtyAreasController.SetUseWorldPartitionedDynamicMode(true);
 				break;
 			}
+#endif // WITH_RECAST
 		}
 	}
 
@@ -2385,11 +2390,13 @@ void UNavigationSystemV1::ApplyWorldOffset(const FVector& InOffset, bool bWorldS
 			if (NavData)
 			{
 				NavData->ConditionalConstructGenerator();
+#if WITH_RECAST
 				ARecastNavMesh* RecastNavMesh = Cast<ARecastNavMesh>(NavData);
 				if (RecastNavMesh)
 				{
 					RecastNavMesh->RequestDrawingUpdate();
 				}
+#endif // WITH_RECAST
 			}
 		}
 	}
diff --git a/Engine/Source/Runtime/NavigationSystem/Private/NavigationTestingActor.cpp b/Engine/Source/Runtime/NavigationSystem/Private/NavigationTestingActor.cpp
--- a/Engine/Source/Runtime/NavigationSystem/Private/NavigationTestingActor.cpp
+++ b/Engine/Source/Runtime/NavigationSystem/Private/NavigationTestingActor.cpp
@@ -312,7 +312,7 @@ void ANavigationTestingActor::UpdatePathfinding()
 	LastPath.Reset();
 	ShowStepIndex = -1;
 	PathfindingSteps = 0;
-#if WITH_EDITORONLY_DATA
+#if WITH_RECAST && WITH_EDITORONLY_DATA
 	DebugSteps.Reset();
 #endif
 	UpdateNavData();
@@ -380,6 +380,7 @@ FVector ANavigationTestingActor::FindClosestWallLocation() const
 	}
 #endif // WITH_EDITORONLY_DATA
 
+#if WITH_RECAST
 	ARecastNavMesh* AsRecastNavMesh = Cast<ARecastNavMesh>(MyNavData);
 	if (AsRecastNavMesh)
 	{
@@ -387,6 +388,7 @@ FVector ANavigationTestingActor::FindClosestWallLocation() const
 		AsRecastNavMesh->FindDistanceToWall(GetActorLocation(), UNavigationQueryFilter::GetQueryFilter(*MyNavData, this, FilterClass), FLT_MAX, &TmpOutLocation);
 		return TmpOutLocation;
 	}
+#endif // WITH_RECAST
 	
 	return FNavigationSystem::InvalidLocation;
 }
diff --git a/Engine/Source/Runtime/NavigationSystem/Public/NavMesh/NavMeshRenderingComponent.h b/Engine/Source/Runtime/NavigationSystem/Public/NavMesh/NavMeshRenderingComponent.h
--- a/Engine/Source/Runtime/NavigationSystem/Public/NavMesh/NavMeshRenderingComponent.h
+++ b/Engine/Source/Runtime/NavigationSystem/Public/NavMesh/NavMeshRenderingComponent.h
@@ -202,8 +202,10 @@ protected:
 
 #if UE_ENABLE_DEBUG_DRAWING
   	NAVIGATIONSYSTEM_API virtual FDebugRenderSceneProxy* CreateDebugSceneProxy() override;
+#if WITH_RECAST
 	virtual FDebugDrawDelegateHelper& GetDebugDrawDelegateHelper() override { return NavMeshDebugDrawDelegateManager; }
-#endif
+#endif // WITH_RECAST
+#endif // UE_ENABLE_DEBUG_DRAWING
 
 	NAVIGATIONSYSTEM_API virtual FBoxSphereBounds CalcBounds(const FTransform &LocalToWorld) const override;
 
@@ -219,7 +221,7 @@ protected:
 	FTimerHandle TimerHandle;
 
 protected:
-#if UE_ENABLE_DEBUG_DRAWING
+#if WITH_RECAST && UE_ENABLE_DEBUG_DRAWING
 	FNavMeshDebugDrawDelegateHelper NavMeshDebugDrawDelegateManager;
 #endif
 };
diff --git a/Engine/Source/Runtime/NavigationSystem/Public/NavMesh/RecastInternalDebugData.cpp b/Engine/Source/Runtime/NavigationSystem/Public/NavMesh/RecastInternalDebugData.cpp
--- a/Engine/Source/Runtime/NavigationSystem/Public/NavMesh/RecastInternalDebugData.cpp
+++ b/Engine/Source/Runtime/NavigationSystem/Public/NavMesh/RecastInternalDebugData.cpp
@@ -2,11 +2,14 @@
 
 #include "NavMesh/RecastInternalDebugData.h"
 
-#include "DebugUtils/DebugDraw.h"
 #include "HAL/PlatformCrt.h"
 #include "Misc/AssertionMacros.h"
 #include "NavMesh/RecastHelpers.h"
 
+#if WITH_RECAST
+
+#include "DebugUtils/DebugDraw.h"
+
 void FRecastInternalDebugData::vertex(const FVector::FReal x, const FVector::FReal y, const FVector::FReal z, unsigned int color, const FVector::FReal u, const FVector::FReal v)
 {
 	FVector::FReal RecastPos[3] = { x,y,z };
@@ -57,3 +60,4 @@ void FRecastInternalDebugData::end()
 		}
 	}
 }
+#endif // WITH_RECAST
diff --git a/Engine/Source/Runtime/NavigationSystem/Public/NavMesh/RecastInternalDebugData.h b/Engine/Source/Runtime/NavigationSystem/Public/NavMesh/RecastInternalDebugData.h
--- a/Engine/Source/Runtime/NavigationSystem/Public/NavMesh/RecastInternalDebugData.h
+++ b/Engine/Source/Runtime/NavigationSystem/Public/NavMesh/RecastInternalDebugData.h
@@ -3,12 +3,15 @@
 #pragma once
 
 #include "Containers/Array.h"
-#include "DebugUtils/DebugDraw.h"
 #include "Math/Color.h"
 #include "Math/UnrealMathSSE.h"
 #include "Math/Vector.h"
 #include "NavMesh/RecastHelpers.h"
 
+#if WITH_RECAST
+
+#include "DebugUtils/DebugDraw.h"
+
 struct FRecastInternalDebugData : public duDebugDraw
 {
 	duDebugDrawPrimitives CurrentPrim = DU_DRAW_POINTS;
@@ -62,4 +65,4 @@ struct FRecastInternalDebugData : public duDebugDraw
 
 	virtual void end() override;
 };
-
+#endif // WITH_RECAST
diff --git a/Engine/Source/Runtime/NavigationSystem/Public/NavMesh/RecastNavMesh.h b/Engine/Source/Runtime/NavigationSystem/Public/NavMesh/RecastNavMesh.h
--- a/Engine/Source/Runtime/NavigationSystem/Public/NavMesh/RecastNavMesh.h
+++ b/Engine/Source/Runtime/NavigationSystem/Public/NavMesh/RecastNavMesh.h
@@ -426,7 +426,9 @@ struct FRecastNavMeshGenerationProperties
 	uint32 bIsWorldPartitioned : 1;
 	
 	NAVIGATIONSYSTEM_API FRecastNavMeshGenerationProperties();
+#if WITH_RECAST
 	NAVIGATIONSYSTEM_API FRecastNavMeshGenerationProperties(const ARecastNavMesh& RecastNavMesh);
+#endif // WITH_RECAST
 };
 
 UENUM()
diff --git a/Engine/Source/Runtime/NavigationSystem/Public/NavMesh/RecastNavMeshDataChunk.h b/Engine/Source/Runtime/NavigationSystem/Public/NavMesh/RecastNavMeshDataChunk.h
--- a/Engine/Source/Runtime/NavigationSystem/Public/NavMesh/RecastNavMeshDataChunk.h
+++ b/Engine/Source/Runtime/NavigationSystem/Public/NavMesh/RecastNavMeshDataChunk.h
@@ -82,6 +82,7 @@ class URecastNavMeshDataChunk : public UNavigationDataChunk
 	UE_DEPRECATED(5.1, "Use overload using ARecastNavMesh& instead")
 	NAVIGATIONSYSTEM_API TArray<uint32> DetachTiles(FPImplRecastNavMesh& NavMeshImpl, const bool bTakeDataOwnership, const bool bTakeCacheDataOwnership);
 
+#if WITH_RECAST
 	/** Attaches tiles to specified navmesh, transferring tile ownership to navmesh */
 	NAVIGATIONSYSTEM_API TArray<FNavTileRef> AttachTiles(ARecastNavMesh& NavMesh);
 
@@ -93,6 +94,7 @@ class URecastNavMeshDataChunk : public UNavigationDataChunk
 
 	/** Detaches tiles from specified navmesh */
 	NAVIGATIONSYSTEM_API TArray<FNavTileRef> DetachTiles(ARecastNavMesh& NavMesh, const bool bTakeDataOwnership, const bool bTakeCacheDataOwnership);
+#endif // WITH_RECAST
 
 	/** 
 	 * Experimental: Moves tiles data on the xy plane by the offset (in tile coordinates) and rotation (in degree).
