[NavigationSystem] replaced requirement of PrimitiveComponent by INavRelevantInterface + UObject.
Added dedicated delegates for UObject based registration and update.

diff --git a/Engine/Source/Runtime/Engine/Classes/AI/Navigation/NavRelevantInterface.h b/Engine/Source/Runtime/Engine/Classes/AI/Navigation/NavRelevantInterface.h
--- a/Engine/Source/Runtime/Engine/Classes/AI/Navigation/NavRelevantInterface.h
+++ b/Engine/Source/Runtime/Engine/Classes/AI/Navigation/NavRelevantInterface.h
@@ -10,6 +10,27 @@
 #include "NavRelevantInterface.generated.h"
 
 struct FNavigableGeometryExport;
+class UBodySetup;
+
+/** Determines if a NavRelevant object contains custom collision for navigation/AI */
+UENUM()
+namespace EHasCustomNavigableGeometry
+{
+enum Type : int
+{
+	/** Primitive doesn't have custom navigation geometry, if collision is enabled then its convex/trimesh collision will be used for generating the navmesh */
+	No,
+
+	/** If primitive would normally affect navmesh, DoCustomNavigableGeometryExport() should be called to export this primitive's navigable geometry */
+	Yes,
+
+	/** DoCustomNavigableGeometryExport() should be called even if the mesh is non-collidable and wouldn't normally affect the navmesh */
+	EvenIfNotCollidable,
+
+	/** Don't export navigable geometry even if primitive is relevant for navigation (can still add modifiers) */
+	DontExport,
+};
+}
 
 struct FNavigationRelevantDataFilter 
 {
@@ -148,13 +169,13 @@ class INavRelevantInterface
 {
 	GENERATED_IINTERFACE_BODY()
 
-	/** Prepare navigation modifiers */
+	/** Prepares navigation modifiers */
 	virtual void GetNavigationData(FNavigationRelevantData& Data) const {}
 
 	/** Get bounds for navigation octree */
 	virtual FBox GetNavigationBounds() const { return FBox(ForceInit); }
 
-	/** if this instance knows how to export sub-sections of self */
+	/** Indicates if this instance knows how to export sub-sections of self */
 	virtual bool SupportsGatheringGeometrySlices() const { return false; }
 
 	/** This function is called "on demand", whenever that specified piece of geometry is needed for navigation generation */
@@ -162,11 +183,25 @@ class INavRelevantInterface
 
 	virtual ENavDataGatheringMode GetGeometryGatheringMode() const { return ENavDataGatheringMode::Default; }
 
-	/** Called on Game-thread to give implementer a change to perform actions that require game-thread to run, 
-	 *	for example precaching physics data */
+	/** Called on Game-thread to give implementer a chance to perform actions that require game-thread to run (e.g. precaching physics data) */
 	virtual void PrepareGeometryExportSync() {}
-
-	/** Update bounds, called after moving owning actor */
+	
+	/** Returns associated body setup (if any) for default geometry export.
+	 *	@return Associated UBodySetup if any, nullptr otherwise. */
+	virtual UBodySetup* GetNavigableGeometryBodySetup() { return nullptr; }
+	
+	/** Returns transform to used for default geometry export.
+	 *	@return Transform to used when exporting geometry. */
+	virtual FTransform GetNavigableGeometryTransform() const { return FTransform::Identity; }
+
+	/** If true then DoCustomNavigableGeometryExport will be called to collect navigable geometry for the object implementing this interface. */
+	virtual EHasCustomNavigableGeometry::Type HasCustomNavigableGeometry() const { return EHasCustomNavigableGeometry::No; }
+
+	/** Collects the custom navigable geometry of the object.
+	 *	@return true if regular navigable geometry exporting (using GetBodySetup) should be run as well */
+	virtual bool DoCustomNavigableGeometryExport(FNavigableGeometryExport& GeomExport) const { return true; }
+
+	/** Update bounds, called after the object moved */
 	virtual void UpdateNavigationBounds() {}
 
 	/** Are modifiers active? */
@@ -175,5 +210,5 @@ class INavRelevantInterface
 	/** Get navigation parent
 	 *  Adds modifiers to existing octree node, GetNavigationBounds and IsNavigationRelevant won't be checked
 	 */
-	virtual UObject* GetNavigationParent() const { return NULL; }
+	virtual UObject* GetNavigationParent() const { return nullptr; }
 };
diff --git a/Engine/Source/Runtime/Engine/Classes/AI/NavigationSystemBase.h b/Engine/Source/Runtime/Engine/Classes/AI/NavigationSystemBase.h
--- a/Engine/Source/Runtime/Engine/Classes/AI/NavigationSystemBase.h
+++ b/Engine/Source/Runtime/Engine/Classes/AI/NavigationSystemBase.h
@@ -148,7 +148,11 @@ namespace FNavigationSystem
 	
 	ENGINE_API void RegisterComponent(UActorComponent& Comp);
 	ENGINE_API void UnregisterComponent(UActorComponent& Comp);
-	
+
+	ENGINE_API void RegisterNavRelevantObject(UObject& Object);
+	ENGINE_API void UpdateNavRelevantObject(UObject& Object);
+	ENGINE_API void UnregisterNavRelevantObject(UObject& Object);
+
 	ENGINE_API void RemoveActorData(AActor& Actor);
 
 	ENGINE_API bool HasComponentData(UActorComponent& Comp);
@@ -190,6 +194,7 @@ namespace FNavigationSystem
 	ENGINE_API void StopMovement(const AController& Controller);
 	ENGINE_API IPathFollowingAgentInterface* FindPathFollowingAgentForActor(const AActor& Actor);
 
+	DECLARE_DELEGATE_OneParam(FObjectBasedSignature, UObject& /*Object*/);
 	DECLARE_DELEGATE_OneParam(FActorBasedSignature, AActor& /*Actor*/);
 	DECLARE_DELEGATE_OneParam(FActorComponentBasedSignature, UActorComponent& /*Comp*/);
 	DECLARE_DELEGATE_OneParam(FSceneComponentBasedSignature, USceneComponent& /*Comp*/);
@@ -276,6 +281,11 @@ protected:
 	static ENGINE_API void SetDefaultObstacleArea(TSubclassOf<UNavAreaBase> InAreaClass);
 
 	static ENGINE_API void ResetEventDelegates();
+
+	static ENGINE_API FNavigationSystem::FObjectBasedSignature& RegisterNavRelevantObjectDelegate();
+	static ENGINE_API FNavigationSystem::FObjectBasedSignature& UpdateNavRelevantObjectDelegate();
+	static ENGINE_API FNavigationSystem::FObjectBasedSignature& UnregisterNavRelevantObjectDelegate();
+	
 	static ENGINE_API FNavigationSystem::FActorBasedSignature& UpdateActorDataDelegate();
 	static ENGINE_API FNavigationSystem::FActorComponentBasedSignature& UpdateComponentDataDelegate();
 	static ENGINE_API FNavigationSystem::FSceneComponentBasedSignature& UpdateComponentDataAfterMoveDelegate();
diff --git a/Engine/Source/Runtime/Engine/Classes/Components/PrimitiveComponent.h b/Engine/Source/Runtime/Engine/Classes/Components/PrimitiveComponent.h
--- a/Engine/Source/Runtime/Engine/Classes/Components/PrimitiveComponent.h
+++ b/Engine/Source/Runtime/Engine/Classes/Components/PrimitiveComponent.h
@@ -73,26 +73,6 @@ enum ECanBeCharacterBase : int
 	ECB_MAX,
 };
 
-/** Determines if a primitive component contains custom collision for navigation/AI */
-UENUM()
-namespace EHasCustomNavigableGeometry
-{
-	enum Type : int
-	{
-		/** Primitive doesn't have custom navigation geometry, if collision is enabled then its convex/trimesh collision will be used for generating the navmesh */
-		No,
-
-		/** If primitive would normally affect navmesh, DoCustomNavigableGeometryExport() should be called to export this primitive's navigable geometry */
-		Yes,
-
-		/** DoCustomNavigableGeometryExport() should be called even if the mesh is non-collidable and wouldn't normally affect the navmesh */
-		EvenIfNotCollidable,
-
-		/** Don't export navigable geometry even if primitive is relevant for navigation (can still add modifiers) */
-		DontExport,
-	};
-}
-
 /** Information about the sprite category, used for visualization in the editor */
 USTRUCT()
 struct FSpriteCategoryInfo
@@ -2941,10 +2921,17 @@ public:
 	ENGINE_API virtual void GetNavigationData(FNavigationRelevantData& OutData) const override;
 	ENGINE_API virtual FBox GetNavigationBounds() const override;
 	ENGINE_API virtual bool IsNavigationRelevant() const override;
-	//~ End INavRelevantInterface Interface
+	ENGINE_API virtual UBodySetup* GetNavigableGeometryBodySetup() override final; // marked as final since PrimitiveComponent derived classes relies on GetBodySetup()
+	ENGINE_API virtual FTransform GetNavigableGeometryTransform() const override final; // marked as final since PrimitiveComponent derived classes relies on GetComponentTransform()
 
 	/** If true then DoCustomNavigableGeometryExport will be called to collect navigable geometry of this component. */
-	FORCEINLINE EHasCustomNavigableGeometry::Type HasCustomNavigableGeometry() const { return bHasCustomNavigableGeometry; }
+	ENGINE_API virtual EHasCustomNavigableGeometry::Type HasCustomNavigableGeometry() const override;
+
+	/** Collects custom navigable geometry of component.
+	 *	@return true if regular navigable geometry exporting should be run as well
+	 */
+	ENGINE_API virtual bool DoCustomNavigableGeometryExport(FNavigableGeometryExport& GeomExport) const override;
+	//~ End INavRelevantInterface Interface
 
 	// Returns true if we should check the GetGenerateOverlapEvents() flag when gathering overlaps, otherwise we'll always just do it.
 	ENGINE_API FORCEINLINE_DEBUGGABLE bool ShouldCheckOverlapFlagToQueueOverlaps(const UPrimitiveComponent& ThisComponent) const;
@@ -2952,10 +2939,6 @@ public:
 	/** Set value of HasCustomNavigableGeometry */
 	ENGINE_API void SetCustomNavigableGeometry(const EHasCustomNavigableGeometry::Type InType);
 
-	/** Collects custom navigable geometry of component.
-	*	@return true if regular navigable geometry exporting should be run as well */
-	virtual bool DoCustomNavigableGeometryExport(FNavigableGeometryExport& GeomExport) const { return true; }
-
 	static ENGINE_API void DispatchMouseOverEvents(UPrimitiveComponent* CurrentComponent, UPrimitiveComponent* NewComponent);
 	static ENGINE_API void DispatchTouchOverEvents(ETouchIndex::Type FingerIndex, UPrimitiveComponent* CurrentComponent, UPrimitiveComponent* NewComponent);
 	ENGINE_API void DispatchOnClicked(FKey ButtonClicked = EKeys::LeftMouseButton);
diff --git a/Engine/Source/Runtime/Engine/Private/AI/NavigationSystemBase.cpp b/Engine/Source/Runtime/Engine/Private/AI/NavigationSystemBase.cpp
--- a/Engine/Source/Runtime/Engine/Private/AI/NavigationSystemBase.cpp
+++ b/Engine/Source/Runtime/Engine/Private/AI/NavigationSystemBase.cpp
@@ -127,6 +127,10 @@ namespace FNavigationSystem
 	class FDelegates
 	{
 	public:
+		FObjectBasedSignature RegisterNavRelevantObject;
+		FObjectBasedSignature UpdateNavRelevantObject;
+		FObjectBasedSignature UnregisterNavRelevantObject;
+		
 		FActorBasedSignature UpdateActorData;
 		FActorComponentBasedSignature UpdateComponentData;
 		FSceneComponentBasedSignature UpdateComponentDataAfterMove;
@@ -167,6 +171,10 @@ namespace FNavigationSystem
 
 		FDelegates()
 		{
+			RegisterNavRelevantObject.BindLambda([](UObject&) {});
+			UpdateNavRelevantObject.BindLambda([](UObject&) {});
+			UnregisterNavRelevantObject.BindLambda([](UObject&) {});
+
 			UpdateActorData.BindLambda([](AActor&) {});
 			UpdateComponentData.BindLambda([](UActorComponent&) {});
 			UpdateComponentDataAfterMove.BindLambda([](UActorComponent&) {});
@@ -207,6 +215,10 @@ namespace FNavigationSystem
 
 	void ResetDelegates() { new(&Delegates)FDelegates(); }
 
+	void RegisterNavRelevantObject(UObject& Object) { Delegates.RegisterNavRelevantObject.Execute(Object); }
+	void UpdateNavRelevantObject(UObject& Object) { Delegates.UpdateNavRelevantObject.Execute(Object); }
+	void UnregisterNavRelevantObject(UObject& Object) { Delegates.UnregisterNavRelevantObject.Execute(Object); }
+
 	void UpdateActorData(AActor& Actor) { Delegates.UpdateActorData.Execute(Actor); }
 	void UpdateComponentData(UActorComponent& Comp) { Delegates.UpdateComponentData.Execute(Comp); }
 	void UpdateActorAndComponentData(AActor& Actor, bool bUpdateAttachedActors) { Delegates.UpdateActorAndComponentData.Execute(Actor, bUpdateAttachedActors); }
@@ -387,6 +399,11 @@ void UNavigationSystemBase::SetDefaultObstacleArea(TSubclassOf<UNavAreaBase> InA
 
 
 void UNavigationSystemBase::ResetEventDelegates() { FNavigationSystem::ResetDelegates(); }
+
+FNavigationSystem::FObjectBasedSignature& UNavigationSystemBase::RegisterNavRelevantObjectDelegate() { return FNavigationSystem::Delegates.RegisterNavRelevantObject; }
+FNavigationSystem::FObjectBasedSignature& UNavigationSystemBase::UpdateNavRelevantObjectDelegate() { return FNavigationSystem::Delegates.UpdateNavRelevantObject; }
+FNavigationSystem::FObjectBasedSignature& UNavigationSystemBase::UnregisterNavRelevantObjectDelegate() { return FNavigationSystem::Delegates.UnregisterNavRelevantObject; }
+
 FNavigationSystem::FActorBasedSignature& UNavigationSystemBase::UpdateActorDataDelegate() { return FNavigationSystem::Delegates.UpdateActorData; }
 FNavigationSystem::FActorComponentBasedSignature& UNavigationSystemBase::UpdateComponentDataDelegate() { return FNavigationSystem::Delegates.UpdateComponentData; }
 FNavigationSystem::FSceneComponentBasedSignature& UNavigationSystemBase::UpdateComponentDataAfterMoveDelegate() { return FNavigationSystem::Delegates.UpdateComponentDataAfterMove; }
diff --git a/Engine/Source/Runtime/Engine/Private/Components/PrimitiveComponent.cpp b/Engine/Source/Runtime/Engine/Private/Components/PrimitiveComponent.cpp
--- a/Engine/Source/Runtime/Engine/Private/Components/PrimitiveComponent.cpp
+++ b/Engine/Source/Runtime/Engine/Private/Components/PrimitiveComponent.cpp
@@ -2964,6 +2964,26 @@ bool UPrimitiveComponent::IsNavigationRelevant() const
 		(ResponseToChannels.GetResponse(ECC_Pawn) == ECR_Block || ResponseToChannels.GetResponse(ECC_Vehicle) == ECR_Block);
 }
 
+UBodySetup* UPrimitiveComponent::GetNavigableGeometryBodySetup()
+{
+	return GetBodySetup();
+}
+
+FTransform UPrimitiveComponent::GetNavigableGeometryTransform() const
+{
+	return GetComponentTransform();
+}
+
+EHasCustomNavigableGeometry::Type UPrimitiveComponent::HasCustomNavigableGeometry() const
+{
+	return bHasCustomNavigableGeometry;
+}
+
+bool UPrimitiveComponent::DoCustomNavigableGeometryExport(FNavigableGeometryExport& GeomExport) const
+{
+	return true;
+}
+
 FBox UPrimitiveComponent::GetNavigationBounds() const
 {
 	// Return invalid box when retrieving NavigationBounds before they are being computed at component registration
diff --git a/Engine/Source/Runtime/NavigationSystem/Private/NavMesh/RecastNavMeshGenerator.cpp b/Engine/Source/Runtime/NavigationSystem/Private/NavMesh/RecastNavMeshGenerator.cpp
--- a/Engine/Source/Runtime/NavigationSystem/Private/NavMesh/RecastNavMeshGenerator.cpp
+++ b/Engine/Source/Runtime/NavigationSystem/Private/NavMesh/RecastNavMeshGenerator.cpp
@@ -2,7 +2,6 @@
 
 #include "NavMesh/RecastNavMeshGenerator.h"
 #include "AI/Navigation/NavRelevantInterface.h"
-#include "Components/PrimitiveComponent.h"
 #include "Compression/OodleDataCompression.h"
 #include "Engine/Level.h"
 #include "GameFramework/Pawn.h"
@@ -996,23 +995,18 @@ FORCEINLINE_DEBUGGABLE void ExportRigidBodySetup(UBodySetup& BodySetup, TNavStat
 	TemporaryShapeBuffer.Reset();
 }
 
-FORCEINLINE_DEBUGGABLE void ExportComponent(UActorComponent* Component, FRecastGeometryExport& GeomExport, const FBox* ClipBounds=NULL)
+void ExportObject(INavRelevantInterface& NavRelevantInterface, FRecastGeometryExport& GeomExport)
 {
-	bool bHasData = false;
-
-	UPrimitiveComponent* PrimComp = Cast<UPrimitiveComponent>(Component);
-	if (PrimComp && PrimComp->IsNavigationRelevant() && (PrimComp->HasCustomNavigableGeometry() != EHasCustomNavigableGeometry::DontExport))
+	if (NavRelevantInterface.IsNavigationRelevant() && (NavRelevantInterface.HasCustomNavigableGeometry() != EHasCustomNavigableGeometry::DontExport))
 	{
-		if ((PrimComp->HasCustomNavigableGeometry() != EHasCustomNavigableGeometry::Type::No) && !PrimComp->DoCustomNavigableGeometryExport(GeomExport))
-		{
-			bHasData = true;
-		}
+		const bool bHasData = (NavRelevantInterface.HasCustomNavigableGeometry() != EHasCustomNavigableGeometry::Type::No)
+			&& !NavRelevantInterface.DoCustomNavigableGeometryExport(GeomExport);
 
 		UBodySetup* BodySetup = nullptr;
 		{
 			// Might need to create the BodySetup outside of the main thread so garbage collection guard is required.
 			FGCScopeGuard GCGuard;
-			BodySetup = PrimComp->GetBodySetup();
+			BodySetup = NavRelevantInterface.GetNavigableGeometryBodySetup();
 
 			if (BodySetup)
 			{
@@ -1025,8 +1019,7 @@ FORCEINLINE_DEBUGGABLE void ExportComponent(UActorComponent* Component, FRecastG
 		{
 			if (!bHasData)
 			{
-				ExportRigidBodySetup(*BodySetup, GeomExport.VertexBuffer, GeomExport.IndexBuffer, GeomExport.Data->Bounds, PrimComp->GetComponentTransform());
-				bHasData = true;
+				ExportRigidBodySetup(*BodySetup, GeomExport.VertexBuffer, GeomExport.IndexBuffer, GeomExport.Data->Bounds, NavRelevantInterface.GetNavigableGeometryTransform());
 			}
 
 			GeomExport.SlopeOverride = BodySetup->WalkableSlopeOverride;
@@ -6868,10 +6861,18 @@ void FRecastNavMeshGenerator::GetDebugGeometry(const FNavigationRelevantData& En
 }
 #endif // !UE_BUILD_SHIPPING
 
-void FRecastNavMeshGenerator::ExportComponentGeometry(UActorComponent* Component, FNavigationRelevantData& Data)
+void FRecastNavMeshGenerator::ExportComponentGeometry(UActorComponent* InOutComponent, FNavigationRelevantData& OutData)
 {
-	FRecastGeometryExport GeomExport(Data);
-	RecastGeometryExport::ExportComponent(Component, GeomExport);
+	if (INavRelevantInterface* NavRelevantInterface = Cast<INavRelevantInterface>(InOutComponent))
+	{
+		ExportNavRelevantObjectGeometry(*NavRelevantInterface, OutData);
+	}
+}
+
+void FRecastNavMeshGenerator::ExportNavRelevantObjectGeometry(INavRelevantInterface& InOutNavRelevantInterface, FNavigationRelevantData& OutData)
+{
+	FRecastGeometryExport GeomExport(OutData);
+	RecastGeometryExport::ExportObject(InOutNavRelevantInterface, GeomExport);
 
 #if !UE_BUILD_SHIPPING	
 	RecastGeometryExport::ValidateGeometryExport(GeomExport);
@@ -6881,10 +6882,10 @@ void FRecastNavMeshGenerator::ExportComponentGeometry(UActorComponent* Component
 	RecastGeometryExport::StoreCollisionCache(GeomExport);
 }
 
-void FRecastNavMeshGenerator::ExportVertexSoupGeometry(const TArray<FVector>& Verts, FNavigationRelevantData& Data)
+void FRecastNavMeshGenerator::ExportVertexSoupGeometry(const TArray<FVector>& InVerts, FNavigationRelevantData& OutData)
 {
-	FRecastGeometryExport GeomExport(Data);
-	RecastGeometryExport::ExportVertexSoup(Verts, GeomExport.VertexBuffer, GeomExport.IndexBuffer, GeomExport.Data->Bounds);
+	FRecastGeometryExport GeomExport(OutData);
+	RecastGeometryExport::ExportVertexSoup(InVerts, GeomExport.VertexBuffer, GeomExport.IndexBuffer, GeomExport.Data->Bounds);
 
 #if !UE_BUILD_SHIPPING	
 	RecastGeometryExport::ValidateGeometryExport(GeomExport);
@@ -6893,12 +6894,12 @@ void FRecastNavMeshGenerator::ExportVertexSoupGeometry(const TArray<FVector>& Ve
 	RecastGeometryExport::StoreCollisionCache(GeomExport);
 }
 
-void FRecastNavMeshGenerator::ExportRigidBodyGeometry(UBodySetup& BodySetup, TNavStatArray<FVector>& OutVertexBuffer, TNavStatArray<int32>& OutIndexBuffer, const FTransform& LocalToWorld)
+void FRecastNavMeshGenerator::ExportRigidBodyGeometry(UBodySetup& InOutBodySetup, TNavStatArray<FVector>& OutVertexBuffer, TNavStatArray<int32>& OutIndexBuffer, const FTransform& LocalToWorld)
 {
 	TNavStatArray<FVector::FReal> VertCoords;
 	FBox TempBounds;
 
-	RecastGeometryExport::ExportRigidBodySetup(BodySetup, VertCoords, OutIndexBuffer, TempBounds, LocalToWorld);
+	RecastGeometryExport::ExportRigidBodySetup(InOutBodySetup, VertCoords, OutIndexBuffer, TempBounds, LocalToWorld);
 
 	OutVertexBuffer.Reserve(OutVertexBuffer.Num() + (VertCoords.Num() / 3));
 	for (int32 i = 0; i < VertCoords.Num(); i += 3)
@@ -6907,17 +6908,17 @@ void FRecastNavMeshGenerator::ExportRigidBodyGeometry(UBodySetup& BodySetup, TNa
 	}
 }
 
-void FRecastNavMeshGenerator::ExportRigidBodyGeometry(UBodySetup& BodySetup, TNavStatArray<FVector>& OutTriMeshVertexBuffer, TNavStatArray<int32>& OutTriMeshIndexBuffer
+void FRecastNavMeshGenerator::ExportRigidBodyGeometry(UBodySetup& InOutBodySetup, TNavStatArray<FVector>& OutTriMeshVertexBuffer, TNavStatArray<int32>& OutTriMeshIndexBuffer
 	, TNavStatArray<FVector>& OutConvexVertexBuffer, TNavStatArray<int32>& OutConvexIndexBuffer, TNavStatArray<int32>& OutShapeBuffer
 	, const FTransform& LocalToWorld)
 {
-	BodySetup.CreatePhysicsMeshes();
+	InOutBodySetup.CreatePhysicsMeshes();
 
 	TNavStatArray<FVector::FReal> VertCoords;
 	FBox TempBounds;
 
 	VertCoords.Reset();
-	RecastGeometryExport::ExportRigidBodyTriMesh(BodySetup, VertCoords, OutTriMeshIndexBuffer, TempBounds, LocalToWorld);
+	RecastGeometryExport::ExportRigidBodyTriMesh(InOutBodySetup, VertCoords, OutTriMeshIndexBuffer, TempBounds, LocalToWorld);
 
 	OutTriMeshVertexBuffer.Reserve(OutTriMeshVertexBuffer.Num() + (VertCoords.Num() / 3));
 	for (int32 i = 0; i < VertCoords.Num(); i += 3)
@@ -6927,10 +6928,10 @@ void FRecastNavMeshGenerator::ExportRigidBodyGeometry(UBodySetup& BodySetup, TNa
 
 	const int32 NumExistingVerts = OutConvexVertexBuffer.Num();
 	VertCoords.Reset();
-	RecastGeometryExport::ExportRigidBodyConvexElements(BodySetup, VertCoords, OutConvexIndexBuffer, OutShapeBuffer, TempBounds, LocalToWorld);
-	RecastGeometryExport::ExportRigidBodyBoxElements(BodySetup.AggGeom, VertCoords, OutConvexIndexBuffer, OutShapeBuffer, TempBounds, LocalToWorld, NumExistingVerts);
-	RecastGeometryExport::ExportRigidBodySphylElements(BodySetup.AggGeom, VertCoords, OutConvexIndexBuffer, OutShapeBuffer, TempBounds, LocalToWorld, NumExistingVerts);
-	RecastGeometryExport::ExportRigidBodySphereElements(BodySetup.AggGeom, VertCoords, OutConvexIndexBuffer, OutShapeBuffer, TempBounds, LocalToWorld, NumExistingVerts);
+	RecastGeometryExport::ExportRigidBodyConvexElements(InOutBodySetup, VertCoords, OutConvexIndexBuffer, OutShapeBuffer, TempBounds, LocalToWorld);
+	RecastGeometryExport::ExportRigidBodyBoxElements(InOutBodySetup.AggGeom, VertCoords, OutConvexIndexBuffer, OutShapeBuffer, TempBounds, LocalToWorld, NumExistingVerts);
+	RecastGeometryExport::ExportRigidBodySphylElements(InOutBodySetup.AggGeom, VertCoords, OutConvexIndexBuffer, OutShapeBuffer, TempBounds, LocalToWorld, NumExistingVerts);
+	RecastGeometryExport::ExportRigidBodySphereElements(InOutBodySetup.AggGeom, VertCoords, OutConvexIndexBuffer, OutShapeBuffer, TempBounds, LocalToWorld, NumExistingVerts);
 	
 	OutConvexVertexBuffer.Reserve(OutConvexVertexBuffer.Num() + (VertCoords.Num() / 3));
 	for (int32 i = 0; i < VertCoords.Num(); i += 3)
diff --git a/Engine/Source/Runtime/NavigationSystem/Private/NavigationDataHandler.cpp b/Engine/Source/Runtime/NavigationSystem/Private/NavigationDataHandler.cpp
--- a/Engine/Source/Runtime/NavigationSystem/Private/NavigationDataHandler.cpp
+++ b/Engine/Source/Runtime/NavigationSystem/Private/NavigationDataHandler.cpp
@@ -356,9 +356,9 @@ void FNavigationDataHandler::UpdateNavOctreeParentChain(UObject& ElementOwner, b
 	}
 }
 
-bool FNavigationDataHandler::UpdateNavOctreeElementBounds(UActorComponent& Comp, const FBox& NewBounds, const FBox& DirtyArea)
+bool FNavigationDataHandler::UpdateNavOctreeElementBounds(UObject& Object, const FBox& NewBounds, const FBox& DirtyArea)
 {
-	const FOctreeElementId2* ElementId = OctreeController.GetObjectsNavOctreeId(Comp);
+	const FOctreeElementId2* ElementId = OctreeController.GetObjectsNavOctreeId(Object);
 	if (ElementId != nullptr && ensure(OctreeController.IsValidElement(*ElementId)))
 	{
 		OctreeController.NavOctree->UpdateNode(*ElementId, NewBounds);
@@ -366,12 +366,12 @@ bool FNavigationDataHandler::UpdateNavOctreeElementBounds(UActorComponent& Comp,
 		// Add dirty area
 		if (DirtyArea.IsValid)
 		{
-			// Refresh ElementId since components may be stored in a different node after updating bounds
-			ElementId = OctreeController.GetObjectsNavOctreeId(Comp);
+			// Refresh ElementId since object may be stored in a different node after updating bounds
+			ElementId = OctreeController.GetObjectsNavOctreeId(Object);
 			if (ElementId != nullptr && ensure(OctreeController.IsValidElement(*ElementId)))
 			{
 				const FNavigationOctreeElement& ElementData = OctreeController.NavOctree->GetElementById(*ElementId);
-				DirtyAreasController.AddArea(DirtyArea, ElementData.Data->GetDirtyFlag(), [&Comp] { return &Comp; }, nullptr, "Bounds change");
+				DirtyAreasController.AddArea(DirtyArea, ElementData.Data->GetDirtyFlag(), [&Object] { return &Object; }, nullptr, "Bounds change");
 			}
 		}
 
diff --git a/Engine/Source/Runtime/NavigationSystem/Private/NavigationDirtyAreasController.cpp b/Engine/Source/Runtime/NavigationSystem/Private/NavigationDirtyAreasController.cpp
--- a/Engine/Source/Runtime/NavigationSystem/Private/NavigationDirtyAreasController.cpp
+++ b/Engine/Source/Runtime/NavigationSystem/Private/NavigationDirtyAreasController.cpp
@@ -92,26 +92,22 @@ void FNavigationDirtyAreasController::AddArea(const FBox& NewArea, const int32 F
 
 #if !UE_BUILD_SHIPPING
 	auto DumpExtraInfo = [ObjectProviderFunc, DebugReason, BoundsSize, NewArea]() {
-		const UObject* ObjectProvider = nullptr;
+		FString ObjectInfo;
+		const UObject* Object = nullptr;
 		if (ObjectProviderFunc)
 		{
-			ObjectProvider = ObjectProviderFunc();
-		}
-
-		FString ComponentInfo;
-		if (const UActorComponent* ObjectAsComponent = Cast<UActorComponent>(ObjectProvider))
-		{
-			if (const AActor* ComponentOwner = ObjectAsComponent->GetOwner())
+			Object = ObjectProviderFunc();
+			if (const UObject* ObjectOwner = (Object != nullptr ? Object->GetOuter() : nullptr))
 			{
-				UE_VLOG_BOX(ComponentOwner, LogNavigationDirtyArea, Log, NewArea, FColor::Red, TEXT(""));
-				ComponentInfo = FString::Printf(TEXT(" | Component's owner: %s"), *GetFullNameSafe(ComponentOwner));
+				UE_VLOG_BOX(ObjectOwner, LogNavigationDirtyArea, Log, NewArea, FColor::Red, TEXT(""));
+				ObjectInfo = FString::Printf(TEXT(" | Element's owner: %s"), *GetFullNameSafe(ObjectOwner));
 			}
 		}
 
 		return FString::Printf(TEXT("Object: %s (from: %s)%s | Bounds: %s"),
-			*GetFullNameSafe(ObjectProvider),
+			*GetFullNameSafe(Object),
 			*DebugReason.ToString(),
-			*ComponentInfo,
+			*ObjectInfo,
 			*BoundsSize.ToString());
 	};
 
diff --git a/Engine/Source/Runtime/NavigationSystem/Private/NavigationOctree.cpp b/Engine/Source/Runtime/NavigationSystem/Private/NavigationOctree.cpp
--- a/Engine/Source/Runtime/NavigationSystem/Private/NavigationOctree.cpp
+++ b/Engine/Source/Runtime/NavigationSystem/Private/NavigationOctree.cpp
@@ -47,8 +47,9 @@ void FNavigationOctree::DemandLazyDataGathering(FNavigationRelevantData& Element
 {
 	LLM_SCOPE_BYTAG(NavigationOctree);
 
-	UObject* ElementOb = ElementData.GetOwner();
-	if (ElementOb == nullptr)
+	UObject* const NavRelevantObject = ElementData.GetOwner();
+	INavRelevantInterface* const NavRelevantInterface = Cast<INavRelevantInterface>(NavRelevantObject);
+	if (NavRelevantInterface == nullptr)
 	{
 		return;
 	}
@@ -59,9 +60,8 @@ void FNavigationOctree::DemandLazyDataGathering(FNavigationRelevantData& Element
 	if (ElementData.IsPendingLazyGeometryGathering() == true && ElementData.SupportsGatheringGeometrySlices() == false)
 	{
 		QUICK_SCOPE_CYCLE_COUNTER(STAT_RecastNavMeshGenerator_LazyGeometryExport);
-		UActorComponent& ActorComp = *CastChecked<UActorComponent>(ElementOb);
-		ComponentExportDelegate.ExecuteIfBound(&ActorComp, ElementData);
 
+		NavRelevantGeometryExportDelegate.ExecuteIfBound(*NavRelevantInterface, ElementData);
 		bShrink = true;
 
 		// mark this element as no longer needing geometry gathering
@@ -77,9 +77,7 @@ void FNavigationOctree::DemandLazyDataGathering(FNavigationRelevantData& Element
 		const double StartTime = bCanOutputDurationWarning ? FPlatformTime::Seconds() : 0.0f;
 #endif //!UE_BUILD_SHIPPING
 
-		INavRelevantInterface* NavElement = Cast<INavRelevantInterface>(ElementOb);
-		check(NavElement);
-		NavElement->GetNavigationData(ElementData);
+		NavRelevantInterface->GetNavigationData(ElementData);
 		ElementData.bPendingLazyModifiersGathering = false;
 		bShrink = true;
 
@@ -90,9 +88,12 @@ void FNavigationOctree::DemandLazyDataGathering(FNavigationRelevantData& Element
 			const double DeltaTime = FPlatformTime::Seconds() - StartTime;
 			if (DeltaTime > GatheringNavModifiersTimeLimitWarning)
 			{
-				const UActorComponent* ObjectAsComponent = Cast<UActorComponent>(ElementOb);
-				const AActor* ComponentOwner = ObjectAsComponent ? ObjectAsComponent->GetOwner() : nullptr;
-				UE_LOG(LogNavigation, Warning, TEXT("The time (%f sec) for gathering navigation data on an INavRelevantInterface navigation element exceeded the time limit (%f sec) | NavElement = %s | Potential Component = %s | Potential Component's Owner = %s | Level = %s"), DeltaTime, GatheringNavModifiersTimeLimitWarning, *GetNameSafe(ElementOb), *GetNameSafe(ObjectAsComponent), *GetNameSafe(ComponentOwner), ElementOb ? *GetNameSafe(ElementOb->GetOutermost()) : TEXT("Invalid element"));
+				UE_LOG(LogNavigation, Warning, TEXT("The time (%f sec) for gathering navigation data on an INavRelevantInterface navigation element exceeded the time limit (%f sec) | NavElement = %s | NavElement's Owner = %s | Level = %s"),
+					DeltaTime,
+					GatheringNavModifiersTimeLimitWarning,
+					*GetNameSafe(NavRelevantObject),
+					*GetNameSafe(NavRelevantObject->GetOuter()),
+					*GetNameSafe(NavRelevantObject->GetOutermost()));
 			}
 		}
 #endif //!UE_BUILD_SHIPPING
@@ -107,7 +108,7 @@ void FNavigationOctree::DemandLazyDataGathering(FNavigationRelevantData& Element
 	}
 
 	const int32 ElementMemoryChange = IntCastChecked<int32>(ElementData.GetGeometryAllocatedSize()) - OrgElementMemory;
-	const_cast<FNavigationOctree*>(this)->NodesMemory += ElementMemoryChange;
+	NodesMemory += ElementMemoryChange;
 	INC_MEMORY_STAT_BY(STAT_Navigation_CollisionTreeMemory, ElementMemoryChange);
 }
 
@@ -144,40 +145,38 @@ void FNavigationOctree::AddNode(UObject* ElementOb, INavRelevantInterface* NavEl
 
 	if (!Bounds.IsValid)
 	{
-		UE_LOG(LogNavigation, Warning, TEXT("%s: Ignoring %s because bounds are invalid."), ANSI_TO_TCHAR(__FUNCTION__), *GetNameSafe(ElementOb));
+		UE_LOG(LogNavigation, Warning, TEXT("%hs: Ignoring %s because bounds are invalid."), __FUNCTION__, *GetNameSafe(ElementOb));
 		return;
 	}
 	
-	// we assume NavElement is ElementOb already cast
-	Element.Bounds = Bounds;	
+	Element.Bounds = Bounds;
 
 	if (NavElement)
 	{
+		checkf(ElementOb, TEXT("We assume NavElement is ElementOb already cast"));
+
 		const bool bDoInstantGathering = !IsLazyGathering(*NavElement);
 
 		if (bGatherGeometry)
 		{
-			UActorComponent* ActorComp = Cast<UActorComponent>(ElementOb);
-			if (ActorComp)
-			{
-				bool bIsActorCompiling = false;
+			bool bIsCompiling = false;
 #if WITH_EDITOR
-				IInterface_AsyncCompilation* AsyncCompiledActor = Cast<IInterface_AsyncCompilation>(ActorComp);
-				bIsActorCompiling = AsyncCompiledActor && AsyncCompiledActor->IsCompiling();
+			const IInterface_AsyncCompilation* AsyncCompiledObject = Cast<IInterface_AsyncCompilation>(ElementOb);
+			bIsCompiling = AsyncCompiledObject && AsyncCompiledObject->IsCompiling();
 #endif
-				// Skip custom navigation export during async compilation, the node will be invalidated once
-				// compilation finishes.
-				if (!bIsActorCompiling)
+
+			// Skip custom navigation export during async compilation, the node will be invalidated once
+			// compilation finishes.
+			if (!bIsCompiling)
+			{
+				if (bDoInstantGathering)
+				{
+					NavRelevantGeometryExportDelegate.ExecuteIfBound(*NavElement, *Element.Data);
+				}
+				else
 				{
-					if (bDoInstantGathering)
-					{
-						ComponentExportDelegate.ExecuteIfBound(ActorComp, *Element.Data);
-					}
-					else
-					{
-						Element.Data->bPendingLazyGeometryGathering = true;
-						Element.Data->bSupportsGatheringGeometrySlices = NavElement && NavElement->SupportsGatheringGeometrySlices();
-					}
+					Element.Data->bPendingLazyGeometryGathering = true;
+					Element.Data->bSupportsGatheringGeometrySlices = NavElement && NavElement->SupportsGatheringGeometrySlices();
 				}
 			}
 		}
@@ -199,9 +198,12 @@ void FNavigationOctree::AddNode(UObject* ElementOb, INavRelevantInterface* NavEl
 				const double DeltaTime = FPlatformTime::Seconds() - StartTime;
 				if (DeltaTime > GatheringNavModifiersTimeLimitWarning)
 				{
-					const UActorComponent* ObjectAsComponent = Cast<UActorComponent>(ElementOb);
-					const AActor* ComponentOwner = ObjectAsComponent ? ObjectAsComponent->GetOwner() : nullptr;
-					UE_LOG(LogNavigation, Warning, TEXT("The time (%f sec) for gathering navigation data on an INavRelevantInterface navigation element exceeded the time limit (%f sec) | NavElement = %s | Potential Component = %s | Potential Component's Owner = %s | Level = %s"), DeltaTime, GatheringNavModifiersTimeLimitWarning, *GetNameSafe(ElementOb), *GetNameSafe(ObjectAsComponent), *GetNameSafe(ComponentOwner), ElementOb ? *GetNameSafe(ElementOb->GetOutermost()) : TEXT("Invalid element"));
+					UE_LOG(LogNavigation, Warning, TEXT("The time (%f sec) for gathering navigation data on an INavRelevantInterface navigation element exceeded the time limit (%f sec) | NavElement = %s | Element's Owner = %s | Level = %s"),
+						DeltaTime,
+						GatheringNavModifiersTimeLimitWarning,
+						*GetNameSafe(ElementOb),
+						*GetNameSafe(ElementOb->GetOuter()),
+						*GetNameSafe(ElementOb->GetOutermost()));
 				}
 			}
 #endif //!UE_BUILD_SHIPPING
@@ -228,7 +230,7 @@ void FNavigationOctree::AppendToNode(const FOctreeElementId2& Id, INavRelevantIn
 {
 	LLM_SCOPE_BYTAG(NavigationOctree);
 
-	FNavigationOctreeElement OrgData = GetElementById(Id);
+	const FNavigationOctreeElement OrgData = GetElementById(Id);
 
 	Element = OrgData;
 	Element.Bounds = Bounds + OrgData.Bounds.GetBox();
diff --git a/Engine/Source/Runtime/NavigationSystem/Private/NavigationSystem.cpp b/Engine/Source/Runtime/NavigationSystem/Private/NavigationSystem.cpp
--- a/Engine/Source/Runtime/NavigationSystem/Private/NavigationSystem.cpp
+++ b/Engine/Source/Runtime/NavigationSystem/Private/NavigationSystem.cpp
@@ -571,6 +571,10 @@ UNavigationSystemV1::UNavigationSystemV1(const FObjectInitializer& ObjectInitial
 	{
 		FDelegatesInitializer()
 		{
+			UNavigationSystemBase::RegisterNavRelevantObjectDelegate().BindLambda([](UObject& Object) { UNavigationSystemV1::OnNavRelevantObjectRegistered(Object); });
+			UNavigationSystemBase::UpdateNavRelevantObjectDelegate().BindStatic(&UNavigationSystemV1::UpdateNavRelevantObjectInNavOctree);
+			UNavigationSystemBase::UnregisterNavRelevantObjectDelegate().BindLambda([](UObject& Object) { UNavigationSystemV1::OnNavRelevantObjectUnregistered(Object); });
+
 			UNavigationSystemBase::UpdateActorDataDelegate().BindStatic(&UNavigationSystemV1::UpdateActorInNavOctree);
 			UNavigationSystemBase::UpdateComponentDataDelegate().BindStatic(&UNavigationSystemV1::UpdateComponentInNavOctree);
 			UNavigationSystemBase::UpdateComponentDataAfterMoveDelegate().BindLambda([](USceneComponent& Comp) { UNavigationSystemV1::UpdateNavOctreeAfterMove(&Comp); });
@@ -627,7 +631,7 @@ UNavigationSystemV1::UNavigationSystemV1(const FObjectInitializer& ObjectInitial
 				UNavigationSystemV1* NavSys = FNavigationSystem::GetCurrent<UNavigationSystemV1>(Comp.GetWorld());
 				if (NavSys)
 				{
-					NavSys->UpdateNavOctreeElementBounds(&Comp, NewBounds, DirtyArea);
+					NavSys->UpdateNavOctreeElementBounds(Comp, NewBounds, DirtyArea);
 				}
 			});
 			//UNavigationSystemBase::GetNavDataForPropsDelegate();
@@ -926,7 +930,7 @@ bool UNavigationSystemV1::ConditionalPopulateNavOctree()
 			if (bStoreNavGeometry)
 			{
 #if WITH_RECAST
-				DefaultOctreeController.NavOctree->ComponentExportDelegate = FNavigationOctree::FNavigableGeometryComponentExportDelegate::CreateStatic(&FRecastNavMeshGenerator::ExportComponentGeometry);
+				DefaultOctreeController.NavOctree->NavRelevantGeometryExportDelegate = FNavigationOctree::FNavRelevantGeometryExportDelegate::CreateStatic(&FRecastNavMeshGenerator::ExportNavRelevantObjectGeometry);
 #endif // WITH_RECAST
 			}
 
@@ -3074,6 +3078,24 @@ ANavigationData* UNavigationSystemV1::GetNavDataWithID(const uint16 NavDataID) c
 	return NULL;
 }
 
+void UNavigationSystemV1::OnNavRelevantObjectRegistered(UObject& Object)
+{
+	if (IsNavigationSystemStatic())
+	{
+		return;
+	}
+
+	SCOPE_CYCLE_COUNTER(STAT_DebugNavOctree);
+	if (INavRelevantInterface* NavInterface = Cast<INavRelevantInterface>(&Object))
+	{
+		UWorld* World = Object.GetTypedOuter<UWorld>();
+		if (UNavigationSystemV1* NavSys = FNavigationSystem::GetCurrent<UNavigationSystemV1>(World))
+		{
+			NavSys->RegisterNavOctreeElement(&Object, NavInterface, FNavigationOctreeController::OctreeUpdate_Default);
+		}
+	}
+}
+
 void UNavigationSystemV1::RegisterComponentToNavOctree(UActorComponent* Comp)
 {
 	if ((Comp == nullptr) || IsNavigationSystemStatic())
@@ -3097,6 +3119,24 @@ void UNavigationSystemV1::RegisterComponentToNavOctree(UActorComponent* Comp)
 	}
 }
 
+void UNavigationSystemV1::OnNavRelevantObjectUnregistered(UObject& Object)
+{
+	if (IsNavigationSystemStatic())
+	{
+		return;
+	}
+
+	SCOPE_CYCLE_COUNTER(STAT_DebugNavOctree);
+	if (INavRelevantInterface* NavInterface = Cast<INavRelevantInterface>(&Object))
+	{
+		UWorld* World = Object.GetTypedOuter<UWorld>();
+		if (UNavigationSystemV1* NavSys = FNavigationSystem::GetCurrent<UNavigationSystemV1>(World))
+		{
+			NavSys->UnregisterNavOctreeElement(&Object, NavInterface, FNavigationOctreeController::OctreeUpdate_Default);
+		}
+	}
+}
+
 void UNavigationSystemV1::UnregisterComponentToNavOctree(UActorComponent* Comp)
 {
 	if ((Comp == nullptr) || IsNavigationSystemStatic())
@@ -3200,6 +3240,26 @@ FNavigationRelevantData* UNavigationSystemV1::GetMutableDataForObject(const UObj
 	return DefaultOctreeController.GetMutableDataForObject(Object);
 }
 
+void UNavigationSystemV1::UpdateNavRelevantObjectInNavOctree(UObject& Object)
+{
+	SCOPE_CYCLE_COUNTER(STAT_DebugNavOctree);
+
+	if (INavRelevantInterface* NavElement = Cast<INavRelevantInterface>(&Object))
+	{
+		if (UNavigationSystemV1* NavSys = FNavigationSystem::GetCurrent<UNavigationSystemV1>(Object.GetTypedOuter<UWorld>()))
+		{
+			if (NavElement->IsNavigationRelevant())
+			{
+				NavSys->UpdateNavOctreeElement(&Object, NavElement, FNavigationOctreeController::OctreeUpdate_Default);	
+			}
+			else
+			{
+				NavSys->UnregisterNavOctreeElement(&Object, NavElement, FNavigationOctreeController::OctreeUpdate_Default);
+			}
+		}
+	}
+}
+
 void UNavigationSystemV1::UpdateActorInNavOctree(AActor& Actor)
 {
 	if (IsNavigationSystemStatic())
@@ -3406,6 +3466,11 @@ bool UNavigationSystemV1::UpdateNavOctreeElementBounds(UActorComponent* Comp, co
 		: false;
 }
 
+bool UNavigationSystemV1::UpdateNavOctreeElementBounds(UObject& Object, const FBox& NewBounds, const FBox& DirtyArea)
+{
+	return FNavigationDataHandler(DefaultOctreeController, DefaultDirtyAreasController).UpdateNavOctreeElementBounds(Object, NewBounds, DirtyArea);
+}
+
 bool UNavigationSystemV1::ReplaceAreaInOctreeData(const UObject& Object, TSubclassOf<UNavArea> OldArea, TSubclassOf<UNavArea> NewArea, bool bReplaceChildClasses)
 {
 	return FNavigationDataHandler(DefaultOctreeController, DefaultDirtyAreasController).ReplaceAreaInOctreeData(Object, OldArea, NewArea, bReplaceChildClasses);
diff --git a/Engine/Source/Runtime/NavigationSystem/Public/NavMesh/RecastNavMeshGenerator.h b/Engine/Source/Runtime/NavigationSystem/Public/NavMesh/RecastNavMeshGenerator.h
--- a/Engine/Source/Runtime/NavigationSystem/Public/NavMesh/RecastNavMeshGenerator.h
+++ b/Engine/Source/Runtime/NavigationSystem/Public/NavMesh/RecastNavMeshGenerator.h
@@ -760,14 +760,13 @@ public:
 	NAVIGATIONSYSTEM_API virtual void GrabDebugSnapshot(struct FVisualLogEntry* Snapshot, const FBox& BoundingBox, const FName& CategoryName, ELogVerbosity::Type Verbosity) const override;
 #endif
 
-	/** 
-	 *	@param Actor is a reference to make callee responsible for assuring it's valid
-	 */
-	static NAVIGATIONSYSTEM_API void ExportComponentGeometry(UActorComponent* Component, FNavigationRelevantData& Data);
-	static NAVIGATIONSYSTEM_API void ExportVertexSoupGeometry(const TArray<FVector>& Verts, FNavigationRelevantData& Data);
+	UE_DEPRECATED(5.4, "Use ExportNavRelevantObjectGeometry")
+	static NAVIGATIONSYSTEM_API void ExportComponentGeometry(UActorComponent* InOutComponent, FNavigationRelevantData& OutData);
+	static NAVIGATIONSYSTEM_API void ExportNavRelevantObjectGeometry(INavRelevantInterface& InOutNavRelevantInterface, FNavigationRelevantData& OutData);
+	static NAVIGATIONSYSTEM_API void ExportVertexSoupGeometry(const TArray<FVector>& InVerts, FNavigationRelevantData& OutData);
 
-	static NAVIGATIONSYSTEM_API void ExportRigidBodyGeometry(UBodySetup& BodySetup, TNavStatArray<FVector>& OutVertexBuffer, TNavStatArray<int32>& OutIndexBuffer, const FTransform& LocalToWorld = FTransform::Identity);
-	static NAVIGATIONSYSTEM_API void ExportRigidBodyGeometry(UBodySetup& BodySetup, TNavStatArray<FVector>& OutTriMeshVertexBuffer, TNavStatArray<int32>& OutTriMeshIndexBuffer, TNavStatArray<FVector>& OutConvexVertexBuffer, TNavStatArray<int32>& OutConvexIndexBuffer, TNavStatArray<int32>& OutShapeBuffer, const FTransform& LocalToWorld = FTransform::Identity);
+	static NAVIGATIONSYSTEM_API void ExportRigidBodyGeometry(UBodySetup& InOutBodySetup, TNavStatArray<FVector>& OutVertexBuffer, TNavStatArray<int32>& OutIndexBuffer, const FTransform& LocalToWorld = FTransform::Identity);
+	static NAVIGATIONSYSTEM_API void ExportRigidBodyGeometry(UBodySetup& InOutBodySetup, TNavStatArray<FVector>& OutTriMeshVertexBuffer, TNavStatArray<int32>& OutTriMeshIndexBuffer, TNavStatArray<FVector>& OutConvexVertexBuffer, TNavStatArray<int32>& OutConvexIndexBuffer, TNavStatArray<int32>& OutShapeBuffer, const FTransform& LocalToWorld = FTransform::Identity);
 	static NAVIGATIONSYSTEM_API void ExportAggregatedGeometry(const FKAggregateGeom& AggGeom, TNavStatArray<FVector>& OutConvexVertexBuffer, TNavStatArray<int32>& OutConvexIndexBuffer, TNavStatArray<int32>& OutShapeBuffer, const FTransform& LocalToWorld = FTransform::Identity);
 
 #if UE_ENABLE_DEBUG_DRAWING
diff --git a/Engine/Source/Runtime/NavigationSystem/Public/NavigationDataHandler.h b/Engine/Source/Runtime/NavigationSystem/Public/NavigationDataHandler.h
--- a/Engine/Source/Runtime/NavigationSystem/Public/NavigationDataHandler.h
+++ b/Engine/Source/Runtime/NavigationSystem/Public/NavigationDataHandler.h
@@ -30,7 +30,7 @@ struct FNavigationDataHandler
 	NAVIGATIONSYSTEM_API bool UnregisterNavOctreeElement(UObject& ElementOwner, INavRelevantInterface& ElementInterface, int32 UpdateFlags);
 	NAVIGATIONSYSTEM_API void UpdateNavOctreeElement(UObject& ElementOwner, INavRelevantInterface& ElementInterface, int32 UpdateFlags);
 	NAVIGATIONSYSTEM_API void UpdateNavOctreeParentChain(UObject& ElementOwner, bool bSkipElementOwnerUpdate = false);
-	NAVIGATIONSYSTEM_API bool UpdateNavOctreeElementBounds(UActorComponent& Comp, const FBox& NewBounds, const FBox& DirtyArea);	
+	NAVIGATIONSYSTEM_API bool UpdateNavOctreeElementBounds(UObject& Object, const FBox& NewBounds, const FBox& DirtyArea);
 	NAVIGATIONSYSTEM_API void FindElementsInNavOctree(const FBox& QueryBox, const FNavigationOctreeFilter& Filter, TArray<FNavigationOctreeElement>& Elements);
 	NAVIGATIONSYSTEM_API bool ReplaceAreaInOctreeData(const UObject& Object, TSubclassOf<UNavArea> OldArea, TSubclassOf<UNavArea> NewArea, bool bReplaceChildClasses);
 	NAVIGATIONSYSTEM_API void AddLevelCollisionToOctree(ULevel& Level);
diff --git a/Engine/Source/Runtime/NavigationSystem/Public/NavigationOctree.h b/Engine/Source/Runtime/NavigationSystem/Public/NavigationOctree.h
--- a/Engine/Source/Runtime/NavigationSystem/Public/NavigationOctree.h
+++ b/Engine/Source/Runtime/NavigationSystem/Public/NavigationOctree.h
@@ -142,8 +142,15 @@ struct FNavigationOctreeSemantics
 class FNavigationOctree : public TOctree2<FNavigationOctreeElement, FNavigationOctreeSemantics>, public TSharedFromThis<FNavigationOctree, ESPMode::ThreadSafe>
 {
 public:
+	UE_DEPRECATED(5.4, "Use FNavRelevantGeometryExportDelegate.")
 	DECLARE_DELEGATE_TwoParams(FNavigableGeometryComponentExportDelegate, UActorComponent*, FNavigationRelevantData&);
+
+PRAGMA_DISABLE_DEPRECATION_WARNINGS
 	FNavigableGeometryComponentExportDelegate ComponentExportDelegate;
+PRAGMA_ENABLE_DEPRECATION_WARNINGS
+
+	DECLARE_DELEGATE_TwoParams(FNavRelevantGeometryExportDelegate, INavRelevantInterface&, FNavigationRelevantData&);
+	FNavRelevantGeometryExportDelegate NavRelevantGeometryExportDelegate;
 
 	enum ENavGeometryStoringMode {
 		SkipNavGeometry,
diff --git a/Engine/Source/Runtime/NavigationSystem/Public/NavigationSystem.h b/Engine/Source/Runtime/NavigationSystem/Public/NavigationSystem.h
--- a/Engine/Source/Runtime/NavigationSystem/Public/NavigationSystem.h
+++ b/Engine/Source/Runtime/NavigationSystem/Public/NavigationSystem.h
@@ -880,6 +880,10 @@ public:
 	//----------------------------------------------------------------------//
 	// navigation octree related functions
 	//----------------------------------------------------------------------//
+	static NAVIGATIONSYSTEM_API void OnNavRelevantObjectRegistered(UObject& Object);
+	static NAVIGATIONSYSTEM_API void UpdateNavRelevantObjectInNavOctree(UObject& Object);
+	static NAVIGATIONSYSTEM_API void OnNavRelevantObjectUnregistered(UObject& Object);
+
 	static NAVIGATIONSYSTEM_API void OnComponentRegistered(UActorComponent* Comp);
 	static NAVIGATIONSYSTEM_API void OnComponentUnregistered(UActorComponent* Comp);
 	static NAVIGATIONSYSTEM_API void RegisterComponent(UActorComponent* Comp);
@@ -945,9 +949,13 @@ public:
 	/** force updating parent node and all its children */
 	NAVIGATIONSYSTEM_API void UpdateNavOctreeParentChain(UObject* ElementOwner, bool bSkipElementOwnerUpdate = false);
 
+	UE_DEPRECATED(5.4, "Use UpdateNavOctreeElementBounds with object reference.")
 	/** update component bounds in navigation octree and mark only specified area as dirty, doesn't re-export component geometry */
 	NAVIGATIONSYSTEM_API bool UpdateNavOctreeElementBounds(UActorComponent* Comp, const FBox& NewBounds, const FBox& DirtyArea);
 
+	/** update object bounds in navigation octree and mark only specified area as dirty, doesn't re-export geometry */
+	NAVIGATIONSYSTEM_API bool UpdateNavOctreeElementBounds(UObject& Object, const FBox& NewBounds, const FBox& DirtyArea);
+
 	/** fetched Object's data from the octree and replaces occurences of OldArea with NewArea */
 	NAVIGATIONSYSTEM_API bool ReplaceAreaInOctreeData(const UObject& Object, TSubclassOf<UNavArea> OldArea, TSubclassOf<UNavArea> NewArea, bool bReplaceChildClasses = false);
 
@@ -1245,7 +1253,7 @@ protected:
 
 	FDelegateHandle ReloadCompleteDelegateHandle;
 
-	/** called to notify NavigaitonSystem about finished reload */
+	/** called to notify NavigationSystem about finished reload */
 	NAVIGATIONSYSTEM_API virtual void OnReloadComplete(EReloadCompleteReason Reason);
 
 	/** Registers given navigation data with this Navigation System.
