Fixed navmesh dirty area logs to display bound sizes as indicated and Improved log readability.

diff --git a/Engine/Source/Runtime/NavigationSystem/Private/NavMesh/RecastNavMeshGenerator.cpp b/Engine/Source/Runtime/NavigationSystem/Private/NavMesh/RecastNavMeshGenerator.cpp
--- a/Engine/Source/Runtime/NavigationSystem/Private/NavMesh/RecastNavMeshGenerator.cpp
+++ b/Engine/Source/Runtime/NavigationSystem/Private/NavMesh/RecastNavMeshGenerator.cpp
@@ -5651,11 +5651,12 @@ void FRecastNavMeshGenerator::LogDirtyAreas(
 		const UObject* const SourceObject = DirtyResultsTuple.DirtyArea.OptionalSourceObject.Get();
 		const UActorComponent* const ObjectAsComponent = Cast<UActorComponent>(SourceObject);
 		const AActor* const ComponentOwner = ObjectAsComponent ? ObjectAsComponent->GetOwner() : nullptr;
+		const FVector2D BoundsSize(DirtyResultsTuple.DirtyArea.Bounds.GetSize());
 		
 		UE_LOG(LogNavigationDirtyArea, VeryVerbose,
-			TEXT("(navmesh: %s) Dirty area trying to dirt %d tiles (out of which %d are newly added/not pending) | Source Object = %s | Potential component's owner = %s | Bounds size = %s)"),
+			TEXT("(navmesh: %-30s) Dirty area trying to dirt %2d tiles (out of which %2d are newly added/not pending) | Source Object = %s | Potential component's owner = %s | Bounds size = %s)"),
 			*GetNameSafe(GetOwner()), DirtyResultsTuple.TotalDirtyTiles, DirtyResultsTuple.NewlyAddedDirtyTiles, *GetFullNameSafe(SourceObject),
-			*GetFullNameSafe(ComponentOwner), *DirtyResultsTuple.DirtyArea.Bounds.ToString());
+			*GetFullNameSafe(ComponentOwner), *BoundsSize.ToString());
 	}
 }
 #endif
@@ -6092,7 +6093,7 @@ void FRecastNavMeshGenerator::MarkDirtyTiles(const TArray<FNavigationDirtyArea>&
 					const FVector2D BoundsSize(DirtyArea.Bounds.GetSize());
 			
 					UE_LOG(LogNavigationDirtyArea, Warning,
-						TEXT("(navmesh: %s) Added an oversized dirty area | Tiles marked: %u | Source object = %s | Potential comp owner = %s | Bounds size = %s | Threshold: %.0f"),
+						TEXT("(navmesh: %-30s) Added an oversized dirty area | Tiles marked: %2u | Source object = %s | Potential comp owner = %s | Bounds size = %s | Threshold: %.0f"),
 						*GetNameSafe(GetOwner()), PendingTilesMarked, *GetFullNameSafe(SourceObject), *GetFullNameSafe(ComponentOwner),
 						*BoundsSize.ToString(), NavSys->GetDirtyAreaWarningSizeThreshold());
 				}
