Improve FRecastNavMeshGenerator::MarkDirtyTiles speed by prefiltering dirty area bounds using the bounds of invokers seeds.
Also added some profile tags in MarkDirtyTiles.

diff --git a/Engine/Source/Runtime/NavigationSystem/Private/NavMesh/RecastNavMeshGenerator.cpp b/Engine/Source/Runtime/NavigationSystem/Private/NavMesh/RecastNavMeshGenerator.cpp
--- a/Engine/Source/Runtime/NavigationSystem/Private/NavMesh/RecastNavMeshGenerator.cpp
+++ b/Engine/Source/Runtime/NavigationSystem/Private/NavMesh/RecastNavMeshGenerator.cpp
@@ -5937,10 +5937,16 @@ void FRecastNavMeshGenerator::MarkDirtyTiles(const TArray<FNavigationDirtyArea>&
 
 	const bool bGameStaticNavMesh = IsGameStaticNavMesh(DestNavMesh);
 
+	const UNavigationSystemV1* NavSys = FNavigationSystem::GetCurrent<UNavigationSystemV1>(GetWorld());
+	const TArray<FBox>* SeedsBoundsArrayPtr = nullptr;
+	if (NavSys && NavSys->IsActiveTilesGenerationEnabled())
+	{
+		SeedsBoundsArrayPtr = &NavSys->GetInvokersSeedBounds();
+	}
+
 	const ARecastNavMesh* const OwnerNav = GetOwner();
 	const bool bUseVirtualGeometryFilteringAndDirtying = OwnerNav != nullptr && OwnerNav->bUseVirtualGeometryFilteringAndDirtying;
 	// Those are set only if bUseVirtualGeometryFilteringAndDirtying is enabled since we do not use them for anything else
-	const UNavigationSystemV1* NavSys = bUseVirtualGeometryFilteringAndDirtying ? FNavigationSystem::GetCurrent<UNavigationSystemV1>(GetWorld()) : nullptr;
 	const FNavigationOctree* const NavOctreeInstance = (bUseVirtualGeometryFilteringAndDirtying && NavSys) ? NavSys->GetNavOctree() : nullptr;
 	const FNavDataConfig* const NavDataConfig = bUseVirtualGeometryFilteringAndDirtying ? &DestNavMesh->GetConfig() : nullptr;
 	// ~
@@ -6003,139 +6009,167 @@ void FRecastNavMeshGenerator::MarkDirtyTiles(const TArray<FNavigationDirtyArea>&
 			bDoTileInclusionTest = (FindInclusionBoundEncapsulatingBox(CutDownArea) == INDEX_NONE);
 		}
 
-		uint32 PendingTilesMarked = 0;
-		const FRcTileBox TileBox(AdjustedAreaBounds, RcNavMeshOrigin, TileSizeInWorldUnits);
+		QUICK_SCOPE_CYCLE_COUNTER(STAT_RecastNavMeshGenerator_DirtyArea);
 
-		for (int32 TileY = TileBox.YMin; TileY <= TileBox.YMax; ++TileY)
+		TArray<FBox, TInlineAllocator<32>> SubAreaBoundsArray;
+		if (SeedsBoundsArrayPtr != nullptr && SeedsBoundsArrayPtr->Num() > 0)
 		{
-			for (int32 TileX = TileBox.XMin; TileX <= TileBox.XMax; ++TileX)
+			for (const FBox& SeedBounds : *SeedsBoundsArrayPtr)
 			{
-				if (IsInActiveSet(FIntPoint(TileX, TileY)) == false)
+				// Compute sub area bound
+				const FBox OverlapBox = AdjustedAreaBounds.Overlap(SeedBounds);
+				if (OverlapBox.IsValid)
 				{
-					UE_SUPPRESS(LogNavigation, VeryVerbose,
-					{
-						const FBox TileBounds = FRecastTileGenerator::CalculateTileBounds(TileX, TileY, RcNavMeshOrigin, TotalNavBounds, TileSizeInWorldUnits);
-						UE_VLOG_BOX(OwnerNav, LogNavigation, VeryVerbose, TileBounds, FColor::Red, TEXT("Not in active set"));
-					});
-					continue;
+					SubAreaBoundsArray.Add(OverlapBox);
 				}
+			}
+		}
+		else
+		{
+			SubAreaBoundsArray.Add(AdjustedAreaBounds);
+		}
 
-				if (bDoTileInclusionTest == true && DirtyArea.HasFlag(ENavigationDirtyFlag::NavigationBounds) == false)
-				{
-					const FBox TileBounds = FRecastTileGenerator::CalculateTileBounds(TileX, TileY, RcNavMeshOrigin, TotalNavBounds, TileSizeInWorldUnits);
+		uint32 PendingTilesMarked = 0;
+		
+		for (const FBox& SubArea : SubAreaBoundsArray)
+		{
+			const FRcTileBox TileBox(SubArea, RcNavMeshOrigin, TileSizeInWorldUnits);
 
-					// do per tile check since we can have lots of tiles inbetween navigable bounds volumes
-					if (IntersectBounds(TileBounds, InclusionBounds) == false)
+			for (int32 TileY = TileBox.YMin; TileY <= TileBox.YMax; ++TileY)
+			{
+				for (int32 TileX = TileBox.XMin; TileX <= TileBox.XMax; ++TileX)
+				{
+					if (IsInActiveSet(FIntPoint(TileX, TileY)) == false)
 					{
-						// Skip this tile
+						UE_SUPPRESS(LogNavigation, VeryVerbose,
+						{
+							const FBox TileBounds = FRecastTileGenerator::CalculateTileBounds(TileX, TileY, RcNavMeshOrigin, TotalNavBounds, TileSizeInWorldUnits);
+							UE_VLOG_BOX(OwnerNav, LogNavigation, VeryVerbose, TileBounds, FColor::Red, TEXT("Not in active set"));
+						});
 						continue;
 					}
-				}
-												
-				FPendingTileElement Element;
-				Element.Coord = FIntPoint(TileX, TileY);
-				Element.bRebuildGeometry = DirtyArea.HasFlag(ENavigationDirtyFlag::Geometry) || DirtyArea.HasFlag(ENavigationDirtyFlag::NavigationBounds);
-				Element.CreationTime = CurrentTimeSeconds;
-				if (Element.bRebuildGeometry == false)
-				{
-					Element.DirtyAreas.Add(AdjustedAreaBounds);
-				}
-				PendingTilesMarked++;
-				
-				FPendingTileElement* ExistingElement = DirtyTiles.Find(Element);
-				if (ExistingElement)
-				{
-					ExistingElement->bRebuildGeometry |= Element.bRebuildGeometry;
-					// Append area bounds to existing list 
-					if (ExistingElement->bRebuildGeometry == false)
+
+					if (bDoTileInclusionTest == true && DirtyArea.HasFlag(ENavigationDirtyFlag::NavigationBounds) == false)
+					{
+						const FBox TileBounds = FRecastTileGenerator::CalculateTileBounds(TileX, TileY, RcNavMeshOrigin, TotalNavBounds, TileSizeInWorldUnits);
+
+						// do per tile check since we can have lots of tiles in between navigable bounds volumes
+						if (IntersectBounds(TileBounds, InclusionBounds) == false)
+						{
+							// Skip this tile
+							continue;
+						}
+					}
+											
+					FPendingTileElement Element;
+					Element.Coord = FIntPoint(TileX, TileY);
+					Element.bRebuildGeometry = DirtyArea.HasFlag(ENavigationDirtyFlag::Geometry) || DirtyArea.HasFlag(ENavigationDirtyFlag::NavigationBounds);
+					Element.CreationTime = CurrentTimeSeconds;
+					if (Element.bRebuildGeometry == false)
 					{
-						ExistingElement->DirtyAreas.Append(Element.DirtyAreas);
+						Element.DirtyAreas.Add(SubArea);
+					}
+					PendingTilesMarked++;
+			
+					FPendingTileElement* ExistingElement = DirtyTiles.Find(Element);
+					if (ExistingElement)
+					{
+						ExistingElement->bRebuildGeometry |= Element.bRebuildGeometry;
+						// Append area bounds to existing list 
+						if (ExistingElement->bRebuildGeometry == false)
+						{
+							ExistingElement->DirtyAreas.Append(Element.DirtyAreas);
+						}
+						else
+						{
+							ExistingElement->DirtyAreas.Empty();
+						}
 					}
 					else
 					{
-						ExistingElement->DirtyAreas.Empty();
+						DirtyTiles.Add(Element);
 					}
-				}
-				else
-				{
-					DirtyTiles.Add(Element);
-				}
-				
+			
 #if !UE_BUILD_SHIPPING
-				UE_SUPPRESS(LogNavigationDirtyArea, VeryVerbose, 
-				{
-					const bool bAlreadyAdded = ExistingElement != nullptr;
-					DirtyAreasDebugging.FindOrAdd(Element).Add({DirtyArea, bAlreadyAdded});
-				});
+					UE_SUPPRESS(LogNavigationDirtyArea, VeryVerbose, 
+					{
+						const bool bAlreadyAdded = ExistingElement != nullptr;
+						DirtyAreasDebugging.FindOrAdd(Element).Add({DirtyArea, bAlreadyAdded});
+					});
 #endif
+				}
 			}
-		}
 
-#if !UE_BUILD_SHIPPING		
-		// Warn if this is from a big dirty area
-		UE_SUPPRESS(LogNavigationDirtyArea, Warning, 
-		{
-			if (PendingTilesMarked > 0)
+#if !UE_BUILD_SHIPPING
+			// Warn if this is from a big dirty area
+			UE_SUPPRESS(LogNavigationDirtyArea, Warning, 
 			{
-				if (NavSys == nullptr)
+				if (PendingTilesMarked > 0)
 				{
-					// NavSys might not have been initialized yet if not using bUseVirtualGeometryFilteringAndDirtying
-					NavSys = FNavigationSystem::GetCurrent<UNavigationSystemV1>(GetWorld());
-				}
-				
-				// If not using active tile generation, those are reported earlier in FNavigationDirtyAreasController::AddArea
-				if (NavSys && NavSys->GetOperationMode() == FNavigationSystemRunMode::GameMode && NavSys->IsActiveTilesGenerationEnabled() &&
-					AdjustedAreaBounds.GetSize().GetMax() > NavSys->GetDirtyAreaWarningSizeThreshold())
-				{
-					const UObject* const SourceObject = DirtyArea.OptionalSourceObject.Get();
-					const UActorComponent* const ObjectAsComponent = Cast<UActorComponent>(SourceObject);
-					const AActor* const ComponentOwner = ObjectAsComponent ? ObjectAsComponent->GetOwner() : nullptr;
-					const FVector2D BoundsSize(DirtyArea.Bounds.GetSize());
+					if (NavSys == nullptr)
+					{
+						// NavSys might not have been initialized yet if not using bUseVirtualGeometryFilteringAndDirtying
+						NavSys = FNavigationSystem::GetCurrent<UNavigationSystemV1>(GetWorld());
+					}
 			
-					UE_LOG(LogNavigationDirtyArea, Warning,
-						TEXT("(navmesh: %-30s) Added an oversized dirty area | Tiles marked: %2u | Source object = %s | Potential comp owner = %s | Bounds size = %s | Threshold: %.0f"),
-						*GetNameSafe(GetOwner()), PendingTilesMarked, *GetFullNameSafe(SourceObject), *GetFullNameSafe(ComponentOwner),
-						*BoundsSize.ToString(), NavSys->GetDirtyAreaWarningSizeThreshold());
+					// If not using active tile generation, those are reported earlier in FNavigationDirtyAreasController::AddArea
+					if (NavSys && NavSys->GetOperationMode() == FNavigationSystemRunMode::GameMode && NavSys->IsActiveTilesGenerationEnabled() &&
+						AdjustedAreaBounds.GetSize().GetMax() > NavSys->GetDirtyAreaWarningSizeThreshold())
+					{
+						const UObject* const SourceObject = DirtyArea.OptionalSourceObject.Get();
+						const UActorComponent* const ObjectAsComponent = Cast<UActorComponent>(SourceObject);
+						const AActor* const ComponentOwner = ObjectAsComponent ? ObjectAsComponent->GetOwner() : nullptr;
+						const FVector2D BoundsSize(DirtyArea.Bounds.GetSize());
+		
+						UE_LOG(LogNavigationDirtyArea, Warning,
+							TEXT("(navmesh: %-30s) Added an oversized dirty area | Tiles marked: %2u | Source object = %s | Potential comp owner = %s | Bounds size = %s | Threshold: %.0f"),
+							*GetNameSafe(GetOwner()), PendingTilesMarked, *GetFullNameSafe(SourceObject), *GetFullNameSafe(ComponentOwner),
+							*BoundsSize.ToString(), NavSys->GetDirtyAreaWarningSizeThreshold());
+					}
 				}
-			}
-		});
+			});
 #endif // !UE_BUILD_SHIPPING
-		
+			
+		} // SubAreaBoundsArray loop	
 	}
 	
 	int32 NumTilesMarked = DirtyTiles.Num();
 
-	// Merge new dirty tiles info with existing pending elements
-	for (FPendingTileElement& ExistingElement : PendingDirtyTiles)
 	{
-		FSetElementId Id = DirtyTiles.FindId(ExistingElement);
-		if (Id.IsValidId())
+		QUICK_SCOPE_CYCLE_COUNTER(STAT_RecastNavMeshGenerator_Merging);
+		
+		// Merge new dirty tiles info with existing pending elements
+		for (FPendingTileElement& ExistingElement : PendingDirtyTiles)
 		{
-			const FPendingTileElement& DirtyElement = DirtyTiles[Id];
-
-			ExistingElement.bRebuildGeometry |= DirtyElement.bRebuildGeometry;
-			ExistingElement.CreationTime = FMath::Min(DirtyElement.CreationTime, ExistingElement.CreationTime);
-			// Append area bounds to existing list 
-			if (ExistingElement.bRebuildGeometry == false)
-			{
-				ExistingElement.DirtyAreas.Append(DirtyElement.DirtyAreas);
-			}
-			else
+			FSetElementId Id = DirtyTiles.FindId(ExistingElement);
+			if (Id.IsValidId())
 			{
-				ExistingElement.DirtyAreas.Empty();
-			}
-			DirtyTiles.Remove(Id);
+				const FPendingTileElement& DirtyElement = DirtyTiles[Id];
 
-#if !UE_BUILD_SHIPPING
-			UE_SUPPRESS(LogNavigationDirtyArea, VeryVerbose, 
-			{
-				// Flag everything in the map to set a true boolean (since it was already inserted)
-				for (FNavigationDirtyAreaPerTileDebugInformation& Pair : DirtyAreasDebugging.FindChecked(ExistingElement))
+				ExistingElement.bRebuildGeometry |= DirtyElement.bRebuildGeometry;
+				ExistingElement.CreationTime = FMath::Min(DirtyElement.CreationTime, ExistingElement.CreationTime);
+				// Append area bounds to existing list 
+				if (ExistingElement.bRebuildGeometry == false)
 				{
-					Pair.bTileWasAlreadyAdded = true;
+					ExistingElement.DirtyAreas.Append(DirtyElement.DirtyAreas);
 				}
-			});
+				else
+				{
+					ExistingElement.DirtyAreas.Empty();
+				}
+				DirtyTiles.Remove(Id);
+
+#if !UE_BUILD_SHIPPING
+				UE_SUPPRESS(LogNavigationDirtyArea, VeryVerbose, 
+				{
+					// Flag everything in the map to set a true boolean (since it was already inserted)
+					for (FNavigationDirtyAreaPerTileDebugInformation& Pair : DirtyAreasDebugging.FindChecked(ExistingElement))
+					{
+						Pair.bTileWasAlreadyAdded = true;
+					}
+				});
 #endif
+			}
 		}
 	}
 
@@ -6170,6 +6204,8 @@ bool FRecastNavMeshGenerator::ShouldDirtyTilesRequestedByObject(const UNavigatio
 
 void FRecastNavMeshGenerator::SortPendingBuildTiles()
 {
+	QUICK_SCOPE_CYCLE_COUNTER(STAT_RecastNavMeshGenerator_SortPendingBuildTiles);
+	
 	if (SortPendingTilesMethod == ENavigationSortPendingTilesMethod::SortWithSeedLocations)
 	{
 		UWorld* CurWorld = GetWorld();
diff --git a/Engine/Source/Runtime/NavigationSystem/Private/NavigationSystem.cpp b/Engine/Source/Runtime/NavigationSystem/Private/NavigationSystem.cpp
--- a/Engine/Source/Runtime/NavigationSystem/Private/NavigationSystem.cpp
+++ b/Engine/Source/Runtime/NavigationSystem/Private/NavigationSystem.cpp
@@ -5071,6 +5071,7 @@ void UNavigationSystemV1::UpdateInvokers()
 	if (CurrentTime >= NextInvokersUpdateTime)
 	{
 		InvokerLocations.Reset();
+		InvokersSeedBounds.Reset();
 
 		if (Invokers.Num() > 0)
 		{
@@ -5081,6 +5082,14 @@ void UNavigationSystemV1::UpdateInvokers()
 			if (bCheckMaximumDistanceFromSeeds)
 			{
 				GetInvokerSeedLocations(*World, SeedLocations);
+
+				// Fill seed bounds
+				for (const FVector2D SeedLocation : SeedLocations)
+				{
+					InvokersSeedBounds.Emplace(
+						FVector(SeedLocation.X-InvokersMaximumDistanceFromSeed, SeedLocation.Y-InvokersMaximumDistanceFromSeed, DBL_MIN),
+						FVector(SeedLocation.X+InvokersMaximumDistanceFromSeed, SeedLocation.Y+InvokersMaximumDistanceFromSeed, DBL_MAX));
+				}
 			}
 			
 			const double StartTime = FPlatformTime::Seconds();
diff --git a/Engine/Source/Runtime/NavigationSystem/Public/NavigationSystem.h b/Engine/Source/Runtime/NavigationSystem/Public/NavigationSystem.h
--- a/Engine/Source/Runtime/NavigationSystem/Public/NavigationSystem.h
+++ b/Engine/Source/Runtime/NavigationSystem/Public/NavigationSystem.h
@@ -455,6 +455,8 @@ private:
 	TMap<AActor*, FNavigationInvoker> Invokers;
 	/** Contains pre-digested and cached invokers' info. Generated by UpdateInvokers */
 	TArray<FNavigationInvokerRaw> InvokerLocations;
+	
+	TArray<FBox> InvokersSeedBounds;
 
 	double NextInvokersUpdateTime;
 	NAVIGATIONSYSTEM_API void UpdateInvokers();
@@ -819,6 +821,7 @@ public:
 	static NAVIGATIONSYSTEM_API void UnregisterNavigationInvoker(AActor& Invoker);
 
 	const TArray<FNavigationInvokerRaw>& GetInvokerLocations() const { return InvokerLocations; }
+	const TArray<FBox>& GetInvokersSeedBounds() const { return InvokersSeedBounds; }
 
 	//----------------------------------------------------------------------//
 	// Bookkeeping 
