Log invoker priority as string instead of the enum value

diff --git a/Engine/Source/Runtime/NavigationSystem/Private/NavigationSystem.cpp b/Engine/Source/Runtime/NavigationSystem/Private/NavigationSystem.cpp
--- a/Engine/Source/Runtime/NavigationSystem/Private/NavigationSystem.cpp
+++ b/Engine/Source/Runtime/NavigationSystem/Private/NavigationSystem.cpp
@@ -5049,8 +5049,8 @@ void UNavigationSystemV1::RegisterInvoker(AActor& Invoker, float TileGenerationR
 			}
 		}
 
-		const FString RegisterText = FString::Printf(TEXT("Register invoker r: %.0f, r area: %.0f m2, removal r: %.0f, priority: %d, (%s %s) "),
-			TileGenerationRadius, UE_PI*FMath::Square(TileGenerationRadius/100.f), TileRemovalRadius, InPriority, *Invoker.GetName(), *InvokerNavData);
+		const FString RegisterText = FString::Printf(TEXT("Register invoker r: %.0f, r area: %.0f m2, removal r: %.0f, priority: %s, (%s %s) "),
+			TileGenerationRadius, UE_PI*FMath::Square(TileGenerationRadius/100.f), TileRemovalRadius, *UEnum::GetDisplayValueAsText(InPriority).ToString(), *Invoker.GetName(), *InvokerNavData);
 		UE_LOG(LogNavInvokers, Log, TEXT("%s"), *RegisterText);
 
 		UE_VLOG_CYLINDER(this, LogNavInvokers, Log, Invoker.GetActorLocation(), Invoker.GetActorLocation() + FVector(0, 0, 20), TileGenerationRadius, FColorList::LimeGreen, TEXT("%s"), *RegisterText);
