[Navigation] Fixed warning log caused by NavMeshBoundsVolume actors gathered before their components get registered resulting in invalid dirty areas.
#rnx

diff --git a/Engine/Source/Runtime/NavigationSystem/Private/NavigationSystem.cpp b/Engine/Source/Runtime/NavigationSystem/Private/NavigationSystem.cpp
--- a/Engine/Source/Runtime/NavigationSystem/Private/NavigationSystem.cpp
+++ b/Engine/Source/Runtime/NavigationSystem/Private/NavigationSystem.cpp
@@ -3772,8 +3772,10 @@ void UNavigationSystemV1::GatherNavigationBounds()
 	RegisteredNavBounds.Empty();
 	for (TActorIterator<ANavMeshBoundsVolume> It(GetWorld()); It; ++It)
 	{
-		ANavMeshBoundsVolume* V = (*It);
-		if (IsValid(V))
+		// Iterator can access actors with unregistered components which can result in invalid bounding boxes.
+		// In this case we skip these actors and wait calls to OnNavigationBoundsAdded.
+		const ANavMeshBoundsVolume* V = (*It);
+		if (IsValid(V) && V->HasActorRegisteredAllComponents())
 		{
 			FNavigationBounds NavBounds;
 			NavBounds.UniqueID = V->GetUniqueID();
