Fix usage of SHOW_NAV_EXPORT_PREVIEW define.

diff --git a/Engine/Source/Runtime/NavigationSystem/Private/NavMesh/RecastNavMeshGenerator.cpp b/Engine/Source/Runtime/NavigationSystem/Private/NavMesh/RecastNavMeshGenerator.cpp
--- a/Engine/Source/Runtime/NavigationSystem/Private/NavMesh/RecastNavMeshGenerator.cpp
+++ b/Engine/Source/Runtime/NavigationSystem/Private/NavMesh/RecastNavMeshGenerator.cpp
@@ -501,7 +501,7 @@ void ExportChaosConvexMesh(const FKConvexElem* const Convex, const FTransform& L
 #if SHOW_NAV_EXPORT_PREVIEW
 	if (DebugWorld)
 	{
-		for (int32 Index = StartVertOffset; Index < VertexBuffer.Num(); Index += 3)
+		for (int32 Index = VertOffset; Index < VertexBuffer.Num(); Index += 3)
 		{
 			FVector V0(VertexBuffer[IndexBuffer[Index] * 3], VertexBuffer[IndexBuffer[Index] * 3 + 1], VertexBuffer[IndexBuffer[Index] * 3] + 2);
 			FVector V1(VertexBuffer[IndexBuffer[Index + 1] * 3], VertexBuffer[IndexBuffer[Index + 1] * 3 + 1], VertexBuffer[IndexBuffer[Index + 1] * 3] + 2);
@@ -509,9 +509,10 @@ void ExportChaosConvexMesh(const FKConvexElem* const Convex, const FTransform& L
 
 			DrawDebugLine(DebugWorld, V0, V1, FColor::Blue, true);
 			DrawDebugLine(DebugWorld, V1, V2, FColor::Blue, true);
-			DrawDebugLine(DebugWorld, V2, V0, fColor::Blue, true);
+			DrawDebugLine(DebugWorld, V2, V0, FColor::Blue, true);
 		}
-	#endif // SHOW_NAV_EXPORT_PREVIEW
+	}
+#endif // SHOW_NAV_EXPORT_PREVIEW
 }
 
 void ExportChaosHeightField(const Chaos::FHeightField* const HeightField, const FTransform& LocalToWorld
