[Navigation]
- Added console command "ai.debug.nav.DrawDistance" to tweak navmesh draw distance at runtime.
- Added labels indicating the number of elements in each node if octree details are enabled.

diff --git a/Engine/Source/Runtime/AIModule/Private/GameplayDebugger/GameplayDebuggerCategory_Navmesh.cpp b/Engine/Source/Runtime/AIModule/Private/GameplayDebugger/GameplayDebuggerCategory_Navmesh.cpp
--- a/Engine/Source/Runtime/AIModule/Private/GameplayDebugger/GameplayDebuggerCategory_Navmesh.cpp
+++ b/Engine/Source/Runtime/AIModule/Private/GameplayDebugger/GameplayDebuggerCategory_Navmesh.cpp
@@ -25,7 +25,9 @@ namespace
 
 		FAutoConsoleVariableRef(TEXT("ai.debug.nav.DisplaySize"), 
 			FGameplayDebuggerCategoryNavmeshTweakables::DisplaySize,
-			TEXT("Area we want to display in tiles (DisplaySize x DisplaySize). Note that size will round up to an odd number of tiles")),
+			TEXT("Area to display in tiles (DisplaySize x DisplaySize) in gameplay debugger."
+				 " Size will round up to an odd number of tiles."
+				 " Culling distance can be modified using 'ai.debug.nav.DrawDistance'.")),
 
 		FAutoConsoleVariableRef(TEXT("ai.debug.nav.RefreshInterval"),
 			FGameplayDebuggerCategoryNavmeshTweakables::RefreshInterval,
diff --git a/Engine/Source/Runtime/NavigationSystem/Private/NavMesh/NavMeshRenderingComponent.cpp b/Engine/Source/Runtime/NavigationSystem/Private/NavMesh/NavMeshRenderingComponent.cpp
--- a/Engine/Source/Runtime/NavigationSystem/Private/NavMesh/NavMeshRenderingComponent.cpp
+++ b/Engine/Source/Runtime/NavigationSystem/Private/NavMesh/NavMeshRenderingComponent.cpp
@@ -1080,11 +1080,16 @@ void FNavMeshSceneProxyData::GatherData(const ARecastNavMesh* NavMesh, int32 InN
 				TArray<uint32> CollidingIndices;
 
 				FBox CurrentNodeBoundsBox;
-				NavOctree->FindElementsWithPredicate([bGatherOctree, bGatherPathCollidingGeometry, this, &CurrentNodeBoundsBox](FNavigationOctree::FNodeIndex /*ParentNodeIndex*/, FNavigationOctree::FNodeIndex /*NodeIndex*/, const FBoxCenterAndExtent& NodeBounds)
+				NavOctree->FindElementsWithPredicate([bGatherOctree, bGatherOctreeDetails, bGatherPathCollidingGeometry, this, &CurrentNodeBoundsBox, NavOctree](FNavigationOctree::FNodeIndex /*ParentNodeIndex*/, FNavigationOctree::FNodeIndex NodeIndex, const FBoxCenterAndExtent& NodeBounds)
 				{
 					if (bGatherOctree)
 					{
 						OctreeBounds.Add(NodeBounds);
+
+						if (bGatherOctreeDetails)
+						{
+							DebugLabels.Emplace(NodeBounds.Center, FString::Printf(TEXT("%d elements"), NavOctree->GetElementsForNode(NodeIndex).Num()));
+						}
 					}
 
 					if (bGatherPathCollidingGeometry)
diff --git a/Engine/Source/Runtime/NavigationSystem/Private/NavMesh/RecastNavMesh.cpp b/Engine/Source/Runtime/NavigationSystem/Private/NavMesh/RecastNavMesh.cpp
--- a/Engine/Source/Runtime/NavigationSystem/Private/NavMesh/RecastNavMesh.cpp
+++ b/Engine/Source/Runtime/NavigationSystem/Private/NavMesh/RecastNavMesh.cpp
@@ -237,6 +237,16 @@ void FNavMeshTileData::MakeUnique()
 	}
 }
 
+static FAutoConsoleCommandWithWorldAndArgs CmdNavMeshDrawDistance(
+	TEXT("ai.debug.nav.DrawDistance"),
+	TEXT("Sets the culling distance used by the navmesh rendering for lines and labels."),
+	FConsoleCommandWithWorldAndArgsDelegate::CreateLambda([](const TArray<FString>& Args, UWorld*)
+		{
+			const FVector::FReal DrawDistance = (Args.Num() == 0) ? GetDefault<ARecastNavMesh>()->DefaultDrawDistance : FCString::Atof(*Args[0]);
+			ARecastNavMesh::SetDrawDistance(DrawDistance);
+		}
+	));
+
 FVector::FReal ARecastNavMesh::DrawDistanceSq = 0.;
 float ARecastNavMesh::MinimumSizeForChaosNavMeshInfluenceSq = 0.0f;
 #if !WITH_RECAST
