added more debug info to help track down cause of conversion bugs

diff --git a/Engine/Source/Runtime/NavigationSystem/Public/NavMesh/RecastNavMeshGenerator.h b/Engine/Source/Runtime/NavigationSystem/Public/NavMesh/RecastNavMeshGenerator.h
--- a/Engine/Source/Runtime/NavigationSystem/Public/NavMesh/RecastNavMeshGenerator.h
+++ b/Engine/Source/Runtime/NavigationSystem/Public/NavMesh/RecastNavMeshGenerator.h
@@ -189,6 +189,9 @@ struct FRcTileBox
 	FRcTileBox(const FBox& UnrealBounds, const FVector& RcNavMeshOrigin, const FVector::FReal TileSizeInWorldUnits)
 	{
 		check(TileSizeInWorldUnits > 0);
+		checkf(!RcNavMeshOrigin.ContainsNaN(), TEXT("%hs: RcNavMeshOrigin ContainsNaN"), __FUNCTION__);
+		checkf(!UnrealBounds.ContainsNaN(), TEXT("%hs: UnrealBounds ContainsNaN"), __FUNCTION__);
+		checkf(UnrealBounds.IsValid, TEXT("%hs: UnrealBounds !IsValid()"), __FUNCTION__);
 
 		auto CalcMaxCoordExclusive = [](const FVector::FReal MaxAsFloat, const int32 MinCoord) -> int32
 		{
