[Navigation] Added bounds to NavCollision since geometry can differ from StaticMesh
- Bumped the DDC version of the NavCollision
- Updated StaticMeshComponent to return bounds from the NavCollision when possible
- Fixed InstancedStaticMeshComponent not using the additional geometry of the NavCollision when exporting the custom navigable geomerty
- Updated InstancedStaticMeshComponent to compute bounds using the NavCollision for GetNavigationBounds, PartialNavigationUpdate and when gathering per instance transforms for navigation
- Updated UHierarchicalInstancedStaticMeshComponent to use bounds from the NavCollision for PartialNavigationUpdate

diff --git a/Engine/Source/Runtime/Engine/Classes/AI/Navigation/NavCollisionBase.h b/Engine/Source/Runtime/Engine/Classes/AI/Navigation/NavCollisionBase.h
--- a/Engine/Source/Runtime/Engine/Classes/AI/Navigation/NavCollisionBase.h
+++ b/Engine/Source/Runtime/Engine/Classes/AI/Navigation/NavCollisionBase.h
@@ -50,6 +50,8 @@ public:
 	*	collision data, stores it and uploads to DDC */
 	ENGINE_API virtual void Setup(class UBodySetup* BodySetup) PURE_VIRTUAL(UNavCollisionBase::Setup, );
 
+	[[nodiscard]] ENGINE_API virtual FBox GetBounds() const PURE_VIRTUAL(UNavCollisionBase::GetBounds, static FBox InvalidBox; return InvalidBox; );
+
 	/** Export collision data */
 	ENGINE_API virtual bool ExportGeometry(const FTransform& LocalToWorld, FNavigableGeometryExport& GeoExport) const PURE_VIRTUAL(UNavCollisionBase::ExportGeometry, return false; );
 
diff --git a/Engine/Source/Runtime/Engine/Classes/Components/HierarchicalInstancedStaticMeshComponent.h b/Engine/Source/Runtime/Engine/Classes/Components/HierarchicalInstancedStaticMeshComponent.h
--- a/Engine/Source/Runtime/Engine/Classes/Components/HierarchicalInstancedStaticMeshComponent.h
+++ b/Engine/Source/Runtime/Engine/Classes/Components/HierarchicalInstancedStaticMeshComponent.h
@@ -307,7 +307,6 @@ protected:
 	ENGINE_API virtual void GetNavigationPerInstanceTransforms(const FBox& AreaBox, TArray<FTransform>& InstanceData) const override;
 	ENGINE_API virtual void PartialNavigationUpdate(int32 InstanceIdx) override;
 	virtual bool SupportsPartialNavigationUpdate() const override { return true; }
-	ENGINE_API virtual FBox GetNavigationBounds() const override;
 	ENGINE_API void FlushAccumulatedNavigationUpdates();
 	mutable FBox AccumulatedNavigationDirtyArea;
 
diff --git a/Engine/Source/Runtime/Engine/Classes/Components/InstancedStaticMeshComponent.h b/Engine/Source/Runtime/Engine/Classes/Components/InstancedStaticMeshComponent.h
--- a/Engine/Source/Runtime/Engine/Classes/Components/InstancedStaticMeshComponent.h
+++ b/Engine/Source/Runtime/Engine/Classes/Components/InstancedStaticMeshComponent.h
@@ -522,6 +522,12 @@ protected:
 	/** Request to navigation system to update for the bounds of the ISM. */
 	ENGINE_API virtual void FullNavigationUpdate();
 
+	/**
+	 * Calculates bounds from all instances.
+	 * @param bForNavigation Indicates if NavCollision must be used if available, using static mesh bounds otherwise.
+	 */
+	FBoxSphereBounds CalcBoundsImpl(const FTransform& BoundTransform, bool bForNavigation) const;
+	
 	/** Does this component support partial navigation updates */
 	virtual bool SupportsPartialNavigationUpdate() const { return true; }
 
@@ -534,6 +540,13 @@ protected:
 	/** Internal version of RemoveInstance */	
 	ENGINE_API bool RemoveInstanceInternal(int32 InstanceIndex, bool InstanceAlreadyRemoved);
 
+	/**
+	 * Returns the bounds of a single instance in local space. It uses the NavCollision if available,
+	 * then the StaticMesh bounds if available or an invalid box otherwise.
+	 * @return The bounds of a single instance in local space 
+	 */
+	FBox GetInstanceNavigationBounds() const;
+	
 	/** Handles request from navigation system to gather instance transforms in a specific area box. */
 	ENGINE_API virtual void GetNavigationPerInstanceTransforms(const FBox& AreaBox, TArray<FTransform>& InstanceData) const;
 
diff --git a/Engine/Source/Runtime/Engine/Classes/Components/StaticMeshComponent.h b/Engine/Source/Runtime/Engine/Classes/Components/StaticMeshComponent.h
--- a/Engine/Source/Runtime/Engine/Classes/Components/StaticMeshComponent.h
+++ b/Engine/Source/Runtime/Engine/Classes/Components/StaticMeshComponent.h
@@ -566,6 +566,7 @@ public:
 
 	//~ Begin INavRelevantInterface Interface.
 	ENGINE_API virtual bool IsNavigationRelevant() const override;
+	ENGINE_API virtual FBox GetNavigationBounds() const override;
 	ENGINE_API virtual void GetNavigationData(FNavigationRelevantData& Data) const override;
 	//~ End INavRelevantInterface Interface.
 	/**
diff --git a/Engine/Source/Runtime/Engine/Private/Components/StaticMeshComponent.cpp b/Engine/Source/Runtime/Engine/Private/Components/StaticMeshComponent.cpp
--- a/Engine/Source/Runtime/Engine/Private/Components/StaticMeshComponent.cpp
+++ b/Engine/Source/Runtime/Engine/Private/Components/StaticMeshComponent.cpp
@@ -3079,6 +3079,25 @@ bool UStaticMeshComponent::IsNavigationRelevant() const
 	return GetStaticMesh() != nullptr && GetStaticMesh()->IsNavigationRelevant() && Super::IsNavigationRelevant();
 }
 
+FBox UStaticMeshComponent::GetNavigationBounds() const
+{
+	if (GetStaticMesh())
+	{
+		if (GetStaticMesh()->IsCompiling())
+		{
+			// Navigation bounds will get queried again once async static mesh compilation finishes
+			return FBox();
+		}
+
+		if (const UNavCollisionBase* NavCollision = GetStaticMesh()->GetNavCollision())
+		{
+			return NavCollision->GetBounds().TransformBy(GetComponentTransform());
+		}
+	}
+
+	return Super::GetNavigationBounds();
+}
+
 void UStaticMeshComponent::GetNavigationData(FNavigationRelevantData& Data) const
 {
 	Super::GetNavigationData(Data);
diff --git a/Engine/Source/Runtime/Engine/Private/HierarchicalInstancedStaticMesh.cpp b/Engine/Source/Runtime/Engine/Private/HierarchicalInstancedStaticMesh.cpp
--- a/Engine/Source/Runtime/Engine/Private/HierarchicalInstancedStaticMesh.cpp
+++ b/Engine/Source/Runtime/Engine/Private/HierarchicalInstancedStaticMesh.cpp
@@ -11,6 +11,7 @@
 #include "EngineStats.h"
 #include "Engine/Level.h"
 #include "AI/NavigationSystemBase.h"
+#include "AI/Navigation/NavCollisionBase.h"
 #include "Engine/MapBuildDataRegistry.h"
 #include "Materials/MaterialInterface.h"
 #include "Materials/MaterialRenderProxy.h"
@@ -3569,30 +3570,28 @@ void UHierarchicalInstancedStaticMeshComponent::GetNavigationPerInstanceTransfor
 	}
 }
 
-void UHierarchicalInstancedStaticMeshComponent::PartialNavigationUpdate(int32 InstanceIdx)
+void UHierarchicalInstancedStaticMeshComponent::PartialNavigationUpdate(const int32 InstanceIdx)
 {
 	if (InstanceIdx == INDEX_NONE)
 	{
 		AccumulatedNavigationDirtyArea.Init();
 		FNavigationSystem::UpdateComponentData(*this);
 	}
-	else if (GetStaticMesh())
+	else if (GetStaticMesh() && FNavigationSystem::HasComponentData(*this))
 	{
 		// Accumulate dirty areas and send them to navigation system once cluster tree is rebuilt
-		if (FNavigationSystem::HasComponentData(*this))
+		FTransform InstanceTransformInWorldSpace;
+		if (GetInstanceTransform(InstanceIdx, InstanceTransformInWorldSpace, /*bWorldSpace*/true))
 		{
-			const FTransform InstanceTransform(PerInstanceSMData[InstanceIdx].Transform);
-			const FBox InstanceBox = GetStaticMesh()->GetBounds().TransformBy(InstanceTransform*GetComponentTransform()).GetBox(); // in world space
-			AccumulatedNavigationDirtyArea += InstanceBox;
+			const FBox InstanceBounds = GetInstanceNavigationBounds();
+			if (InstanceBounds.IsValid)
+			{
+				AccumulatedNavigationDirtyArea += InstanceBounds.TransformBy(InstanceTransformInWorldSpace);
+			}
 		}
 	}
 }
 
-FBox UHierarchicalInstancedStaticMeshComponent::GetNavigationBounds() const
-{
-	return CalcBounds(GetComponentTransform()).GetBox();
-}
-
 void UHierarchicalInstancedStaticMeshComponent::FlushAccumulatedNavigationUpdates()
 {
 	if (AccumulatedNavigationDirtyArea.IsValid)
diff --git a/Engine/Source/Runtime/Engine/Private/InstancedStaticMesh.cpp b/Engine/Source/Runtime/Engine/Private/InstancedStaticMesh.cpp
--- a/Engine/Source/Runtime/Engine/Private/InstancedStaticMesh.cpp
+++ b/Engine/Source/Runtime/Engine/Private/InstancedStaticMesh.cpp
@@ -3150,16 +3150,19 @@ bool UInstancedStaticMeshComponent::CanEditSimulatePhysics()
 
 FBoxSphereBounds UInstancedStaticMeshComponent::CalcBounds(const FTransform& BoundTransform) const
 {
-	if(GetStaticMesh() && PerInstanceSMData.Num() > 0)
-	{
-		FMatrix BoundTransformMatrix = BoundTransform.ToMatrixWithScale();
-
-		FBoxSphereBounds RenderBounds = GetStaticMesh()->GetBounds();
+	return CalcBoundsImpl(BoundTransform, /*bForNavigation*/false);
+}
 
+FBoxSphereBounds UInstancedStaticMeshComponent::CalcBoundsImpl(const FTransform& BoundTransform, bool bForNavigation) const
+{
+	const FBox InstanceBounds = GetInstanceNavigationBounds();
+	if (InstanceBounds.IsValid && PerInstanceSMData.Num() > 0)
+	{
+		const FMatrix BoundTransformMatrix = BoundTransform.ToMatrixWithScale();
 		FBoxSphereBounds::Builder BoundsBuilder;
 		for (int32 InstanceIndex = 0; InstanceIndex < PerInstanceSMData.Num(); InstanceIndex++)
 		{
-			BoundsBuilder += RenderBounds.TransformBy(PerInstanceSMData[InstanceIndex].Transform * BoundTransformMatrix);
+			BoundsBuilder += InstanceBounds.TransformBy(PerInstanceSMData[InstanceIndex].Transform * BoundTransformMatrix);
 		}
 
 		return BoundsBuilder;
@@ -5235,23 +5238,16 @@ void UInstancedStaticMeshComponent::PartialNavigationUpdate(int32 InstanceIdx)
 	{
 		return;
 	}
-	
-	if (!InstanceBodies.IsValidIndex(InstanceIdx))
-	{
-		// The physics state might not be created and InstanceBodies might not be populated yet.
-		// This flow can occur from PostEditChangeChainProperty then AddInstanceInternal. 
-		return;
-	}
 
-	const FBodyInstance* const InstanceBodyInstance = InstanceBodies[InstanceIdx];
-	// Not having a body is a valid case when our physics state cannot be created (see CreateAllInstanceBodies)
-	if (!InstanceBodyInstance)
+	const FBox InstanceBounds = GetInstanceNavigationBounds();
+	if (InstanceBounds.IsValid)
 	{
-		return;
+		FTransform InstanceTransform;
+		if (GetInstanceTransform(InstanceIdx, InstanceTransform, /*bWorldSpace*/true))
+		{
+			FNavigationSystem::OnComponentBoundsChanged(*this, GetNavigationBounds(), InstanceBounds.TransformBy(InstanceTransform));
+		}
 	}
-	
-	const FBox InstanceBounds = InstanceBodyInstance->GetBodyBounds();
-	FNavigationSystem::OnComponentBoundsChanged(*this, GetNavigationBounds(),InstanceBounds);
 }
 
 bool UInstancedStaticMeshComponent::DoCustomNavigableGeometryExport(FNavigableGeometryExport& GeomExport) const
@@ -5266,11 +5262,7 @@ bool UInstancedStaticMeshComponent::DoCustomNavigableGeometryExport(FNavigableGe
 		
 		if (NavCollision->HasConvexGeometry())
 		{
-			GeomExport.ExportCustomMesh(NavCollision->GetConvexCollision().VertexBuffer.GetData(), NavCollision->GetConvexCollision().VertexBuffer.Num(),
-				NavCollision->GetConvexCollision().IndexBuffer.GetData(), NavCollision->GetConvexCollision().IndexBuffer.Num(), FTransform::Identity);
-
-			GeomExport.ExportCustomMesh(NavCollision->GetTriMeshCollision().VertexBuffer.GetData(), NavCollision->GetTriMeshCollision().VertexBuffer.Num(),
-				NavCollision->GetTriMeshCollision().IndexBuffer.GetData(), NavCollision->GetTriMeshCollision().IndexBuffer.Num(), FTransform::Identity);
+			NavCollision->ExportGeometry(FTransform::Identity, GeomExport);
 		}
 		else
 		{
@@ -5450,7 +5442,7 @@ void UInstancedStaticMeshComponent::GetNavigationData(FNavigationRelevantData& D
 
 FBox UInstancedStaticMeshComponent::GetNavigationBounds() const
 {
-	return CalcBounds(GetComponentTransform()).GetBox();
+	return CalcBoundsImpl(GetComponentTransform(), /*bForNavigation*/true).GetBox();
 }
 
 bool UInstancedStaticMeshComponent::IsNavigationRelevant() const
@@ -5458,19 +5450,29 @@ bool UInstancedStaticMeshComponent::IsNavigationRelevant() const
 	return GetInstanceCount() > 0 && Super::IsNavigationRelevant();
 }
 
-void UInstancedStaticMeshComponent::GetNavigationPerInstanceTransforms(const FBox& AreaBox, TArray<FTransform>& InstanceData) const
+FBox UInstancedStaticMeshComponent::GetInstanceNavigationBounds() const
 {
-	if (GetStaticMesh())
+	if (const UStaticMesh* Mesh = GetStaticMesh())
 	{
-		const FBox LocalAreaBox =  AreaBox.InverseTransformBy(GetComponentTransform());
-		const FBoxSphereBounds RenderBounds = GetStaticMesh()->GetBounds();
+		const UNavCollisionBase* NavCollision = Mesh->GetNavCollision();
+		return NavCollision ? NavCollision->GetBounds() : Mesh->GetBounds().GetBox();
+	}
 
+	return FBox();
+}
+
+void UInstancedStaticMeshComponent::GetNavigationPerInstanceTransforms(const FBox& AreaBox, TArray<FTransform>& InstanceData) const
+{
+	const FBox InstanceBounds = GetInstanceNavigationBounds();
+	if (InstanceBounds.IsValid)
+	{
+		const FBox LocalAreaBox = AreaBox.InverseTransformBy(GetComponentTransform());
 		for (const auto& InstancedData : PerInstanceSMData)
 		{
 			const FTransform InstanceToComponent(InstancedData.Transform);
 			if (!InstanceToComponent.GetScale3D().IsZero())
 			{
-				const FBoxSphereBounds TransformedInstanceBounds = RenderBounds.TransformBy(InstancedData.Transform);
+				const FBoxSphereBounds TransformedInstanceBounds = InstanceBounds.TransformBy(InstancedData.Transform);
 				if (LocalAreaBox.Intersect(TransformedInstanceBounds.GetBox()))
 				{
 					InstanceData.Add(InstanceToComponent*GetComponentTransform());
diff --git a/Engine/Source/Runtime/NavigationSystem/Private/NavCollision.cpp b/Engine/Source/Runtime/NavigationSystem/Private/NavCollision.cpp
--- a/Engine/Source/Runtime/NavigationSystem/Private/NavCollision.cpp
+++ b/Engine/Source/Runtime/NavigationSystem/Private/NavCollision.cpp
@@ -47,11 +47,13 @@ public:
 	FNavCollisionConvex& TriMeshCollision;
 	FNavCollisionConvex& ConvexCollision;
 	TNavStatArray<int32>& ConvexShapeIndices;
+	FBox& Bounds;
 
-	FNavCollisionDataReader(FByteBulkData& InBulkData, FNavCollisionConvex& InTriMeshCollision, FNavCollisionConvex& InConvexCollision, TNavStatArray<int32>& InShapeIndices)
+	FNavCollisionDataReader(FByteBulkData& InBulkData, FNavCollisionConvex& InTriMeshCollision, FNavCollisionConvex& InConvexCollision, TNavStatArray<int32>& InShapeIndices, FBox& InBounds)
 		: TriMeshCollision(InTriMeshCollision)
 		, ConvexCollision(InConvexCollision)
 		, ConvexShapeIndices(InShapeIndices)
+		, Bounds(InBounds)
 	{
 		// Read cooked data
 		uint8* DataPtr = (uint8*)InBulkData.Lock( LOCK_READ_ONLY );
@@ -66,6 +68,7 @@ public:
 		Ar << ConvexCollision.VertexBuffer;
 		Ar << ConvexCollision.IndexBuffer;
 		Ar << ConvexShapeIndices;
+		Ar << Bounds;
 
 		InBulkData.Unlock();
 	}
@@ -93,7 +96,7 @@ public:
 
 	virtual const TCHAR* GetVersionString() const override
 	{
-		return TEXT("F33FFAC3B070461781F1B11C4EAC7100");
+		return TEXT("8F4645C7A18B40C487C7E50E19CD6B6C");
 	}
 
 	virtual FString GetPluginSpecificCacheKeySuffix() const override
@@ -158,6 +161,7 @@ bool FDerivedDataNavCollisionCooker::Build( TArray<uint8>& OutData )
 	Ar << NavCollisionInstance->GetMutableConvexCollision().VertexBuffer;
 	Ar << NavCollisionInstance->GetMutableConvexCollision().IndexBuffer;
 	Ar << NavCollisionInstance->ConvexShapeIndices;
+	Ar << NavCollisionInstance->Bounds;
 
 	// Whatever got cached return true. We want to cache 'failure' too.
 	return true;
@@ -239,7 +243,7 @@ void UNavCollision::Setup(UBodySetup* BodySetup)
 		if (FormatData->IsLocked() == false)
 		{
 			// Create physics objects
-			FNavCollisionDataReader CookedDataReader(*FormatData, TriMeshCollision, ConvexCollision, ConvexShapeIndices);
+			FNavCollisionDataReader CookedDataReader(*FormatData, TriMeshCollision, ConvexCollision, ConvexShapeIndices, Bounds);
 			bHasConvexGeometry = true;
 		}
 	}
@@ -249,6 +253,11 @@ void UNavCollision::Setup(UBodySetup* BodySetup)
 	}
 }
 
+FBox UNavCollision::GetBounds() const
+{
+	return Bounds;
+}
+
 void UNavCollision::GatherCollision()
 {
 	ClearCollision();
@@ -301,6 +310,7 @@ void UNavCollision::ClearCollision()
 	ConvexCollision.VertexBuffer.Reset();
 	ConvexCollision.IndexBuffer.Reset();
 	ConvexShapeIndices.Reset();
+	Bounds = FBox();
 
 	bHasConvexGeometry = false;
 }
diff --git a/Engine/Source/Runtime/NavigationSystem/Private/NavMesh/RecastNavMeshGenerator.cpp b/Engine/Source/Runtime/NavigationSystem/Private/NavMesh/RecastNavMeshGenerator.cpp
--- a/Engine/Source/Runtime/NavigationSystem/Private/NavMesh/RecastNavMeshGenerator.cpp
+++ b/Engine/Source/Runtime/NavigationSystem/Private/NavMesh/RecastNavMeshGenerator.cpp
@@ -6932,12 +6932,25 @@ void FRecastNavMeshGenerator::ExportVertexSoupGeometry(const TArray<FVector>& In
 	RecastGeometryExport::StoreCollisionCache(GeomExport);
 }
 
-void FRecastNavMeshGenerator::ExportRigidBodyGeometry(UBodySetup& InOutBodySetup, TNavStatArray<FVector>& OutVertexBuffer, TNavStatArray<int32>& OutIndexBuffer, const FTransform& LocalToWorld)
+void FRecastNavMeshGenerator::ExportRigidBodyGeometry(
+	UBodySetup& InOutBodySetup,
+	TNavStatArray<FVector>& OutVertexBuffer,
+	TNavStatArray<int32>& OutIndexBuffer,
+	const FTransform& LocalToWorld)
 {
-	TNavStatArray<FVector::FReal> VertCoords;
 	FBox TempBounds;
+	ExportRigidBodyGeometry(InOutBodySetup, OutVertexBuffer, OutIndexBuffer, TempBounds, LocalToWorld);
+}
 
-	RecastGeometryExport::ExportRigidBodySetup(InOutBodySetup, VertCoords, OutIndexBuffer, TempBounds, LocalToWorld);
+void FRecastNavMeshGenerator::ExportRigidBodyGeometry(
+	UBodySetup& InOutBodySetup,
+	TNavStatArray<FVector>& OutVertexBuffer,
+	TNavStatArray<int32>& OutIndexBuffer,
+	FBox& OutBounds,
+	const FTransform& LocalToWorld)
+{
+	TNavStatArray<FVector::FReal> VertCoords;
+	RecastGeometryExport::ExportRigidBodySetup(InOutBodySetup, VertCoords, OutIndexBuffer, OutBounds, LocalToWorld);
 
 	OutVertexBuffer.Reserve(OutVertexBuffer.Num() + (VertCoords.Num() / 3));
 	for (int32 i = 0; i < VertCoords.Num(); i += 3)
@@ -6946,17 +6959,41 @@ void FRecastNavMeshGenerator::ExportRigidBodyGeometry(UBodySetup& InOutBodySetup
 	}
 }
 
-void FRecastNavMeshGenerator::ExportRigidBodyGeometry(UBodySetup& InOutBodySetup, TNavStatArray<FVector>& OutTriMeshVertexBuffer, TNavStatArray<int32>& OutTriMeshIndexBuffer
-	, TNavStatArray<FVector>& OutConvexVertexBuffer, TNavStatArray<int32>& OutConvexIndexBuffer, TNavStatArray<int32>& OutShapeBuffer
-	, const FTransform& LocalToWorld)
+void FRecastNavMeshGenerator::ExportRigidBodyGeometry(
+	UBodySetup& InOutBodySetup,
+	TNavStatArray<FVector>& OutTriMeshVertexBuffer,
+	TNavStatArray<int32>& OutTriMeshIndexBuffer,
+	TNavStatArray<FVector>& OutConvexVertexBuffer,
+	TNavStatArray<int32>& OutConvexIndexBuffer,
+	TNavStatArray<int32>& OutShapeBuffer,
+	const FTransform& LocalToWorld)
+{
+	FBox TempBounds;
+	ExportRigidBodyGeometry(
+		InOutBodySetup,
+		OutTriMeshVertexBuffer,
+		OutTriMeshIndexBuffer,
+		OutConvexVertexBuffer,
+		OutConvexIndexBuffer,
+		OutShapeBuffer,
+		TempBounds,
+		LocalToWorld);
+}
+
+void FRecastNavMeshGenerator::ExportRigidBodyGeometry(
+	UBodySetup& InOutBodySetup,
+	TNavStatArray<FVector>& OutTriMeshVertexBuffer,
+	TNavStatArray<int32>& OutTriMeshIndexBuffer,
+	TNavStatArray<FVector>& OutConvexVertexBuffer,
+	TNavStatArray<int32>& OutConvexIndexBuffer,
+	TNavStatArray<int32>& OutShapeBuffer,
+	FBox& OutBounds,
+	const FTransform& LocalToWorld)
 {
 	InOutBodySetup.CreatePhysicsMeshes();
 
 	TNavStatArray<FVector::FReal> VertCoords;
-	FBox TempBounds;
-
-	VertCoords.Reset();
-	RecastGeometryExport::ExportRigidBodyTriMesh(InOutBodySetup, VertCoords, OutTriMeshIndexBuffer, TempBounds, LocalToWorld);
+	RecastGeometryExport::ExportRigidBodyTriMesh(InOutBodySetup, VertCoords, OutTriMeshIndexBuffer, OutBounds, LocalToWorld);
 
 	OutTriMeshVertexBuffer.Reserve(OutTriMeshVertexBuffer.Num() + (VertCoords.Num() / 3));
 	for (int32 i = 0; i < VertCoords.Num(); i += 3)
@@ -6966,10 +7003,10 @@ void FRecastNavMeshGenerator::ExportRigidBodyGeometry(UBodySetup& InOutBodySetup
 
 	const int32 NumExistingVerts = OutConvexVertexBuffer.Num();
 	VertCoords.Reset();
-	RecastGeometryExport::ExportRigidBodyConvexElements(InOutBodySetup, VertCoords, OutConvexIndexBuffer, OutShapeBuffer, TempBounds, LocalToWorld);
-	RecastGeometryExport::ExportRigidBodyBoxElements(InOutBodySetup.AggGeom, VertCoords, OutConvexIndexBuffer, OutShapeBuffer, TempBounds, LocalToWorld, NumExistingVerts);
-	RecastGeometryExport::ExportRigidBodySphylElements(InOutBodySetup.AggGeom, VertCoords, OutConvexIndexBuffer, OutShapeBuffer, TempBounds, LocalToWorld, NumExistingVerts);
-	RecastGeometryExport::ExportRigidBodySphereElements(InOutBodySetup.AggGeom, VertCoords, OutConvexIndexBuffer, OutShapeBuffer, TempBounds, LocalToWorld, NumExistingVerts);
+	RecastGeometryExport::ExportRigidBodyConvexElements(InOutBodySetup, VertCoords, OutConvexIndexBuffer, OutShapeBuffer, OutBounds, LocalToWorld);
+	RecastGeometryExport::ExportRigidBodyBoxElements(InOutBodySetup.AggGeom, VertCoords, OutConvexIndexBuffer, OutShapeBuffer, OutBounds, LocalToWorld, NumExistingVerts);
+	RecastGeometryExport::ExportRigidBodySphylElements(InOutBodySetup.AggGeom, VertCoords, OutConvexIndexBuffer, OutShapeBuffer, OutBounds, LocalToWorld, NumExistingVerts);
+	RecastGeometryExport::ExportRigidBodySphereElements(InOutBodySetup.AggGeom, VertCoords, OutConvexIndexBuffer, OutShapeBuffer, OutBounds, LocalToWorld, NumExistingVerts);
 	
 	OutConvexVertexBuffer.Reserve(OutConvexVertexBuffer.Num() + (VertCoords.Num() / 3));
 	for (int32 i = 0; i < VertCoords.Num(); i += 3)
@@ -6978,19 +7015,34 @@ void FRecastNavMeshGenerator::ExportRigidBodyGeometry(UBodySetup& InOutBodySetup
 	}
 }
 
-void FRecastNavMeshGenerator::ExportAggregatedGeometry(const FKAggregateGeom& AggGeom, TNavStatArray<FVector>& OutConvexVertexBuffer, TNavStatArray<int32>& OutConvexIndexBuffer, TNavStatArray<int32>& OutShapeBuffer, const FTransform& LocalToWorld)
+void FRecastNavMeshGenerator::ExportAggregatedGeometry(
+	const FKAggregateGeom& AggGeom,
+	TNavStatArray<FVector>& OutConvexVertexBuffer,
+	TNavStatArray<int32>& OutConvexIndexBuffer,
+	TNavStatArray<int32>& OutShapeBuffer,
+	const FTransform& LocalToWorld)
 {
-	TNavStatArray<FVector::FReal> VertCoords;
 	FBox TempBounds;
+	ExportAggregatedGeometry(AggGeom, OutConvexVertexBuffer, OutConvexIndexBuffer, OutShapeBuffer, TempBounds, LocalToWorld);
+}
 
+void FRecastNavMeshGenerator::ExportAggregatedGeometry(
+	const FKAggregateGeom& AggGeom,
+	TNavStatArray<FVector>& OutConvexVertexBuffer,
+	TNavStatArray<int32>& OutConvexIndexBuffer,
+	TNavStatArray<int32>& OutShapeBuffer,
+	FBox& OutBounds,
+	const FTransform& LocalToWorld)
+{
+	TNavStatArray<FVector::FReal> VertCoords;
 	const int32 NumExistingVerts = OutConvexVertexBuffer.Num();
 
 	// convex and tri mesh are NOT supported, since they require BodySetup.CreatePhysicsMeshes() call
 	// only simple shapes
 
-	RecastGeometryExport::ExportRigidBodyBoxElements(AggGeom, VertCoords, OutConvexIndexBuffer, OutShapeBuffer, TempBounds, LocalToWorld, NumExistingVerts);
-	RecastGeometryExport::ExportRigidBodySphylElements(AggGeom, VertCoords, OutConvexIndexBuffer, OutShapeBuffer, TempBounds, LocalToWorld, NumExistingVerts);
-	RecastGeometryExport::ExportRigidBodySphereElements(AggGeom, VertCoords, OutConvexIndexBuffer, OutShapeBuffer, TempBounds, LocalToWorld, NumExistingVerts);
+	RecastGeometryExport::ExportRigidBodyBoxElements(AggGeom, VertCoords, OutConvexIndexBuffer, OutShapeBuffer, OutBounds, LocalToWorld, NumExistingVerts);
+	RecastGeometryExport::ExportRigidBodySphylElements(AggGeom, VertCoords, OutConvexIndexBuffer, OutShapeBuffer, OutBounds, LocalToWorld, NumExistingVerts);
+	RecastGeometryExport::ExportRigidBodySphereElements(AggGeom, VertCoords, OutConvexIndexBuffer, OutShapeBuffer, OutBounds, LocalToWorld, NumExistingVerts);
 
 	OutConvexVertexBuffer.Reserve(OutConvexVertexBuffer.Num() + (VertCoords.Num() / 3));
 	for (int32 i = 0; i < VertCoords.Num(); i += 3)
diff --git a/Engine/Source/Runtime/NavigationSystem/Private/NavigationSystemTypes.cpp b/Engine/Source/Runtime/NavigationSystem/Private/NavigationSystemTypes.cpp
--- a/Engine/Source/Runtime/NavigationSystem/Private/NavigationSystemTypes.cpp
+++ b/Engine/Source/Runtime/NavigationSystem/Private/NavigationSystemTypes.cpp
@@ -41,14 +41,14 @@ namespace
 
 namespace NavigationHelper
 {
-	void GatherCollision(UBodySetup* RigidBody, TNavStatArray<FVector>& OutVertexBuffer, TNavStatArray<int32>& OutIndexBuffer, const FTransform& LocalToWorld)
+	void GatherCollision(UBodySetup* RigidBody, TNavStatArray<FVector>& OutVertexBuffer, TNavStatArray<int32>& OutIndexBuffer, const FTransform& LocalToWorld, FBox& OutBounds)
 	{
 		if (RigidBody == NULL)
 		{
 			return;
 		}
 #if WITH_RECAST
-		FRecastNavMeshGenerator::ExportRigidBodyGeometry(*RigidBody, OutVertexBuffer, OutIndexBuffer, LocalToWorld);
+		FRecastNavMeshGenerator::ExportRigidBodyGeometry(*RigidBody, OutVertexBuffer, OutIndexBuffer, OutBounds, LocalToWorld);
 #endif // WITH_RECAST
 	}
 
@@ -62,14 +62,20 @@ namespace NavigationHelper
 		FRecastNavMeshGenerator::ExportRigidBodyGeometry(*RigidBody
 			, NavCollision->GetMutableTriMeshCollision().VertexBuffer, NavCollision->GetMutableTriMeshCollision().IndexBuffer
 			, NavCollision->GetMutableConvexCollision().VertexBuffer, NavCollision->GetMutableConvexCollision().IndexBuffer
-			, NavCollision->ConvexShapeIndices);
+			, NavCollision->ConvexShapeIndices
+			, NavCollision->Bounds);
 #endif // WITH_RECAST
 	}
 
 	void GatherCollision(const FKAggregateGeom& AggGeom, UNavCollision& NavCollision)
 	{
 #if WITH_RECAST
-		FRecastNavMeshGenerator::ExportAggregatedGeometry(AggGeom, NavCollision.GetMutableConvexCollision().VertexBuffer, NavCollision.GetMutableConvexCollision().IndexBuffer, NavCollision.ConvexShapeIndices);
+		FRecastNavMeshGenerator::ExportAggregatedGeometry(
+			AggGeom,
+			NavCollision.GetMutableConvexCollision().VertexBuffer,
+			 NavCollision.GetMutableConvexCollision().IndexBuffer,
+			 NavCollision.ConvexShapeIndices,
+			 NavCollision.Bounds);
 #endif // WITH_RECAST
 	}
 
diff --git a/Engine/Source/Runtime/NavigationSystem/Public/NavCollision.h b/Engine/Source/Runtime/NavigationSystem/Public/NavCollision.h
--- a/Engine/Source/Runtime/NavigationSystem/Public/NavCollision.h
+++ b/Engine/Source/Runtime/NavigationSystem/Public/NavCollision.h
@@ -50,6 +50,8 @@ class UNavCollision : public UNavCollisionBase
 
 	TNavStatArray<int32> ConvexShapeIndices;
 
+	FBox Bounds;
+
 	/** list of nav collision cylinders */
 	UPROPERTY(EditAnywhere, Category=Navigation)
 	TArray<FNavCollisionCylinder> CylinderCollision;
@@ -96,6 +98,8 @@ class UNavCollision : public UNavCollisionBase
 	 *	collision data, stores it and uploads to DDC */
 	NAVIGATIONSYSTEM_API virtual void Setup(class UBodySetup* BodySetup) override;
 
+	NAVIGATIONSYSTEM_API virtual FBox GetBounds() const override;
+
 	/** copy user settings from other nav collision data */
 	NAVIGATIONSYSTEM_API void CopyUserSettings(const UNavCollision& OtherData);
 
diff --git a/Engine/Source/Runtime/NavigationSystem/Public/NavMesh/RecastNavMeshGenerator.h b/Engine/Source/Runtime/NavigationSystem/Public/NavMesh/RecastNavMeshGenerator.h
--- a/Engine/Source/Runtime/NavigationSystem/Public/NavMesh/RecastNavMeshGenerator.h
+++ b/Engine/Source/Runtime/NavigationSystem/Public/NavMesh/RecastNavMeshGenerator.h
@@ -765,12 +765,58 @@ public:
 
 	UE_DEPRECATED(5.4, "Use ExportNavRelevantObjectGeometry")
 	static NAVIGATIONSYSTEM_API void ExportComponentGeometry(UActorComponent* InOutComponent, FNavigationRelevantData& OutData);
+
+	UE_DEPRECATED(5.4, "Use ExportRigidBodyGeometry that takes bounds as parameter.")
+	static NAVIGATIONSYSTEM_API void ExportRigidBodyGeometry(
+		UBodySetup& InOutBodySetup,
+		TNavStatArray<FVector>& OutVertexBuffer,
+		TNavStatArray<int32>& OutIndexBuffer,
+		const FTransform& LocalToWorld = FTransform::Identity);
+
+	UE_DEPRECATED(5.4, "Use ExportRigidBodyGeometry that takes bounds as parameter.")
+	static NAVIGATIONSYSTEM_API void ExportRigidBodyGeometry(
+		UBodySetup& InOutBodySetup,
+		TNavStatArray<FVector>& OutTriMeshVertexBuffer,
+		TNavStatArray<int32>& OutTriMeshIndexBuffer,
+		TNavStatArray<FVector>& OutConvexVertexBuffer,
+		TNavStatArray<int32>& OutConvexIndexBuffer,
+		TNavStatArray<int32>& OutShapeBuffer,
+		const FTransform& LocalToWorld = FTransform::Identity);
+
+	UE_DEPRECATED(5.4, "Use ExportAggregatedGeometry that takes bounds as parameter.")
+	static NAVIGATIONSYSTEM_API void ExportAggregatedGeometry(
+		const FKAggregateGeom& AggGeom,
+		TNavStatArray<FVector>& OutConvexVertexBuffer,
+		TNavStatArray<int32>& OutConvexIndexBuffer,
+		TNavStatArray<int32>& OutShapeBuffer,
+		const FTransform& LocalToWorld = FTransform::Identity);
+
 	static NAVIGATIONSYSTEM_API void ExportNavRelevantObjectGeometry(INavRelevantInterface& InOutNavRelevantInterface, FNavigationRelevantData& OutData);
 	static NAVIGATIONSYSTEM_API void ExportVertexSoupGeometry(const TArray<FVector>& InVerts, FNavigationRelevantData& OutData);
 
-	static NAVIGATIONSYSTEM_API void ExportRigidBodyGeometry(UBodySetup& InOutBodySetup, TNavStatArray<FVector>& OutVertexBuffer, TNavStatArray<int32>& OutIndexBuffer, const FTransform& LocalToWorld = FTransform::Identity);
-	static NAVIGATIONSYSTEM_API void ExportRigidBodyGeometry(UBodySetup& InOutBodySetup, TNavStatArray<FVector>& OutTriMeshVertexBuffer, TNavStatArray<int32>& OutTriMeshIndexBuffer, TNavStatArray<FVector>& OutConvexVertexBuffer, TNavStatArray<int32>& OutConvexIndexBuffer, TNavStatArray<int32>& OutShapeBuffer, const FTransform& LocalToWorld = FTransform::Identity);
-	static NAVIGATIONSYSTEM_API void ExportAggregatedGeometry(const FKAggregateGeom& AggGeom, TNavStatArray<FVector>& OutConvexVertexBuffer, TNavStatArray<int32>& OutConvexIndexBuffer, TNavStatArray<int32>& OutShapeBuffer, const FTransform& LocalToWorld = FTransform::Identity);
+	static NAVIGATIONSYSTEM_API void ExportRigidBodyGeometry(UBodySetup& InOutBodySetup,
+		TNavStatArray<FVector>& OutVertexBuffer,
+		TNavStatArray<int32>& OutIndexBuffer,
+		FBox& OutBounds,
+		const FTransform& LocalToWorld = FTransform::Identity);
+
+	static NAVIGATIONSYSTEM_API void ExportRigidBodyGeometry(
+		UBodySetup& InOutBodySetup,
+		TNavStatArray<FVector>& OutTriMeshVertexBuffer,
+		TNavStatArray<int32>& OutTriMeshIndexBuffer,
+		TNavStatArray<FVector>& OutConvexVertexBuffer,
+		TNavStatArray<int32>& OutConvexIndexBuffer,
+		TNavStatArray<int32>& OutShapeBuffer,
+		FBox& OutBounds,
+		const FTransform& LocalToWorld = FTransform::Identity);
+
+	static NAVIGATIONSYSTEM_API void ExportAggregatedGeometry(
+		const FKAggregateGeom& AggGeom,
+		TNavStatArray<FVector>& OutConvexVertexBuffer,
+		TNavStatArray<int32>& OutConvexIndexBuffer,
+		TNavStatArray<int32>& OutShapeBuffer,
+		FBox& OutBounds,
+		const FTransform& LocalToWorld = FTransform::Identity);
 
 #if UE_ENABLE_DEBUG_DRAWING
 	/** Converts data encoded in EncodedData.CollisionData to FNavDebugMeshData format */
