Add camera location to invoker seed locations when the pawn is null so nav mesh is not removed in cases where the player camera is near the invoker but the player pawn has been destroyed.

diff --git a/Engine/Source/Runtime/NavigationSystem/Private/NavigationSystem.cpp b/Engine/Source/Runtime/NavigationSystem/Private/NavigationSystem.cpp
--- a/Engine/Source/Runtime/NavigationSystem/Private/NavigationSystem.cpp
+++ b/Engine/Source/Runtime/NavigationSystem/Private/NavigationSystem.cpp
@@ -3793,10 +3793,18 @@ void UNavigationSystemV1::GetInvokerSeedLocations(const UWorld& InWorld, TArray<
 	for (FConstPlayerControllerIterator PlayerIt = InWorld.GetPlayerControllerIterator(); PlayerIt; ++PlayerIt)
 	{
 		const APlayerController* PlayerController = PlayerIt->Get();
-		if (PlayerController && PlayerController->GetPawn())
+		if (PlayerController)
 		{
-			const FVector2D SeedLoc(PlayerController->GetPawn()->GetActorLocation());
-			OutSeedLocations.Add(SeedLoc);
+			if (PlayerController->GetPawn())
+			{
+				const FVector2D SeedLoc(PlayerController->GetPawn()->GetActorLocation());
+				OutSeedLocations.Add(SeedLoc);
+			}
+			else if (PlayerController->PlayerCameraManager)
+			{
+				const FVector2D SeedLoc(PlayerController->PlayerCameraManager->GetCameraLocation());
+				OutSeedLocations.Add(SeedLoc);
+			}
 		}
 	}
 }
diff --git a/Engine/Source/Runtime/NavigationSystem/Public/NavigationSystem.h b/Engine/Source/Runtime/NavigationSystem/Public/NavigationSystem.h
--- a/Engine/Source/Runtime/NavigationSystem/Public/NavigationSystem.h
+++ b/Engine/Source/Runtime/NavigationSystem/Public/NavigationSystem.h
@@ -868,7 +868,8 @@ protected:
 	/** Searches for all valid navigation bounds in the world and stores them */
 	NAVIGATIONSYSTEM_API virtual void GatherNavigationBounds();
 
-	/** Get seed locations for invokers, @see InvokersMaximumDistanceFromSeed */ 
+	/** Get seed locations for invokers, @see InvokersMaximumDistanceFromSeed 
+	 *	By default these are the player pawn locations. If a player controller has no pawn assigned, the player's camera location will be used instead. */ 
 	NAVIGATIONSYSTEM_API virtual void GetInvokerSeedLocations(const UWorld& InWorld, TArray<FVector2D, TInlineAllocator<32>>& OutSeedLocations);
 
 	/** @return pointer to ANavigationData instance of given ID, or NULL if it was not found. Note it looks only through registered navigation data */
