Fixed unimplemented virtual functions when WITH_RECAST is disabled
#rnx

diff --git a/Engine/Source/Runtime/NavigationSystem/Public/NavMesh/RecastNavMesh.h b/Engine/Source/Runtime/NavigationSystem/Public/NavMesh/RecastNavMesh.h
--- a/Engine/Source/Runtime/NavigationSystem/Public/NavMesh/RecastNavMesh.h
+++ b/Engine/Source/Runtime/NavigationSystem/Public/NavMesh/RecastNavMesh.h
@@ -1516,6 +1516,23 @@ private:
 
 private:
 	static NAVIGATIONSYSTEM_API const FRecastQueryFilter* NamedFilters[ERecastNamedFilter::NamedFiltersCount];
+#else
+	virtual bool IsNodeRefValid(NavNodeRef NodeRef) const override { return true; }
+	virtual FBox GetBounds() const override { return FBox(); }
+	virtual void BatchRaycast(TArray<FNavigationRaycastWork>& Workload, FSharedConstNavQueryFilter QueryFilter, const UObject* Querier = nullptr) const override {}
+	virtual bool FindMoveAlongSurface(const FNavLocation& StartLocation, const FVector& TargetPosition, FNavLocation& OutLocation, FSharedConstNavQueryFilter Filter = nullptr, const UObject* Querier = nullptr) const override { return false; }
+	virtual bool FindOverlappingEdges(const FNavLocation& StartLocation, TConstArrayView<FVector> ConvexPolygon, TArray<FVector>& OutEdges, FSharedConstNavQueryFilter Filter = nullptr, const UObject* Querier = nullptr) const override { return false; }
+	virtual bool GetPathSegmentBoundaryEdges(const FNavigationPath& Path, const FNavPathPoint& StartPoint, const FNavPathPoint& EndPoint, const TConstArrayView<FVector> SearchArea, TArray<FVector>& OutEdges, const float MaxAreaEnterCost, FSharedConstNavQueryFilter Filter = nullptr, const UObject* Querier = nullptr) const override { return false; }
+	virtual FNavLocation GetRandomPoint(FSharedConstNavQueryFilter Filter = nullptr, const UObject* Querier = nullptr) const override { return FNavLocation(); }
+	virtual bool GetRandomReachablePointInRadius(const FVector& Origin, float Radius, FNavLocation& OutResult, FSharedConstNavQueryFilter Filter = nullptr, const UObject* Querier = nullptr) const override { return false; }
+	virtual bool GetRandomPointInNavigableRadius(const FVector& Origin, float Radius, FNavLocation& OutResult, FSharedConstNavQueryFilter Filter = nullptr, const UObject* Querier = nullptr) const override { return false; }
+	virtual bool ProjectPoint(const FVector& Point, FNavLocation& OutLocation, const FVector& Extent, FSharedConstNavQueryFilter Filter = nullptr, const UObject* Querier = nullptr) const override { return false; }
+	virtual void BatchProjectPoints(TArray<FNavigationProjectionWork>& Workload, const FVector& Extent, FSharedConstNavQueryFilter Filter = nullptr, const UObject* Querier = nullptr) const override {}
+	virtual void BatchProjectPoints(TArray<FNavigationProjectionWork>& Workload, FSharedConstNavQueryFilter Filter = nullptr, const UObject* Querier = nullptr) const override {}
+	virtual ENavigationQueryResult::Type CalcPathCost(const FVector& PathStart, const FVector& PathEnd, FVector::FReal& OutPathCost, FSharedConstNavQueryFilter QueryFilter = nullptr, const UObject* Querier = nullptr) const override { return ENavigationQueryResult::Invalid; }
+	virtual ENavigationQueryResult::Type CalcPathLength(const FVector& PathStart, const FVector& PathEnd, FVector::FReal& OutPathLength, FSharedConstNavQueryFilter QueryFilter = nullptr, const UObject* Querier = nullptr) const override { return ENavigationQueryResult::Invalid; }
+	virtual ENavigationQueryResult::Type CalcPathLengthAndCost(const FVector& PathStart, const FVector& PathEnd, FVector::FReal& OutPathLength, FVector::FReal& OutPathCost, FSharedConstNavQueryFilter QueryFilter = nullptr, const UObject* Querier = nullptr) const override { return ENavigationQueryResult::Invalid; }
+	virtual bool DoesNodeContainLocation(NavNodeRef NodeRef, const FVector& WorldSpaceLocation) const override { return false; }
 #endif // WITH_RECAST
 
 public:
