Fix for frequent check() in FRcTileBox due to invalid FBox bounds. This code should stop the issue being fatal and handle the bad case reasonably.

diff --git a/Engine/Source/Runtime/NavigationSystem/Public/NavMesh/RecastNavMeshGenerator.h b/Engine/Source/Runtime/NavigationSystem/Public/NavMesh/RecastNavMeshGenerator.h
--- a/Engine/Source/Runtime/NavigationSystem/Public/NavMesh/RecastNavMeshGenerator.h
+++ b/Engine/Source/Runtime/NavigationSystem/Public/NavMesh/RecastNavMeshGenerator.h
@@ -191,7 +191,19 @@ struct FRcTileBox
 		check(TileSizeInWorldUnits > 0);
 		checkf(!RcNavMeshOrigin.ContainsNaN(), TEXT("%hs: RcNavMeshOrigin ContainsNaN"), __FUNCTION__);
 		checkf(!UnrealBounds.ContainsNaN(), TEXT("%hs: UnrealBounds ContainsNaN"), __FUNCTION__);
-		checkf(UnrealBounds.IsValid, TEXT("%hs: UnrealBounds !IsValid()"), __FUNCTION__);
+		
+		if (!ensureMsgf(UnrealBounds.IsValid, TEXT("%hs: UnrealBounds !IsValid"), __FUNCTION__))
+		{
+			// This is a bug but we'll handle it as well as we can.
+
+			// Invalid bounds, set to empty range.
+			// Max is set to -1, because the range is inclusive, used like: for (int32 y = TileBox.YMin; y <= TileBox.YMax; ++y)
+			XMin = 0;
+			XMax = -1;
+			YMin = 0;
+			YMax = -1;
+			return;
+		}
 
 		auto CalcMaxCoordExclusive = [](const FVector::FReal MaxAsFloat, const int32 MinCoord) -> int32
 		{
