Support provided InLocalToWorld transform rather than always returning bounds in absolute world space.

#rnx

diff --git a/Engine/Source/Runtime/NavigationSystem/Private/NavLinkRenderingComponent.cpp b/Engine/Source/Runtime/NavigationSystem/Private/NavLinkRenderingComponent.cpp
--- a/Engine/Source/Runtime/NavigationSystem/Private/NavLinkRenderingComponent.cpp
+++ b/Engine/Source/Runtime/NavigationSystem/Private/NavLinkRenderingComponent.cpp
@@ -42,7 +42,6 @@ FBoxSphereBounds UNavLinkRenderingComponent::CalcBounds(const FTransform& InLoca
 	if (LinkOwnerActor != NULL)
 	{
 		FBox BoundingBox(ForceInit);
-		const FTransform LocalToWorld = LinkOwnerActor->ActorToWorld();
 
 		INavLinkHostInterface* LinkOwnerHost = Cast<INavLinkHostInterface>(LinkOwnerActor);
 		if (LinkOwnerHost != NULL)
@@ -76,7 +75,10 @@ FBoxSphereBounds UNavLinkRenderingComponent::CalcBounds(const FTransform& InLoca
 			}
 		}
 
-		return FBoxSphereBounds(BoundingBox).TransformBy(LocalToWorld);
+		// BoundingBox is in actor space. Incorporate provided InLocalToWorld transform via component space.
+		const FTransform ActorToWorld = LinkOwnerActor->ActorToWorld();
+		const FTransform WorldToComponent = GetComponentTransform().Inverse();
+		return FBoxSphereBounds(BoundingBox).TransformBy(ActorToWorld * WorldToComponent * InLocalToWorld);
 	}
 
 	return FBoxSphereBounds(ForceInitToZero);
