Follow up to [REDACTED_HANDLE] . Designed to help track down issues that are being caused by invalid bounds being passed in to FRcTileBox.

diff --git a/Engine/Source/Runtime/NavigationSystem/Private/NavMesh/PImplRecastNavMesh.cpp b/Engine/Source/Runtime/NavigationSystem/Private/NavMesh/PImplRecastNavMesh.cpp
--- a/Engine/Source/Runtime/NavigationSystem/Private/NavMesh/PImplRecastNavMesh.cpp
+++ b/Engine/Source/Runtime/NavigationSystem/Private/NavMesh/PImplRecastNavMesh.cpp
@@ -3188,14 +3188,17 @@ void FPImplRecastNavMesh::GetNavMeshTilesIn(const TArray<FBox>& InclusionBounds,
 		TSet<FIntPoint>	TileCoords;	
 		for (const FBox& Bounds : InclusionBounds)
 		{
-			const FVector RcNavMeshOrigin(NavMeshOrigin[0], NavMeshOrigin[1], NavMeshOrigin[2]);
-			const FRcTileBox TileBox(Bounds, RcNavMeshOrigin, TileSize);
-
-			for (int32 y = TileBox.YMin; y <= TileBox.YMax; ++y)
+			if (ensureMsgf(Bounds.IsValid, TEXT("%hs Bounds is not valid"), __FUNCTION__))
 			{
-				for (int32 x = TileBox.XMin; x <= TileBox.XMax; ++x)
+				const FVector RcNavMeshOrigin(NavMeshOrigin[0], NavMeshOrigin[1], NavMeshOrigin[2]);
+				const FRcTileBox TileBox(Bounds, RcNavMeshOrigin, TileSize);
+
+				for (int32 y = TileBox.YMin; y <= TileBox.YMax; ++y)
 				{
-					TileCoords.Add(FIntPoint(x, y));
+					for (int32 x = TileBox.XMin; x <= TileBox.XMax; ++x)
+					{
+						TileCoords.Add(FIntPoint(x, y));
+					}
 				}
 			}
 		}
diff --git a/Engine/Source/Runtime/NavigationSystem/Private/NavMesh/RecastNavMesh.cpp b/Engine/Source/Runtime/NavigationSystem/Private/NavMesh/RecastNavMesh.cpp
--- a/Engine/Source/Runtime/NavigationSystem/Private/NavMesh/RecastNavMesh.cpp
+++ b/Engine/Source/Runtime/NavigationSystem/Private/NavMesh/RecastNavMesh.cpp
@@ -2207,6 +2207,11 @@ void ARecastNavMesh::SetPolyArrayArea(const TArray<FNavPoly>& Polys, TSubclassOf
 
 int32 ARecastNavMesh::ReplaceAreaInTileBounds(const FBox& Bounds, TSubclassOf<UNavArea> OldArea, TSubclassOf<UNavArea> NewArea, bool ReplaceLinks, TArray<NavNodeRef>* OutTouchedNodes)
 {
+	if (!ensureMsgf(Bounds.IsValid, TEXT("%hs Bounds is not valid"), __FUNCTION__))
+	{
+		return 0;
+	}
+
 	int32 PolysTouched = 0;
 
 	if (RecastNavMeshImpl && RecastNavMeshImpl->GetRecastMesh())
@@ -2716,23 +2721,30 @@ void ARecastNavMesh::ApplyWorldOffset(const FVector& InOffset, bool bWorldShift)
 
 void ARecastNavMesh::FillNavigationDataChunkActor(const FBox& QueryBounds, ANavigationDataChunkActor& DataChunkActor, FBox& OutTilesBounds) const
 {
-	if (RecastNavMeshImpl)
+	if (!RecastNavMeshImpl)
 	{
-		UE_LOG(LogNavigation, Verbose, TEXT("%s Bounds pos: (%s)  size: (%s)."), ANSI_TO_TCHAR(__FUNCTION__), *QueryBounds.GetCenter().ToString(), *QueryBounds.GetSize().ToString());
+		return;
+	}
 
-		const TArray<FBox> Boxes({ QueryBounds });
-		TArray<int32> TileIndices;
-		RecastNavMeshImpl->GetNavMeshTilesIn(Boxes, TileIndices);
-		if (!TileIndices.IsEmpty())
-		{
-			// Add a data chunk for this navmesh
-			URecastNavMeshDataChunk* DataChunk = NewObject<URecastNavMeshDataChunk>(&DataChunkActor);
-			DataChunk->NavigationDataName = GetFName();
-			DataChunkActor.GetMutableNavDataChunk().Add(DataChunk);
+	if (!ensureMsgf(QueryBounds.IsValid, TEXT("%hs QueryBounds is not valid"), __FUNCTION__))
+	{
+		return;
+	}
 
-			DataChunk->GetTiles(RecastNavMeshImpl, TileIndices, SupportsRuntimeGeneration() ? EGatherTilesCopyMode::CopyDataAndCacheData : EGatherTilesCopyMode::CopyData);
-			DataChunk->GetTilesBounds(*RecastNavMeshImpl, TileIndices, OutTilesBounds);
-		}
+	UE_LOG(LogNavigation, Verbose, TEXT("%s Bounds pos: (%s)  size: (%s)."), ANSI_TO_TCHAR(__FUNCTION__), *QueryBounds.GetCenter().ToString(), *QueryBounds.GetSize().ToString());
+
+	const TArray<FBox> Boxes({ QueryBounds });
+	TArray<int32> TileIndices;
+	RecastNavMeshImpl->GetNavMeshTilesIn(Boxes, TileIndices);
+	if (!TileIndices.IsEmpty())
+	{
+		// Add a data chunk for this navmesh
+		URecastNavMeshDataChunk* DataChunk = NewObject<URecastNavMeshDataChunk>(&DataChunkActor);
+		DataChunk->NavigationDataName = GetFName();
+		DataChunkActor.GetMutableNavDataChunk().Add(DataChunk);
+
+		DataChunk->GetTiles(RecastNavMeshImpl, TileIndices, SupportsRuntimeGeneration() ? EGatherTilesCopyMode::CopyDataAndCacheData : EGatherTilesCopyMode::CopyData);
+		DataChunk->GetTilesBounds(*RecastNavMeshImpl, TileIndices, OutTilesBounds);
 	}
 }
 
@@ -3657,6 +3669,11 @@ void ARecastNavMesh::RebuildTile(const TArray<FNavMeshDirtyTileElement>& Tiles)
 
 void ARecastNavMesh::DirtyTilesInBounds(const FBox& Bounds)
 {
+	if (!ensureMsgf(Bounds.IsValid, TEXT("%hs Bounds is not valid"), __FUNCTION__))
+	{
+		return;
+	}
+
 	if (HasValidNavmesh() == false)
 	{
 		return;
@@ -3814,6 +3831,11 @@ const dtNavMesh* ARecastNavMesh::GetRecastMesh() const
 //----------------------------------------------------------------------//
 bool ARecastNavMesh::K2_ReplaceAreaInTileBounds(FBox Bounds, TSubclassOf<UNavArea> OldArea, TSubclassOf<UNavArea> NewArea, bool ReplaceLinks)
 {
+	if (!ensureMsgf(Bounds.IsValid, TEXT("%hs Attempting to use Bounds which are not valid"), __FUNCTION__))
+	{
+		return false;
+	}
+
 	bool bReplaced = false;
 #if WITH_RECAST
 	bReplaced = ReplaceAreaInTileBounds(Bounds, OldArea, NewArea, ReplaceLinks) > 0;
diff --git a/Engine/Source/Runtime/NavigationSystem/Private/NavMesh/RecastNavMeshGenerator.cpp b/Engine/Source/Runtime/NavigationSystem/Private/NavMesh/RecastNavMeshGenerator.cpp
--- a/Engine/Source/Runtime/NavigationSystem/Private/NavMesh/RecastNavMeshGenerator.cpp
+++ b/Engine/Source/Runtime/NavigationSystem/Private/NavMesh/RecastNavMeshGenerator.cpp
@@ -5869,6 +5869,11 @@ namespace
 
 bool FRecastNavMeshGenerator::HasDirtyTiles(const FBox& AreaBounds) const
 {
+	if (!ensureMsgf(AreaBounds.IsValid, TEXT("%hs AreaBounds is not valid"), __FUNCTION__))
+	{
+		return false;
+	}
+
 	if (HasDirtyTiles() == false)
 	{
 		return false;
@@ -5892,6 +5897,11 @@ bool FRecastNavMeshGenerator::HasDirtyTiles(const FBox& AreaBounds) const
 
 int32 FRecastNavMeshGenerator::GetDirtyTilesCount(const FBox& AreaBounds) const
 {
+	if (!ensureMsgf(AreaBounds.IsValid, TEXT("%hs AreaBounds is not valid"), __FUNCTION__))
+	{
+		return 0;
+	}
+
 	const FVector::FReal TileSizeInWorldUnits = Config.GetTileSizeUU();
 	const FRcTileBox TileBox(AreaBounds, RcNavMeshOrigin, TileSizeInWorldUnits);
 
@@ -5962,6 +5972,11 @@ void FRecastNavMeshGenerator::MarkDirtyTiles(const TArray<FNavigationDirtyArea>&
 #endif
 	for (const FNavigationDirtyArea& DirtyArea : DirtyAreas)
 	{
+		if (!ensureMsgf(DirtyArea.Bounds.IsValid, TEXT("%hs Attempting to use DirtyArea.Bounds which are not valid. SourceObject: %s"), __FUNCTION__, *GetFullNameSafe(DirtyArea.OptionalSourceObject.Get())))
+		{
+			continue;
+		}
+
 		// Static navmeshes accept only area modifiers updates
 		if (bGameStaticNavMesh && (!DirtyArea.HasFlag(ENavigationDirtyFlag::DynamicModifier) || DirtyArea.HasFlag(ENavigationDirtyFlag::NavigationBounds)))
 		{
@@ -6027,7 +6042,10 @@ void FRecastNavMeshGenerator::MarkDirtyTiles(const TArray<FNavigationDirtyArea>&
 		}
 		else
 		{
-			SubAreaBoundsArray.Add(AdjustedAreaBounds);
+			if (ensureMsgf(AdjustedAreaBounds.IsValid, TEXT("%hs Attempting to use AdjustedAreaBounds which are not valid"), __FUNCTION__))
+			{
+				SubAreaBoundsArray.Add(AdjustedAreaBounds);
+			}
 		}
 
 		uint32 PendingTilesMarked = 0;
