Ignore dirty areas originating from building actors spawned from instance actors (navmesh is expected to be already up to date from the instance actor being replaced).

diff --git a/Engine/Source/Runtime/NavigationSystem/Private/NavigationDirtyAreasController.cpp b/Engine/Source/Runtime/NavigationSystem/Private/NavigationDirtyAreasController.cpp
--- a/Engine/Source/Runtime/NavigationSystem/Private/NavigationDirtyAreasController.cpp
+++ b/Engine/Source/Runtime/NavigationSystem/Private/NavigationDirtyAreasController.cpp
@@ -90,6 +90,15 @@ void FNavigationDirtyAreasController::AddArea(const FBox& NewArea, const int32 F
 		}		
 	}
 
+	if (ShouldSkipObjectPredicate.IsBound() && ObjectProviderFunc)
+	{
+		UObject* Object = ObjectProviderFunc();
+		if (Object && ShouldSkipObjectPredicate.Execute(*Object))
+		{
+			return;
+		}
+	}
+
 #if !UE_BUILD_SHIPPING
 	auto DumpExtraInfo = [ObjectProviderFunc, DebugReason, BoundsSize, NewArea]() {
 		FString ObjectInfo;
@@ -104,9 +113,9 @@ void FNavigationDirtyAreasController::AddArea(const FBox& NewArea, const int32 F
 			}
 		}
 
-		return FString::Printf(TEXT("Object: %s (from: %s)%s | Bounds: %s"),
-			*GetFullNameSafe(Object),
+		return FString::Printf(TEXT("From: %s | Object: %s %s | Bounds: %s"),
 			*DebugReason.ToString(),
+			*GetFullNameSafe(Object),
 			*ObjectInfo,
 			*BoundsSize.ToString());
 	};
diff --git a/Engine/Source/Runtime/NavigationSystem/Public/NavigationDirtyAreasController.h b/Engine/Source/Runtime/NavigationSystem/Public/NavigationDirtyAreasController.h
--- a/Engine/Source/Runtime/NavigationSystem/Public/NavigationDirtyAreasController.h
+++ b/Engine/Source/Runtime/NavigationSystem/Public/NavigationDirtyAreasController.h
@@ -80,4 +80,7 @@ public:
 #if !UE_BUILD_SHIPPING
 	bool HadDirtyAreasReportedWhileAccumulationLocked() const { return bCanAccumulateDirtyAreas == false && bDirtyAreasReportedWhileAccumulationLocked; }
 #endif // UE_BUILD_SHIPPING
+
+	DECLARE_DELEGATE_RetVal_OneParam(bool, FSkipObjectSignature, UObject& /*Object*/);
+	FSkipObjectSignature ShouldSkipObjectPredicate;  
 };
