[Navigation] Add total number of elements when rendering NavOctree details

diff --git a/Engine/Source/Runtime/NavigationSystem/Private/NavMesh/NavMeshRenderingComponent.cpp b/Engine/Source/Runtime/NavigationSystem/Private/NavMesh/NavMeshRenderingComponent.cpp
--- a/Engine/Source/Runtime/NavigationSystem/Private/NavMesh/NavMeshRenderingComponent.cpp
+++ b/Engine/Source/Runtime/NavigationSystem/Private/NavMesh/NavMeshRenderingComponent.cpp
@@ -1078,9 +1078,10 @@ void FNavMeshSceneProxyData::GatherData(const ARecastNavMesh* NavMesh, int32 InN
 			{
 				TArray<FVector> CollidingVerts;
 				TArray<uint32> CollidingIndices;
+				int32 NumElements = 0;
 
 				FBox CurrentNodeBoundsBox;
-				NavOctree->FindElementsWithPredicate([bGatherOctree, bGatherOctreeDetails, bGatherPathCollidingGeometry, this, &CurrentNodeBoundsBox, NavOctree](FNavigationOctree::FNodeIndex /*ParentNodeIndex*/, FNavigationOctree::FNodeIndex NodeIndex, const FBoxCenterAndExtent& NodeBounds)
+				NavOctree->FindElementsWithPredicate([bGatherOctree, bGatherOctreeDetails, bGatherPathCollidingGeometry, this, &CurrentNodeBoundsBox, NavOctree, &NumElements](FNavigationOctree::FNodeIndex /*ParentNodeIndex*/, FNavigationOctree::FNodeIndex NodeIndex, const FBoxCenterAndExtent& NodeBounds)
 				{
 					if (bGatherOctree)
 					{
@@ -1088,7 +1089,9 @@ void FNavMeshSceneProxyData::GatherData(const ARecastNavMesh* NavMesh, int32 InN
 
 						if (bGatherOctreeDetails)
 						{
-							DebugLabels.Emplace(NodeBounds.Center, FString::Printf(TEXT("%d elements"), NavOctree->GetElementsForNode(NodeIndex).Num()));
+							const int32 NumElementsInNode = NavOctree->GetElementsForNode(NodeIndex).Num();
+							NumElements += NumElementsInNode;
+							DebugLabels.Emplace(NodeBounds.Center, FString::Printf(TEXT("%d elements"), NumElementsInNode));
 						}
 					}
 
@@ -1127,7 +1130,12 @@ void FNavMeshSceneProxyData::GatherData(const ARecastNavMesh* NavMesh, int32 InN
 						}
 					}
 				});
-			
+
+				if (bGatherOctreeDetails)
+				{
+					DebugLabels.Emplace(FString::Printf(TEXT("Total: %d elements"), NumElements));
+				}
+
 				if (CollidingVerts.Num())
 				{
 					FDebugMeshData DebugMeshData;
