Addition of the average tile generation count per second in LogNavigationHistory.

diff --git a/Engine/Source/Runtime/NavigationSystem/Private/NavigationSystem.cpp b/Engine/Source/Runtime/NavigationSystem/Private/NavigationSystem.cpp
--- a/Engine/Source/Runtime/NavigationSystem/Private/NavigationSystem.cpp
+++ b/Engine/Source/Runtime/NavigationSystem/Private/NavigationSystem.cpp
@@ -352,6 +352,7 @@ void FNavRegenTimeSliceManager::ResetTileHistoryData(const TArray<TObjectPtr<ANa
 	{
 		HistoryData.Empty();
 	}
+	TileHistoryStartTime = FPlatformTime::Seconds();
 }
 
 void FNavRegenTimeSliceManager::PushTileHistoryData(const int32 NavDataIndex, const FTileHistoryData& TileData)
@@ -498,6 +499,7 @@ void FNavRegenTimeSliceManager::LogTileStatistics(const TArray<TObjectPtr<ANavig
 	{
 		// Log median tile processing time every 60 frames.
 		const bool bLog = GFrameCounter % 60 == 0;
+		const double HistoryDuration = FPlatformTime::Seconds() - TileHistoryStartTime;
 		for (int32 NavDataIndex = 0; bLog && NavDataIndex < NavDataSet.Num(); ++NavDataIndex)
 		{
 			if (TileHistoryData.IsValidIndex(NavDataIndex))
@@ -517,8 +519,9 @@ void FNavRegenTimeSliceManager::LogTileStatistics(const TArray<TObjectPtr<ANavig
 					const double MedianWaitTimeMs = HistoryData[MedianIndex].TileWaitTime * 1000.f;
 					const double HighWaitTimeMs = HistoryData[HighIndex].TileWaitTime * 1000.f;
 					
-					UE_LOG(LogNavigationHistory, Log, TEXT("%-35s Median tile stats: regen time: %2.2f ms, regen frames %lld, wait time: %4.f ms (high regen time: %2.2f ms, high wait time: %4.f ms) regen count: %i"),
-						*GetNameSafe(NavDataSet[NavDataIndex]), MedianRegenTimeMs, MedianRegenFrames, MedianWaitTimeMs, HighRegenTimeMs, HighWaitTimeMs, HistoryData.Num());
+					UE_LOG(LogNavigationHistory, Log, TEXT("%-35s Median tile stats: regen time: %2.2f ms, regen frames %lld, wait time: %4.f ms (high regen time: %2.2f ms, high wait time: %4.f ms) regen count: %i, regen/s: %0.2f"),
+						*GetNameSafe(NavDataSet[NavDataIndex]), MedianRegenTimeMs, MedianRegenFrames, MedianWaitTimeMs, HighRegenTimeMs, HighWaitTimeMs,
+						HistoryData.Num(), HistoryData.Num()/HistoryDuration);
 				}
 			}
 		}
diff --git a/Engine/Source/Runtime/NavigationSystem/Public/NavigationSystem.h b/Engine/Source/Runtime/NavigationSystem/Public/NavigationSystem.h
--- a/Engine/Source/Runtime/NavigationSystem/Public/NavigationSystem.h
+++ b/Engine/Source/Runtime/NavigationSystem/Public/NavigationSystem.h
@@ -264,6 +264,7 @@ protected:
 #if !UE_BUILD_SHIPPING
 	/** Tile processing time per NavDataIndex */
 	TArray<TArray<FTileHistoryData>> TileHistoryData;
+	double TileHistoryStartTime = 0;
 #endif // UE_BUILD_SHIPPING
 
 	/** If there are enough tiles to process this in the Min Time Slice Duration */
