When rebuilding the NavigationData, we will now wait for assets to be done compiling before calling the navigation data generator. This will fix an issue where some assets that were just loaded (e.g. streamed static meshes) weren't compiled yet when generating the data, which resulted in an incomplete generation.

#rnx

diff --git a/Engine/Source/Runtime/NavigationSystem/Private/NavigationData.cpp b/Engine/Source/Runtime/NavigationSystem/Private/NavigationData.cpp
--- a/Engine/Source/Runtime/NavigationSystem/Private/NavigationData.cpp
+++ b/Engine/Source/Runtime/NavigationSystem/Private/NavigationData.cpp
@@ -1,6 +1,8 @@
 // Copyright Epic Games, Inc. All Rights Reserved.
 
 #include "NavigationData.h"
+
+#include "AssetCompilingManager.h"
 #include "EngineGlobals.h"
 #include "AI/Navigation/NavAgentInterface.h"
 #include "Components/PrimitiveComponent.h"
@@ -540,6 +542,7 @@ void ANavigationData::RebuildAll()
 {
 	const double LoadTime = FPlatformTime::Seconds();
 	LoadBeforeGeneratorRebuild();
+	FAssetCompilingManager::Get().FinishAllCompilation();
 	UE_LOG(LogNavigationDataBuild, Display, TEXT("   %s load time: %.2fs"), ANSI_TO_TCHAR(__FUNCTION__), (FPlatformTime::Seconds() - LoadTime));
 	
 	ConditionalConstructGenerator(); //recreate generator
