Fix for NavLink indeterminism during cook.

diff --git a/Engine/Source/Runtime/Engine/Classes/AI/Navigation/NavigationTypes.h b/Engine/Source/Runtime/Engine/Classes/AI/Navigation/NavigationTypes.h
--- a/Engine/Source/Runtime/Engine/Classes/AI/Navigation/NavigationTypes.h
+++ b/Engine/Source/Runtime/Engine/Classes/AI/Navigation/NavigationTypes.h
@@ -162,10 +162,14 @@ struct FNavLinkAuxiliaryId
 	bool operator==(const FNavLinkAuxiliaryId& Other) const { return Id == Other.Id; }
 	bool operator!=(const FNavLinkAuxiliaryId& Other) const { return !this->operator==(Other); }
 
+
+	UE_DEPRECATED(5.4, "This function is not deterministic in all instances during cooking. Use the version which takes PathName.")
+	static ENGINE_API FNavLinkAuxiliaryId GenerateUniqueAuxiliaryId();
+
 	/**
 	 * Helper function: returns unique Auxiliary ID for custom links.
 	 **/
-	static ENGINE_API FNavLinkAuxiliaryId GenerateUniqueAuxiliaryId();
+	static ENGINE_API FNavLinkAuxiliaryId GenerateUniqueAuxiliaryId(FStringView PathName);
 
 private:
 	FNavLinkAuxiliaryId(uint64 InId)
diff --git a/Engine/Source/Runtime/Engine/Private/AI/Navigation/NavigationTypes.cpp b/Engine/Source/Runtime/Engine/Private/AI/Navigation/NavigationTypes.cpp
--- a/Engine/Source/Runtime/Engine/Private/AI/Navigation/NavigationTypes.cpp
+++ b/Engine/Source/Runtime/Engine/Private/AI/Navigation/NavigationTypes.cpp
@@ -327,6 +327,15 @@ FNavLinkAuxiliaryId FNavLinkAuxiliaryId::GenerateUniqueAuxiliaryId()
 	return FNavLinkAuxiliaryId(AuxiliaryId);
 }
 
+FNavLinkAuxiliaryId FNavLinkAuxiliaryId::GenerateUniqueAuxiliaryId(FStringView PathName)
+{
+	check(!PathName.IsEmpty());
+
+	const uint64 AuxiliaryId = UE::Navigation::NavLinkIdHelpers::Private::MakeIdFromGUID(FGuid::NewDeterministicGuid(PathName));
+
+	return FNavLinkAuxiliaryId(AuxiliaryId);
+}
+
 FNavLinkId FNavLinkId::GenerateUniqueId(FNavLinkAuxiliaryId AuxiliaryId, FGuid ActorInstanceGuid)
 {
 	// Apply NavLinkIdBitMask to differentiate Legacy Ids (that do not have the mask set).
diff --git a/Engine/Source/Runtime/NavigationSystem/Private/NavLinkCustomComponent.cpp b/Engine/Source/Runtime/NavigationSystem/Private/NavLinkCustomComponent.cpp
--- a/Engine/Source/Runtime/NavigationSystem/Private/NavLinkCustomComponent.cpp
+++ b/Engine/Source/Runtime/NavigationSystem/Private/NavLinkCustomComponent.cpp
@@ -53,7 +53,8 @@ PRAGMA_ENABLE_DEPRECATION_WARNINGS
 
 	if (!HasAnyFlags(RF_ClassDefaultObject))
 	{
-		AuxiliaryCustomLinkId = FNavLinkAuxiliaryId::GenerateUniqueAuxiliaryId();
+		const FString PathName = GetPathName();
+		AuxiliaryCustomLinkId = FNavLinkAuxiliaryId::GenerateUniqueAuxiliaryId(*PathName);
 	}
 }
 
@@ -136,7 +137,9 @@ void UNavLinkCustomComponent::PostEditImport()
 
 	// Generate a new AuxiliarLinkUserId and set CustomLinkId to Invalid, this is then inline with the constructor in this regard.
 	// CustomLinkId is set to a valid Id later in OnRegister().
-	AuxiliaryCustomLinkId = FNavLinkAuxiliaryId::GenerateUniqueAuxiliaryId();
+	const FString PathName = GetPathName();
+	AuxiliaryCustomLinkId = FNavLinkAuxiliaryId::GenerateUniqueAuxiliaryId(*PathName);
+
 	CustomLinkId = FNavLinkId::Invalid;
 }
 #endif // WITH_EDITOR
