[Navigation] use virtual method 'GetWorld' instead of 'GetTypedOuter' for navrelevant object registration.

diff --git a/Engine/Source/Runtime/NavigationSystem/Private/NavigationSystem.cpp b/Engine/Source/Runtime/NavigationSystem/Private/NavigationSystem.cpp
--- a/Engine/Source/Runtime/NavigationSystem/Private/NavigationSystem.cpp
+++ b/Engine/Source/Runtime/NavigationSystem/Private/NavigationSystem.cpp
@@ -3091,7 +3091,7 @@ void UNavigationSystemV1::OnNavRelevantObjectRegistered(UObject& Object)
 	SCOPE_CYCLE_COUNTER(STAT_DebugNavOctree);
 	if (INavRelevantInterface* NavInterface = Cast<INavRelevantInterface>(&Object))
 	{
-		UWorld* World = Object.GetTypedOuter<UWorld>();
+		UWorld* World = Object.GetWorld();
 		if (UNavigationSystemV1* NavSys = FNavigationSystem::GetCurrent<UNavigationSystemV1>(World))
 		{
 			NavSys->RegisterNavOctreeElement(&Object, NavInterface, FNavigationOctreeController::OctreeUpdate_Default);
@@ -3132,7 +3132,7 @@ void UNavigationSystemV1::OnNavRelevantObjectUnregistered(UObject& Object)
 	SCOPE_CYCLE_COUNTER(STAT_DebugNavOctree);
 	if (INavRelevantInterface* NavInterface = Cast<INavRelevantInterface>(&Object))
 	{
-		UWorld* World = Object.GetTypedOuter<UWorld>();
+		UWorld* World = Object.GetWorld();
 		if (UNavigationSystemV1* NavSys = FNavigationSystem::GetCurrent<UNavigationSystemV1>(World))
 		{
 			NavSys->UnregisterNavOctreeElement(&Object, NavInterface, FNavigationOctreeController::OctreeUpdate_Default);
@@ -3249,7 +3249,7 @@ void UNavigationSystemV1::UpdateNavRelevantObjectInNavOctree(UObject& Object)
 
 	if (INavRelevantInterface* NavElement = Cast<INavRelevantInterface>(&Object))
 	{
-		if (UNavigationSystemV1* NavSys = FNavigationSystem::GetCurrent<UNavigationSystemV1>(Object.GetTypedOuter<UWorld>()))
+		if (UNavigationSystemV1* NavSys = FNavigationSystem::GetCurrent<UNavigationSystemV1>(Object.GetWorld()))
 		{
 			if (NavElement->IsNavigationRelevant())
 			{
