Navmesh: fix debug rendering missing some layers when there are empty layers in a tile

diff --git a/Engine/Source/Runtime/NavigationSystem/Private/NavMesh/PImplRecastNavMesh.cpp b/Engine/Source/Runtime/NavigationSystem/Private/NavMesh/PImplRecastNavMesh.cpp
--- a/Engine/Source/Runtime/NavigationSystem/Private/NavMesh/PImplRecastNavMesh.cpp
+++ b/Engine/Source/Runtime/NavigationSystem/Private/NavMesh/PImplRecastNavMesh.cpp
@@ -2771,14 +2771,15 @@ bool FPImplRecastNavMesh::GetDebugGeometryForTile(FRecastDebugGeometry& OutGeome
 	}
 	else if (bIsGenerationRestrictedToActiveTiles)
 	{
+		TArray<const dtMeshTile*> Tiles;
 		const TArray<FIntPoint>& ActiveTiles = NavMeshOwner->GetActiveTiles();
 		for (const FIntPoint& TileLocation : ActiveTiles)
 		{
-			const int32 LayersCount = ConstNavMesh->getTileCountAt(TileLocation.X, TileLocation.Y);
-
-			for (int32 Layer = 0; Layer < LayersCount; ++Layer)
+			Tiles.Reset();
+			Tiles.AddZeroed(ConstNavMesh->getTileCountAt(TileLocation.X, TileLocation.Y));
+			ConstNavMesh->getTilesAt(TileLocation.X, TileLocation.Y, Tiles.GetData(), Tiles.Num());
+			for (const dtMeshTile* Tile : Tiles)
 			{
-				dtMeshTile const* const Tile = ConstNavMesh->getTileAt(TileLocation.X, TileLocation.Y, Layer);
 				ComputeSizeToReserve(Tile, NumVertsToReserve, NumIndicesToReserve);
 			}
 		}
@@ -2788,11 +2789,11 @@ bool FPImplRecastNavMesh::GetDebugGeometryForTile(FRecastDebugGeometry& OutGeome
 		uint32 VertBase = OutGeometry.MeshVerts.Num();
 		for (const FIntPoint& TileLocation : ActiveTiles)
 		{
-			const int32 LayersCount = ConstNavMesh->getTileCountAt(TileLocation.X, TileLocation.Y);
-
-			for (int32 Layer = 0; Layer < LayersCount; ++Layer)
+			Tiles.Reset();
+			Tiles.AddZeroed(ConstNavMesh->getTileCountAt(TileLocation.X, TileLocation.Y));
+			ConstNavMesh->getTilesAt(TileLocation.X, TileLocation.Y, Tiles.GetData(), Tiles.Num());
+			for (const dtMeshTile* Tile : Tiles)
 			{
-				dtMeshTile const* const Tile = ConstNavMesh->getTileAt(TileLocation.X, TileLocation.Y, Layer);
 				if (Tile != nullptr && Tile->header != nullptr)
 				{
 					VertBase += GetTilesDebugGeometry(Generator, *Tile, VertBase, OutGeometry, INDEX_NONE, ForbiddenFlags);
