[NavModifierVolume] Modifier volume actors getting their position updated via their parent Instanced Level Actor in the Editor now update their navigation data.

diff --git a/Engine/Source/Runtime/NavigationSystem/Private/NavModifierVolume.cpp b/Engine/Source/Runtime/NavigationSystem/Private/NavModifierVolume.cpp
--- a/Engine/Source/Runtime/NavigationSystem/Private/NavModifierVolume.cpp
+++ b/Engine/Source/Runtime/NavigationSystem/Private/NavModifierVolume.cpp
@@ -68,6 +68,30 @@ void ANavModifierVolume::BeginDestroy()
 }
 
 #if WITH_EDITOR
+
+void ANavModifierVolume::PostRegisterAllComponents()
+{
+	Super::PostRegisterAllComponents();
+
+	if (RootComponent)
+	{
+		RootComponent->TransformUpdated.AddLambda([this](USceneComponent*, EUpdateTransformFlags, ETeleportType)
+			{
+				FNavigationSystem::UpdateActorData(*this);
+			});
+	}
+}
+
+void ANavModifierVolume::PostUnregisterAllComponents()
+{
+	if (RootComponent)
+	{
+		RootComponent->TransformUpdated.RemoveAll(this);
+	}
+
+	Super::PostUnregisterAllComponents();
+}
+
 // This function is only called if GIsEditor == true
 void ANavModifierVolume::OnNavAreaRegistered(const UWorld& World, const UClass* NavAreaClass)
 {
diff --git a/Engine/Source/Runtime/NavigationSystem/Public/NavModifierVolume.h b/Engine/Source/Runtime/NavigationSystem/Public/NavModifierVolume.h
--- a/Engine/Source/Runtime/NavigationSystem/Public/NavModifierVolume.h
+++ b/Engine/Source/Runtime/NavigationSystem/Public/NavModifierVolume.h
@@ -61,6 +61,9 @@ protected:
 	NAVIGATIONSYSTEM_API virtual void BeginDestroy() override;
 
 #if WITH_EDITOR
+	NAVIGATIONSYSTEM_API virtual void PostRegisterAllComponents() override;
+	NAVIGATIONSYSTEM_API virtual void PostUnregisterAllComponents() override;
+
 	NAVIGATIONSYSTEM_API void OnNavAreaRegistered(const UWorld& World, const UClass* NavAreaClass);
 	NAVIGATIONSYSTEM_API void OnNavAreaUnregistered(const UWorld& World, const UClass* NavAreaClass);
 #endif
