Fixed uninitialized property warning by removing obsolete property
#rnx

diff --git a/Engine/Source/Runtime/NavigationSystem/Public/NavMesh/RecastNavMesh.h b/Engine/Source/Runtime/NavigationSystem/Public/NavMesh/RecastNavMesh.h
--- a/Engine/Source/Runtime/NavigationSystem/Public/NavMesh/RecastNavMesh.h
+++ b/Engine/Source/Runtime/NavigationSystem/Public/NavMesh/RecastNavMesh.h
@@ -379,10 +379,6 @@ struct FRecastNavMeshGenerationProperties
 	UPROPERTY(EditAnywhere, Category = Generation, meta = (ClampMin = "1", UIMin = "1"), AdvancedDisplay)
 	int32 TileNumberHardLimit;
 
-	/** filtering methode used for filtering ledge slopes */
-	UPROPERTY(EditAnywhere, Category = Generation, AdvancedDisplay)
-	ENavigationLedgeSlopeFilterMode LedgeFilterMode;
-	
 	/** partitioning method for creating navmesh polys */
 	UPROPERTY(EditAnywhere, Category = Generation, AdvancedDisplay)
 	TEnumAsByte<ERecastPartitioning::Type> RegionPartitioning;
@@ -830,12 +826,12 @@ class ARecastNavMesh : public ANavigationData
 	UPROPERTY(config)
 	float DefaultDrawDistance;
 
-	/** specifes default limit to A* nodes used when performing navigation queries. 
+	/** specifies default limit to A* nodes used when performing navigation queries. 
 	 *	Can be overridden by passing custom FNavigationQueryFilter */
 	UPROPERTY(config)
 	float DefaultMaxSearchNodes;
 
-	/** specifes default limit to A* nodes used when performing hierarchical navigation queries. */
+	/** specifies default limit to A* nodes used when performing hierarchical navigation queries. */
 	UPROPERTY(config)
 	float DefaultMaxHierarchicalSearchNodes;
 
