Fix "Looking for NavData using invalid FNavAgentProperties." warning spam.

diff --git a/Engine/Source/Runtime/NavigationSystem/Private/NavigationSystem.cpp b/Engine/Source/Runtime/NavigationSystem/Private/NavigationSystem.cpp
--- a/Engine/Source/Runtime/NavigationSystem/Private/NavigationSystem.cpp
+++ b/Engine/Source/Runtime/NavigationSystem/Private/NavigationSystem.cpp
@@ -2035,8 +2035,9 @@ const ANavigationData* UNavigationSystemV1::GetNavDataForProps(const FNavAgentPr
 	{
 		return MainNavData;
 	}
-	
-	UE_CLOG(!AgentProperties.IsValid(), LogNavigation, Warning, TEXT("Looking for NavData using invalid FNavAgentProperties."));
+
+	// Because an invalid AgentProperties uses -1 values the code below is able to match the PreferredNavData.
+	UE_CLOG(!(AgentProperties.IsValid() || AgentProperties.PreferredNavData.IsValid()), LogNavigation, Warning, TEXT("Looking for NavData using invalid FNavAgentProperties."));
 	
 	const TWeakObjectPtr<ANavigationData>* NavDataForAgent = AgentToNavDataMap.Find(AgentProperties);
 	const ANavigationData* NavDataInstance = NavDataForAgent ? NavDataForAgent->Get() : nullptr;
