Added extra logs when Navigation DirtyAreas are being reset.

diff --git a/Engine/Source/Runtime/NavigationSystem/Private/NavigationDirtyAreasController.cpp b/Engine/Source/Runtime/NavigationSystem/Private/NavigationDirtyAreasController.cpp
--- a/Engine/Source/Runtime/NavigationSystem/Private/NavigationDirtyAreasController.cpp
+++ b/Engine/Source/Runtime/NavigationSystem/Private/NavigationDirtyAreasController.cpp
@@ -180,6 +180,8 @@ bool FNavigationDirtyAreasController::ShouldReportOversizedDirtyArea() const
 void FNavigationDirtyAreasController::Reset()
 {
 	// discard all pending dirty areas, we are going to rebuild navmesh anyway 
+	UE_LOG(LogNavigationDirtyArea, VeryVerbose, TEXT("%hs: Reseting All Dirty Areas. DirtyAreas.Num = [%d]"),__FUNCTION__, DirtyAreas.Num());
+
 	DirtyAreas.Reset();
 #if !UE_BUILD_SHIPPING
 	bDirtyAreasReportedWhileAccumulationLocked = false;
diff --git a/Engine/Source/Runtime/NavigationSystem/Private/NavigationSystem.cpp b/Engine/Source/Runtime/NavigationSystem/Private/NavigationSystem.cpp
--- a/Engine/Source/Runtime/NavigationSystem/Private/NavigationSystem.cpp
+++ b/Engine/Source/Runtime/NavigationSystem/Private/NavigationSystem.cpp
@@ -917,6 +917,7 @@ void UNavigationSystemV1::ConstructNavOctree()
 bool UNavigationSystemV1::ConditionalPopulateNavOctree()
 {
 	// Discard all navigation updates caused by octree construction
+	UE_LOG(LogNavigationDirtyArea, VeryVerbose, TEXT("%hs: Reseting Dirty Areas added during octree construction. DirtyAreas.Num = [%d]."), __FUNCTION__, DefaultDirtyAreasController.DirtyAreas.Num());
 	TGuardValue<TArray<FNavigationDirtyArea>> DirtyGuard(DefaultDirtyAreasController.DirtyAreas, TArray<FNavigationDirtyArea>());
 
 	// See if any of registered navigation data need navoctree
