Fix CIS: Non-void function should return a value.

diff --git a/Engine/Source/Runtime/NavigationSystem/Public/NavMesh/RecastNavMeshGenerator.h b/Engine/Source/Runtime/NavigationSystem/Public/NavMesh/RecastNavMeshGenerator.h
--- a/Engine/Source/Runtime/NavigationSystem/Public/NavMesh/RecastNavMeshGenerator.h
+++ b/Engine/Source/Runtime/NavigationSystem/Public/NavMesh/RecastNavMeshGenerator.h
@@ -973,7 +973,7 @@ protected:
 	UE_DEPRECATED(5.3, "Use ConstructTileGeneratorImpl instead.")
 	TSharedRef<T> ConstuctTileGeneratorImpl(const FIntPoint& Coord, const TArray<FBox>& DirtyAreas, const double PendingTileCreationTime = 0.)
 	{
-		ConstructTileGeneratorImpl<T>(Coord, DirtyAreas, PendingTileCreationTime);
+		return ConstructTileGeneratorImpl<T>(Coord, DirtyAreas, PendingTileCreationTime);
 	}
 	
 	template <typename T>
