[Navmesh] Fix navmesh not updating below 0 when using invoker's seeds

diff --git a/Engine/Source/Runtime/NavigationSystem/Private/NavigationSystem.cpp b/Engine/Source/Runtime/NavigationSystem/Private/NavigationSystem.cpp
--- a/Engine/Source/Runtime/NavigationSystem/Private/NavigationSystem.cpp
+++ b/Engine/Source/Runtime/NavigationSystem/Private/NavigationSystem.cpp
@@ -5102,8 +5102,8 @@ void UNavigationSystemV1::UpdateInvokers()
 				for (const FVector2D SeedLocation : SeedLocations)
 				{
 					InvokersSeedBounds.Emplace(
-						FVector(SeedLocation.X-InvokersMaximumDistanceFromSeed, SeedLocation.Y-InvokersMaximumDistanceFromSeed, DBL_MIN),
-						FVector(SeedLocation.X+InvokersMaximumDistanceFromSeed, SeedLocation.Y+InvokersMaximumDistanceFromSeed, DBL_MAX));
+						FVector(SeedLocation.X-InvokersMaximumDistanceFromSeed, SeedLocation.Y-InvokersMaximumDistanceFromSeed, TNumericLimits<double>::Lowest()),
+						FVector(SeedLocation.X+InvokersMaximumDistanceFromSeed, SeedLocation.Y+InvokersMaximumDistanceFromSeed, TNumericLimits<double>::Max()));
 				}
 			}
 			
