[Navigation] When updating invokers, use 3D invoker seed locations to compare against the maximum seed distance.
Update InvokersSeedBounds with the vertical limits as well.

diff --git a/Engine/Source/Runtime/NavigationSystem/Private/NavigationSystem.cpp b/Engine/Source/Runtime/NavigationSystem/Private/NavigationSystem.cpp
--- a/Engine/Source/Runtime/NavigationSystem/Private/NavigationSystem.cpp
+++ b/Engine/Source/Runtime/NavigationSystem/Private/NavigationSystem.cpp
@@ -3793,7 +3793,19 @@ void UNavigationSystemV1::GatherNavigationBounds()
 	}
 }
 
+// Deprecated
 void UNavigationSystemV1::GetInvokerSeedLocations(const UWorld& InWorld, TArray<FVector2D, TInlineAllocator<32>>& OutSeedLocations)
+{
+	TArray<FVector, TInlineAllocator<32>> Locations;
+	GetInvokerSeedLocations(InWorld, Locations);
+
+	for (const FVector Location : Locations)
+	{
+		OutSeedLocations.Add(FVector2D(Location));	
+	}
+}
+
+void UNavigationSystemV1::GetInvokerSeedLocations(const UWorld& InWorld, TArray<FVector, TInlineAllocator<32>>& OutSeedLocations)
 {
 	for (FConstPlayerControllerIterator PlayerIt = InWorld.GetPlayerControllerIterator(); PlayerIt; ++PlayerIt)
 	{
@@ -3802,13 +3814,11 @@ void UNavigationSystemV1::GetInvokerSeedLocations(const UWorld& InWorld, TArray<
 		{
 			if (PlayerController->GetPawn())
 			{
-				const FVector2D SeedLoc(PlayerController->GetPawn()->GetActorLocation());
-				OutSeedLocations.Add(SeedLoc);
+				OutSeedLocations.Add(PlayerController->GetPawn()->GetActorLocation());
 			}
 			else if (PlayerController->PlayerCameraManager)
 			{
-				const FVector2D SeedLoc(PlayerController->PlayerCameraManager->GetCameraLocation());
-				OutSeedLocations.Add(SeedLoc);
+				OutSeedLocations.Add(PlayerController->PlayerCameraManager->GetCameraLocation());
 			}
 		}
 	}
@@ -5093,17 +5103,17 @@ void UNavigationSystemV1::UpdateInvokers()
 			QUICK_SCOPE_CYCLE_COUNTER(STAT_NavSys_Clusterize);
 
 			const bool bCheckMaximumDistanceFromSeeds = (InvokersMaximumDistanceFromSeed != -1) && World->IsGameWorld();
-			TArray<FVector2D, TInlineAllocator<32>> SeedLocations;
+			TArray<FVector, TInlineAllocator<32>> SeedLocations;
 			if (bCheckMaximumDistanceFromSeeds)
 			{
 				GetInvokerSeedLocations(*World, SeedLocations);
 
 				// Fill seed bounds
-				for (const FVector2D SeedLocation : SeedLocations)
+				for (const FVector SeedLocation : SeedLocations)
 				{
 					InvokersSeedBounds.Emplace(
-						FVector(SeedLocation.X-InvokersMaximumDistanceFromSeed, SeedLocation.Y-InvokersMaximumDistanceFromSeed, TNumericLimits<double>::Lowest()),
-						FVector(SeedLocation.X+InvokersMaximumDistanceFromSeed, SeedLocation.Y+InvokersMaximumDistanceFromSeed, TNumericLimits<double>::Max()));
+						FVector(SeedLocation.X-InvokersMaximumDistanceFromSeed, SeedLocation.Y-InvokersMaximumDistanceFromSeed, SeedLocation.Z-InvokersMaximumDistanceFromSeed),
+						FVector(SeedLocation.X+InvokersMaximumDistanceFromSeed, SeedLocation.Y+InvokersMaximumDistanceFromSeed, SeedLocation.Z+InvokersMaximumDistanceFromSeed));
 				}
 			}
 			
@@ -5131,9 +5141,9 @@ void UNavigationSystemV1::UpdateInvokers()
 						const double CheckDistanceSq = FMath::Square(InvokersMaximumDistanceFromSeed + GenerationRadius);
 
 						// Check if the invoker is close enough
-						for (const FVector2D SeedLocation : SeedLocations)
+						for (const FVector SeedLocation : SeedLocations)
 						{
-							const double InvokerDistanceToSeedSq = FVector2D::DistSquared(SeedLocation, FVector2D(ActorLocation));
+							const double InvokerDistanceToSeedSq = FVector::DistSquared(SeedLocation, ActorLocation);
 							if (InvokerDistanceToSeedSq <= CheckDistanceSq)
 							{
 								bKeep = true;
diff --git a/Engine/Source/Runtime/NavigationSystem/Public/NavigationSystem.h b/Engine/Source/Runtime/NavigationSystem/Public/NavigationSystem.h
--- a/Engine/Source/Runtime/NavigationSystem/Public/NavigationSystem.h
+++ b/Engine/Source/Runtime/NavigationSystem/Public/NavigationSystem.h
@@ -871,6 +871,9 @@ protected:
 
 	/** Get seed locations for invokers, @see InvokersMaximumDistanceFromSeed 
 	 *	By default these are the player pawn locations. If a player controller has no pawn assigned, the player's camera location will be used instead. */ 
+	NAVIGATIONSYSTEM_API virtual void GetInvokerSeedLocations(const UWorld& InWorld, TArray<FVector, TInlineAllocator<32>>& OutSeedLocations);
+
+	UE_DEPRECATED(5.4, "This function is deprecated. Use GetInvokerSeedLocations using TArray of FVector.")
 	NAVIGATIONSYSTEM_API virtual void GetInvokerSeedLocations(const UWorld& InWorld, TArray<FVector2D, TInlineAllocator<32>>& OutSeedLocations);
 
 	/** @return pointer to ANavigationData instance of given ID, or NULL if it was not found. Note it looks only through registered navigation data */
