Improve memory debugging with additional LLM tags.

diff --git a/Engine/Plugins/Experimental/PCG/Source/PCG/Private/Grid/PCGLandscapeCache.cpp b/Engine/Plugins/Experimental/PCG/Source/PCG/Private/Grid/PCGLandscapeCache.cpp
--- a/Engine/Plugins/Experimental/PCG/Source/PCG/Private/Grid/PCGLandscapeCache.cpp
+++ b/Engine/Plugins/Experimental/PCG/Source/PCG/Private/Grid/PCGLandscapeCache.cpp
@@ -400,6 +400,7 @@ void FPCGLandscapeCacheEntry::SerializeToBulkData(EPCGLandscapeCacheSerializatio
 bool FPCGLandscapeCacheEntry::SerializeFromBulkData() const
 {
 	check(!bDataLoaded);
+	LLM_SCOPE_BYNAME(TEXT("PCGLandscape"));
 
 	// Note: this call is not threadsafe by itself, it is meant to be called from a locked region
 	uint8* Data = nullptr;
diff --git a/Engine/Plugins/Runtime/MassGameplay/Source/MassSpawner/Private/MassSpawnerSubsystem.cpp b/Engine/Plugins/Runtime/MassGameplay/Source/MassSpawner/Private/MassSpawnerSubsystem.cpp
--- a/Engine/Plugins/Runtime/MassGameplay/Source/MassSpawner/Private/MassSpawnerSubsystem.cpp
+++ b/Engine/Plugins/Runtime/MassGameplay/Source/MassSpawner/Private/MassSpawnerSubsystem.cpp
@@ -126,6 +126,7 @@ void UMassSpawnerSubsystem::DoSpawning(const FMassEntityTemplate& EntityTemplate
 		return;
 	}
 
+	LLM_SCOPE_BYNAME(TEXT("Mass/Spawner"))
 	//TRACE_CPUPROFILER_EVENT_SCOPE_STR("MassSpawnerSubsystem DoSpawning");
 
 	// 1. Create required number of entities with EntityTemplate.Archetype
diff --git a/Engine/Source/Runtime/Landscape/Private/LandscapeCollision.cpp b/Engine/Source/Runtime/Landscape/Private/LandscapeCollision.cpp
--- a/Engine/Source/Runtime/Landscape/Private/LandscapeCollision.cpp
+++ b/Engine/Source/Runtime/Landscape/Private/LandscapeCollision.cpp
@@ -2225,6 +2225,8 @@ bool ULandscapeMeshCollisionComponent::RecreateCollision()
 
 void ULandscapeHeightfieldCollisionComponent::Serialize(FArchive& Ar)
 {
+	LLM_SCOPE(ELLMTag::Landscape);
+
 #if WITH_EDITOR
 	if (Ar.UEVer() >= VER_UE4_LANDSCAPE_COLLISION_DATA_COOKING)
 	{
@@ -2282,6 +2284,8 @@ void ULandscapeHeightfieldCollisionComponent::Serialize(FArchive& Ar)
 
 void ULandscapeMeshCollisionComponent::Serialize(FArchive& Ar)
 {
+	LLM_SCOPE(ELLMTag::Landscape);
+
 	Super::Serialize(Ar);
 
 	if (Ar.UEVer() < VER_UE4_LANDSCAPE_COLLISION_DATA_COOKING)
diff --git a/Engine/Source/Runtime/NavigationSystem/Private/NavCollision.cpp b/Engine/Source/Runtime/NavigationSystem/Private/NavCollision.cpp
--- a/Engine/Source/Runtime/NavigationSystem/Private/NavCollision.cpp
+++ b/Engine/Source/Runtime/NavigationSystem/Private/NavCollision.cpp
@@ -230,6 +230,8 @@ void UNavCollision::Setup(UBodySetup* BodySetup)
 		return;
 	}
 
+	LLM_SCOPE_BYNAME(TEXT("NavigationCollision"));
+
 	BodySetupGuid = BodySetup->BodySetupGuid;
 
 	// Make sure all are cleared before we start
diff --git a/Engine/Source/Runtime/NavigationSystem/Private/NavMesh/RecastNavMeshGenerator.cpp b/Engine/Source/Runtime/NavigationSystem/Private/NavMesh/RecastNavMeshGenerator.cpp
--- a/Engine/Source/Runtime/NavigationSystem/Private/NavMesh/RecastNavMeshGenerator.cpp
+++ b/Engine/Source/Runtime/NavigationSystem/Private/NavMesh/RecastNavMeshGenerator.cpp
@@ -6347,6 +6347,7 @@ TArray<uint32> FRecastNavMeshGenerator::ProcessTileTasksAsync(const int32 NumTas
 TArray<FNavTileRef> FRecastNavMeshGenerator::ProcessTileTasksAsyncAndGetUpdatedTiles(const int32 NumTasksToProcess)
 {
 	QUICK_SCOPE_CYCLE_COUNTER(STAT_RecastNavMeshGenerator_ProcessTileTasksAsync);
+	LLM_SCOPE(ELLMTag::NavigationRecast);
 
 	TArray<FNavTileRef> UpdatedTiles;
 	const bool bGameStaticNavMesh = IsGameStaticNavMesh(DestNavMesh);
@@ -6478,6 +6479,7 @@ TArray<FNavTileRef> FRecastNavMeshGenerator::ProcessTileTasksSyncTimeSlicedAndGe
 {
 	QUICK_SCOPE_CYCLE_COUNTER(STAT_RecastNavMeshGenerator_ProcessTileTasksSyncTimeSliced);
 	CSV_SCOPED_TIMING_STAT(NAVREGEN, ProcessTileTasksSyncTimeSliced);
+	LLM_SCOPE(ELLMTag::NavigationRecast);
 
 	check(SyncTimeSlicedData.TimeSliceManager);
 
diff --git a/Engine/Source/Runtime/NavigationSystem/Private/NavigationDataHandler.cpp b/Engine/Source/Runtime/NavigationSystem/Private/NavigationDataHandler.cpp
--- a/Engine/Source/Runtime/NavigationSystem/Private/NavigationDataHandler.cpp
+++ b/Engine/Source/Runtime/NavigationSystem/Private/NavigationDataHandler.cpp
@@ -112,6 +112,7 @@ FSetElementId FNavigationDataHandler::RegisterNavOctreeElement(UObject& ElementO
 void FNavigationDataHandler::AddElementToNavOctree(const FNavigationDirtyElement& DirtyElement)
 {
 	check(OctreeController.NavOctree.IsValid());
+	LLM_SCOPE_BYTAG(NavigationOctree);
 
 	// handle invalidated requests first
 	if (DirtyElement.bInvalidRequest)
