[Navigation] Using VeryVerbose verbosity, LogNavigationDirtyArea now logs purple boxes showing dirty areas that will dirty tiles.

diff --git a/Engine/Source/Runtime/NavigationSystem/Private/NavMesh/RecastNavMeshGenerator.cpp b/Engine/Source/Runtime/NavigationSystem/Private/NavMesh/RecastNavMeshGenerator.cpp
--- a/Engine/Source/Runtime/NavigationSystem/Private/NavMesh/RecastNavMeshGenerator.cpp
+++ b/Engine/Source/Runtime/NavigationSystem/Private/NavMesh/RecastNavMeshGenerator.cpp
@@ -5604,7 +5604,12 @@ bool FRecastNavMeshGenerator::IsAllowedToAddTileLayers(const FIntPoint Tile) con
 }
 
 #if !UE_BUILD_SHIPPING
-void FRecastNavMeshGenerator::LogDirtyAreas(
+// Deprecated
+void FRecastNavMeshGenerator::LogDirtyAreas(const TMap<FPendingTileElement, TArray<FNavigationDirtyAreaPerTileDebugInformation>>& DirtyAreasDebuggingInformation) const
+{
+}
+
+void FRecastNavMeshGenerator::LogDirtyAreas(const UObject& OwnerNav,
 	const TMap<FPendingTileElement, TArray<FNavigationDirtyAreaPerTileDebugInformation>>& DirtyAreasDebuggingInformation) const
 {
 	// Helper struct used to collate the raw information provided to the method, needed for the log results
@@ -5650,6 +5655,9 @@ void FRecastNavMeshGenerator::LogDirtyAreas(
 			TEXT("(navmesh: %-30s) Dirty area trying to dirt %2d tiles (out of which %2d are newly added/not pending) | Source Object = %s | Potential component's owner = %s | Bounds size = %s)"),
 			*GetNameSafe(GetOwner()), DirtyResultsTuple.TotalDirtyTiles, DirtyResultsTuple.NewlyAddedDirtyTiles, *GetFullNameSafe(SourceObject),
 			*GetFullNameSafe(ComponentOwner), *BoundsSize.ToString());
+
+		UE_VLOG_BOX(&OwnerNav, LogNavigationDirtyArea, VeryVerbose, DirtyResultsTuple.DirtyArea.Bounds, FColor::Purple,
+			TEXT("Tiles %d (new: %d), Source: %s"), DirtyResultsTuple.TotalDirtyTiles, DirtyResultsTuple.NewlyAddedDirtyTiles,*GetFullNameSafe(SourceObject));
 	}
 }
 #endif
@@ -6187,7 +6195,10 @@ void FRecastNavMeshGenerator::MarkDirtyTiles(const TArray<FNavigationDirtyArea>&
 #if !UE_BUILD_SHIPPING
 	UE_SUPPRESS(LogNavigationDirtyArea, VeryVerbose, 
 	{
-		LogDirtyAreas(DirtyAreasDebugging);
+		if (OwnerNav)
+		{
+			LogDirtyAreas(*OwnerNav, DirtyAreasDebugging);
+		}
 	});
 #endif
 
diff --git a/Engine/Source/Runtime/NavigationSystem/Public/NavMesh/RecastNavMeshGenerator.h b/Engine/Source/Runtime/NavigationSystem/Public/NavMesh/RecastNavMeshGenerator.h
--- a/Engine/Source/Runtime/NavigationSystem/Public/NavMesh/RecastNavMeshGenerator.h
+++ b/Engine/Source/Runtime/NavigationSystem/Public/NavMesh/RecastNavMeshGenerator.h
@@ -1005,8 +1005,13 @@ protected:
 		bool bTileWasAlreadyAdded = false;
 	};
 
-	/** Used internally, when LogNavigationDirtyArea is VeryVerbose, to log the number of tiles a dirty area is requesting. */
+	UE_DEPRECATED(5.4, "Use new version with ARecastNavMesh owner instead.")
 	NAVIGATIONSYSTEM_API void LogDirtyAreas(const TMap<FPendingTileElement, TArray<FNavigationDirtyAreaPerTileDebugInformation>>& DirtyAreasDebuggingInformation) const; 
+	
+	/** Used internally, when LogNavigationDirtyArea is VeryVerbose, to log the number of tiles a dirty area is requesting. */
+	NAVIGATIONSYSTEM_API void LogDirtyAreas(const UObject& OwnerNav,
+		const TMap<FPendingTileElement, TArray<FNavigationDirtyAreaPerTileDebugInformation>>& DirtyAreasDebuggingInformation) const;
+
 #endif
 	
 protected:
