[Navigation] Optimization to RemoveTiles and addition of profile tags to UpdateActiveTiles()

diff --git a/Engine/Source/Runtime/NavigationSystem/Private/NavMesh/RecastNavMesh.cpp b/Engine/Source/Runtime/NavigationSystem/Private/NavMesh/RecastNavMesh.cpp
--- a/Engine/Source/Runtime/NavigationSystem/Private/NavMesh/RecastNavMesh.cpp
+++ b/Engine/Source/Runtime/NavigationSystem/Private/NavMesh/RecastNavMesh.cpp
@@ -3541,44 +3541,47 @@ void ARecastNavMesh::UpdateActiveTiles(const TArray<FNavigationInvokerRaw>& Invo
 	TileToAppend.Reserve(ActiveTiles.Num());
 	ActiveTiles.Reset();
 
-	for (const FNavigationInvokerRaw& Invoker : InvokerLocations)
 	{
-		if (!Invoker.SupportedAgents.Contains(Config.AgentIndex))
+		TRACE_CPUPROFILER_EVENT_SCOPE(ARecastNavMesh::MinMaxDistance);
+		for (const FNavigationInvokerRaw& Invoker : InvokerLocations)
 		{
-			continue;
-		}
+			if (!Invoker.SupportedAgents.Contains(Config.AgentIndex))
+			{
+				continue;
+			}
 
-		const FVector InvokerRelativeLocation = (NavmeshOrigin - Invoker.Location);
-		const  FVector::FReal TileCenterDistanceToRemoveSq = FMath::Square(TileDim * UE_SQRT_2 / 2 + Invoker.RadiusMax);
-		const  FVector::FReal TileCenterDistanceToAddSq = FMath::Square(TileDim * UE_SQRT_2 / 2 + Invoker.RadiusMin);
+			const FVector InvokerRelativeLocation = (NavmeshOrigin - Invoker.Location);
+			const  FVector::FReal TileCenterDistanceToRemoveSq = FMath::Square(TileDim * UE_SQRT_2 / 2 + Invoker.RadiusMax);
+			const  FVector::FReal TileCenterDistanceToAddSq = FMath::Square(TileDim * UE_SQRT_2 / 2 + Invoker.RadiusMin);
 
-		const int32 MinTileX = IntCastChecked<int32>(FMath::FloorToInt((InvokerRelativeLocation.X - Invoker.RadiusMax) / TileDim));
-		const int32 MaxTileX = IntCastChecked<int32>(FMath::CeilToInt((InvokerRelativeLocation.X + Invoker.RadiusMax) / TileDim));
-		const int32 MinTileY = IntCastChecked<int32>(FMath::FloorToInt((InvokerRelativeLocation.Y - Invoker.RadiusMax) / TileDim));
-		const int32 MaxTileY = IntCastChecked<int32>(FMath::CeilToInt((InvokerRelativeLocation.Y + Invoker.RadiusMax) / TileDim));
+			const int32 MinTileX = IntCastChecked<int32>(FMath::FloorToInt((InvokerRelativeLocation.X - Invoker.RadiusMax) / TileDim));
+			const int32 MaxTileX = IntCastChecked<int32>(FMath::CeilToInt((InvokerRelativeLocation.X + Invoker.RadiusMax) / TileDim));
+			const int32 MinTileY = IntCastChecked<int32>(FMath::FloorToInt((InvokerRelativeLocation.Y - Invoker.RadiusMax) / TileDim));
+			const int32 MaxTileY = IntCastChecked<int32>(FMath::CeilToInt((InvokerRelativeLocation.Y + Invoker.RadiusMax) / TileDim));
 
-		for (int32 X = MinTileX; X <= MaxTileX; ++X)
-		{
-			for (int32 Y = MinTileY; Y <= MaxTileY; ++Y)
+			for (int32 X = MinTileX; X <= MaxTileX; ++X)
 			{
-				const FVector::FReal DistanceSq = (InvokerRelativeLocation - FVector(X * TileDim + TileDim / 2, Y * TileDim + TileDim / 2, 0.f)).SizeSquared2D();
-				if (DistanceSq < TileCenterDistanceToRemoveSq)
+				for (int32 Y = MinTileY; Y <= MaxTileY; ++Y)
 				{
-					TilesInMaxDistance.AddUnique(FIntPoint(X, Y));
-
-					if (DistanceSq < TileCenterDistanceToAddSq)
+					const FVector::FReal DistanceSq = (InvokerRelativeLocation - FVector(X * TileDim + TileDim / 2, Y * TileDim + TileDim / 2, 0.f)).SizeSquared2D();
+					if (DistanceSq < TileCenterDistanceToRemoveSq)
 					{
-						// Add unique tile 
-						FNavMeshDirtyTileElement* FoundTile = TilesInMinDistance.FindByPredicate([X, Y](const FNavMeshDirtyTileElement& Tile){ return Tile.Coordinates == FIntPoint(X, Y);});
-						if (FoundTile)
-						{
-							// Update the priority if already existing
-							FoundTile->InvokerPriority = FMath::Max(FoundTile->InvokerPriority, Invoker.Priority);
-						}
-						else
+						TilesInMaxDistance.AddUnique(FIntPoint(X, Y));
+
+						if (DistanceSq < TileCenterDistanceToAddSq)
 						{
-							TilesInMinDistance.Add(FNavMeshDirtyTileElement{FIntPoint(X,Y), DistanceSq, Invoker.Priority});
-							TileToAppend.Add(FIntPoint(X,Y));
+							// Add unique tile 
+							FNavMeshDirtyTileElement* FoundTile = TilesInMinDistance.FindByPredicate([X, Y](const FNavMeshDirtyTileElement& Tile){ return Tile.Coordinates == FIntPoint(X, Y);});
+							if (FoundTile)
+							{
+								// Update the priority if already existing
+								FoundTile->InvokerPriority = FMath::Max(FoundTile->InvokerPriority, Invoker.Priority);
+							}
+							else
+							{
+								TilesInMinDistance.Add(FNavMeshDirtyTileElement{FIntPoint(X,Y), DistanceSq, Invoker.Priority});
+								TileToAppend.Add(FIntPoint(X,Y));
+							}
 						}
 					}
 				}
@@ -3589,18 +3592,22 @@ void ARecastNavMesh::UpdateActiveTiles(const TArray<FNavigationInvokerRaw>& Invo
 	ActiveTiles.Append(TileToAppend);
 
 	TArray<FIntPoint> TilesToRemove;
-	TilesToRemove.Reserve(OldActiveSet.Num());
-	for (int32 Index = OldActiveSet.Num() - 1; Index >= 0; --Index)
+
 	{
-		const FIntPoint& Tile = OldActiveSet[Index];
-		if (TilesInMaxDistance.Find(Tile) == INDEX_NONE)
-		{
-			TilesToRemove.Add(Tile);
-			OldActiveSet.RemoveAtSwap(Index, 1, /*bAllowShrinking=*/false);
-		}
-		else
+		TRACE_CPUPROFILER_EVENT_SCOPE(ARecastNavMesh::CategorizeTiles);
+		TilesToRemove.Reserve(OldActiveSet.Num());
+		for (int32 Index = OldActiveSet.Num() - 1; Index >= 0; --Index)
 		{
-			ActiveTiles.AddUnique(Tile);
+			const FIntPoint& Tile = OldActiveSet[Index];
+			if (TilesInMaxDistance.Find(Tile) == INDEX_NONE)
+			{
+				TilesToRemove.Add(Tile);
+				OldActiveSet.RemoveAtSwap(Index, 1, /*bAllowShrinking=*/false);
+			}
+			else
+			{
+				ActiveTiles.AddUnique(Tile);
+			}
 		}
 	}
 
@@ -3636,6 +3643,8 @@ void ARecastNavMesh::UpdateActiveTiles(const TArray<FNavigationInvokerRaw>& Invo
 
 void ARecastNavMesh::RemoveTiles(const TArray<FIntPoint>& Tiles)
 {
+	TRACE_CPUPROFILER_EVENT_SCOPE(ARecastNavMesh::RemoveTiles);
+	
 	if (Tiles.Num() > 0)
 	{
 		FRecastNavMeshGenerator* MyGenerator = static_cast<FRecastNavMeshGenerator*>(GetGenerator());
@@ -3660,6 +3669,8 @@ void ARecastNavMesh::RebuildTile(const TArray<FIntPoint>& Tiles)
 
 void ARecastNavMesh::RebuildTile(const TArray<FNavMeshDirtyTileElement>& Tiles)
 {
+	TRACE_CPUPROFILER_EVENT_SCOPE(ARecastNavMesh::RebuildTile);
+	
 	if (Tiles.Num() > 0)
 	{
 		FRecastNavMeshGenerator* MyGenerator = static_cast<FRecastNavMeshGenerator*>(GetGenerator());
diff --git a/Engine/Source/Runtime/NavigationSystem/Private/NavMesh/RecastNavMeshGenerator.cpp b/Engine/Source/Runtime/NavigationSystem/Private/NavMesh/RecastNavMeshGenerator.cpp
--- a/Engine/Source/Runtime/NavigationSystem/Private/NavMesh/RecastNavMeshGenerator.cpp
+++ b/Engine/Source/Runtime/NavigationSystem/Private/NavMesh/RecastNavMeshGenerator.cpp
@@ -5287,9 +5287,15 @@ void FRecastNavMeshGenerator::ResetTimeSlicedTileGeneratorSync()
 //@TODO Investigate removing from RunningDirtyTiles here too (or at least not using the results in any way)
 void FRecastNavMeshGenerator::RemoveTiles(const TArray<FIntPoint>& Tiles)
 {
+	dtNavMesh* DetourMesh = DestNavMesh->GetRecastNavMeshImpl()->GetRecastMesh();
+	const bool bIsDetourMeshValid = DetourMesh && !DetourMesh->isEmpty(); 
+	
 	for (const FIntPoint& TileXY : Tiles)
 	{
-		RemoveTileLayersAndGetUpdatedTiles(TileXY.X, TileXY.Y);
+		if (bIsDetourMeshValid)
+		{
+			RemoveTileLayers(DetourMesh, TileXY.X, TileXY.Y);
+		}
 
 		if (PendingDirtyTiles.Num() > 0)
 		{
@@ -5491,6 +5497,41 @@ TArray<FNavTileRef> FRecastNavMeshGenerator::RemoveTileLayersAndGetUpdatedTiles(
 	return UpdatedIndices;
 }
 
+void FRecastNavMeshGenerator::RemoveTileLayers(dtNavMesh* DetourMesh, const int32 TileX, const int32 TileY)
+{
+	check(DetourMesh && !DetourMesh->isEmpty())
+	
+	const int32 NumLayers = DetourMesh->getTileCountAt(TileX, TileY);
+
+	if (NumLayers > 0)
+	{
+		TArray<dtMeshTile*, TInlineAllocator<16>> Tiles;
+		Tiles.AddZeroed(NumLayers);
+		DetourMesh->getTilesAt(TileX, TileY, (const dtMeshTile**)Tiles.GetData(), NumLayers);
+
+		for (int32 i = 0; i < NumLayers; i++)
+		{
+			const dtPolyRef TileRef = DetourMesh->getTileRef(Tiles[i]);
+			NumActiveTiles--;
+
+			UE_SUPPRESS(LogNavigation, VeryVerbose,
+			{
+				const int32 LayerIndex = Tiles[i]->header->layer;
+				DestNavMesh->LogRecastTile(ANSI_TO_TCHAR(__FUNCTION__), FName(""), FName("removing"), *DetourMesh, TileX, TileY, LayerIndex, TileRef));
+			};			
+
+			DetourMesh->removeTile(TileRef, nullptr, nullptr);
+		}
+	}
+
+	// Remove compressed tile cache layers
+	DestNavMesh->RemoveTileCacheLayers(TileX, TileY);
+
+#if RECAST_INTERNAL_DEBUG_DATA
+	DestNavMesh->RemoveTileDebugData(TileX, TileY);
+#endif
+}
+
 FRecastNavMeshGenerator::FSyncTimeSlicedData::FSyncTimeSlicedData()
 	: CurrentTileRegenDuration(0.)
 #if TIME_SLICE_NAV_REGEN
diff --git a/Engine/Source/Runtime/NavigationSystem/Private/NavigationSystem.cpp b/Engine/Source/Runtime/NavigationSystem/Private/NavigationSystem.cpp
--- a/Engine/Source/Runtime/NavigationSystem/Private/NavigationSystem.cpp
+++ b/Engine/Source/Runtime/NavigationSystem/Private/NavigationSystem.cpp
@@ -5098,6 +5098,8 @@ void UNavigationSystemV1::UnregisterInvoker(AActor& Invoker)
 
 void UNavigationSystemV1::UpdateInvokers()
 {
+	QUICK_SCOPE_CYCLE_COUNTER(STAT_Navigation_UpdateInvokers);
+	
 	const UWorld* World = GetWorld();
 	const double CurrentTime = World->GetTimeSeconds();
 	if (CurrentTime >= NextInvokersUpdateTime)
diff --git a/Engine/Source/Runtime/NavigationSystem/Public/NavMesh/RecastNavMeshGenerator.h b/Engine/Source/Runtime/NavigationSystem/Public/NavMesh/RecastNavMeshGenerator.h
--- a/Engine/Source/Runtime/NavigationSystem/Public/NavMesh/RecastNavMeshGenerator.h
+++ b/Engine/Source/Runtime/NavigationSystem/Public/NavMesh/RecastNavMeshGenerator.h
@@ -939,9 +939,12 @@ public:
 	UE_DEPRECATED(5.1, "Use RemoveTileLayersAndGetUpdatedTiles instead")
 	NAVIGATIONSYSTEM_API TArray<uint32> RemoveTileLayers(const int32 TileX, const int32 TileY, TMap<int32, dtPolyRef>* OldLayerTileIdMap = nullptr);
 
-	/** Removes all tiles at specified grid location */
+	/** Removes all tiles at specified grid location and returns the updated FNavTileRef */
 	NAVIGATIONSYSTEM_API TArray<FNavTileRef> RemoveTileLayersAndGetUpdatedTiles(const int32 TileX, const int32 TileY, TMap<int32, dtPolyRef>* OldLayerTileIdMap = nullptr);
 
+	/** Removes all tiles at specified grid location */
+	NAVIGATIONSYSTEM_API void RemoveTileLayers(dtNavMesh* DetourMesh, const int32 TileX, const int32 TileY);
+
 	NAVIGATIONSYSTEM_API void RemoveTiles(const TArray<FIntPoint>& Tiles);
 
 	UE_DEPRECATED(5.3, "Use overload with FNavMeshDirtyTileElement instead.")
