CIS fix

diff --git a/Engine/Source/Runtime/NavigationSystem/Private/NavMesh/RecastNavMeshGenerator.cpp b/Engine/Source/Runtime/NavigationSystem/Private/NavMesh/RecastNavMeshGenerator.cpp
--- a/Engine/Source/Runtime/NavigationSystem/Private/NavMesh/RecastNavMeshGenerator.cpp
+++ b/Engine/Source/Runtime/NavigationSystem/Private/NavMesh/RecastNavMeshGenerator.cpp
@@ -5517,8 +5517,8 @@ void FRecastNavMeshGenerator::RemoveTileLayers(dtNavMesh* DetourMesh, const int3
 			UE_SUPPRESS(LogNavigation, VeryVerbose,
 			{
 				const int32 LayerIndex = Tiles[i]->header->layer;
-				DestNavMesh->LogRecastTile(ANSI_TO_TCHAR(__FUNCTION__), FName(""), FName("removing"), *DetourMesh, TileX, TileY, LayerIndex, TileRef));
-			};			
+				DestNavMesh->LogRecastTile(ANSI_TO_TCHAR(__FUNCTION__), FName(""), FName("removing"), *DetourMesh, TileX, TileY, LayerIndex, TileRef);
+			});
 
 			DetourMesh->removeTile(TileRef, nullptr, nullptr);
 		}
