[Navmesh] Fix AgentStepHeight not set in NavDataConfig

diff --git a/Engine/Source/Runtime/NavigationSystem/Private/NavMesh/RecastNavMesh.cpp b/Engine/Source/Runtime/NavigationSystem/Private/NavMesh/RecastNavMesh.cpp
--- a/Engine/Source/Runtime/NavigationSystem/Private/NavMesh/RecastNavMesh.cpp
+++ b/Engine/Source/Runtime/NavigationSystem/Private/NavMesh/RecastNavMesh.cpp
@@ -1161,18 +1161,21 @@ bool ARecastNavMesh::CanEditChange(const FProperty* InProperty) const
 #endif // WITH_EDITOR
 
 void ARecastNavMesh::SetConfig(const FNavDataConfig& Src) 
-{ 
-	NavDataConfig = Src; 
-	AgentHeight = Src.AgentHeight;
-	AgentRadius = Src.AgentRadius;
-
-	if (Src.HasStepHeightOverride())
+{
+	// Step 1: set NavDataConfig
+	NavDataConfig = Src;
+	if (!Src.HasStepHeightOverride())
 	{
-		// If there is an override, apply it to all resolutions
-		for (int32 Index = 0; Index < (int32)ENavigationDataResolution::MAX; Index++)
-		{
-			SetAgentMaxStepHeight((ENavigationDataResolution)Index, Src.AgentStepHeight);
-		}
+		// If there is no override, use the navmesh value.
+		NavDataConfig.AgentStepHeight = GetAgentMaxStepHeight(ENavigationDataResolution::Default);
+	}
+	
+	// Step 2: update ARecastNavMesh from the new NavDataConfig
+	AgentHeight = NavDataConfig.AgentHeight;
+	AgentRadius = NavDataConfig.AgentRadius;
+	for (int32 Index = 0; Index < (int32)ENavigationDataResolution::MAX; Index++)
+	{
+		SetAgentMaxStepHeight((ENavigationDataResolution)Index, NavDataConfig.AgentStepHeight);
 	}
 }
 
