Addition of NavigationBuildDetailed CSV profile tags in the NavigationSystem tick.

diff --git a/Engine/Source/Runtime/NavigationSystem/Private/NavigationSystem.cpp b/Engine/Source/Runtime/NavigationSystem/Private/NavigationSystem.cpp
--- a/Engine/Source/Runtime/NavigationSystem/Private/NavigationSystem.cpp
+++ b/Engine/Source/Runtime/NavigationSystem/Private/NavigationSystem.cpp
@@ -158,6 +158,7 @@ DEFINE_STAT(STAT_DetourTileClustersMemory);
 DEFINE_STAT(STAT_DetourTilePolyClustersMemory);
 
 CSV_DEFINE_CATEGORY(NavigationSystem, false);
+CSV_DEFINE_CATEGORY(NavigationBuildDetailed, true);
 CSV_DEFINE_CATEGORY(NavTasksDelays, true);
 CSV_DEFINE_CATEGORY(NavTasks, true);
 CSV_DEFINE_CATEGORY(NavInvokers, true);
@@ -1335,12 +1336,14 @@ void UNavigationSystemV1::Tick(float DeltaSeconds)
 
 	if (NavDataRegistrationQueue.Num() > 0)
 	{
+		CSV_SCOPED_TIMING_STAT(NavigationBuildDetailed, Navigation_ProcessRegistrationCandidates);
 		ProcessRegistrationCandidates();
 	}
 
 	if (DefaultOctreeController.PendingOctreeUpdates.Num() > 0)
 	{
 		SCOPE_CYCLE_COUNTER(STAT_Navigation_AddingActorsToNavOctree);
+		CSV_SCOPED_TIMING_STAT(NavigationBuildDetailed, Navigation_ProcessPendingOctreeUpdates);
 
 		SCOPE_CYCLE_COUNTER(STAT_Navigation_BuildTime)
 		STAT(double ThisTime = 0);
@@ -1356,14 +1359,19 @@ void UNavigationSystemV1::Tick(float DeltaSeconds)
 	{
 		if (bGenerateNavigationOnlyAroundNavigationInvokers)
 		{
+			CSV_SCOPED_TIMING_STAT(NavigationBuildDetailed, Navigation_UpdateInvokers);
 			UpdateInvokers();
 		}
 
-		RebuildDirtyAreas(DeltaSeconds);
+		{
+			CSV_SCOPED_TIMING_STAT(NavigationBuildDetailed, Navigation_RebuildDirtyAreas);
+			RebuildDirtyAreas(DeltaSeconds);
+		}
 
 		// Tick navigation mesh async builders
 		if (bAsyncBuildPaused == false)
 		{
+			CSV_SCOPED_TIMING_STAT(NavigationBuildDetailed, Navigation_TickAsyncBuild);
 			SCOPE_CYCLE_COUNTER(STAT_Navigation_TickAsyncBuild);
 
 			bool bDoStandardTickAsync = true;
@@ -1504,6 +1512,7 @@ void UNavigationSystemV1::Tick(float DeltaSeconds)
 
 	if (CrowdManager.IsValid())
 	{
+		CSV_SCOPED_TIMING_STAT(NavigationBuildDetailed, Navigation_CrowdManager);
 		CrowdManager->Tick(DeltaSeconds);
 	}
 }
