[Navmesh] Apply AgentStepHeight from FNavDataConfig only if there is an override (restoring behavior before change in CL 29416144).

diff --git a/Engine/Source/Runtime/NavigationSystem/Private/NavMesh/RecastNavMesh.cpp b/Engine/Source/Runtime/NavigationSystem/Private/NavMesh/RecastNavMesh.cpp
--- a/Engine/Source/Runtime/NavigationSystem/Private/NavMesh/RecastNavMesh.cpp
+++ b/Engine/Source/Runtime/NavigationSystem/Private/NavMesh/RecastNavMesh.cpp
@@ -1173,9 +1173,14 @@ void ARecastNavMesh::SetConfig(const FNavDataConfig& Src)
 	// Step 2: update ARecastNavMesh from the new NavDataConfig
 	AgentHeight = NavDataConfig.AgentHeight;
 	AgentRadius = NavDataConfig.AgentRadius;
-	for (int32 Index = 0; Index < (int32)ENavigationDataResolution::MAX; Index++)
+
+	if (Src.HasStepHeightOverride())
 	{
-		SetAgentMaxStepHeight((ENavigationDataResolution)Index, NavDataConfig.AgentStepHeight);
+		// If there is an override, apply it to all resolutions
+		for (int32 Index = 0; Index < (int32)ENavigationDataResolution::MAX; Index++)
+		{
+			SetAgentMaxStepHeight((ENavigationDataResolution)Index, NavDataConfig.AgentStepHeight);
+		}
 	}
 }
 
