NavTestRenderingComponent defines cleanup

diff --git a/Engine/Source/Runtime/NavigationSystem/Private/NavMesh/NavTestRenderingComponent.cpp b/Engine/Source/Runtime/NavigationSystem/Private/NavMesh/NavTestRenderingComponent.cpp
--- a/Engine/Source/Runtime/NavigationSystem/Private/NavMesh/NavTestRenderingComponent.cpp
+++ b/Engine/Source/Runtime/NavigationSystem/Private/NavMesh/NavTestRenderingComponent.cpp
@@ -1,12 +1,14 @@
 // Copyright Epic Games, Inc. All Rights Reserved.
 
 #include "NavMesh/NavTestRenderingComponent.h"
-#include "Engine/Engine.h"
+#include "NavigationTestingActor.h"
+#include "NavMesh/RecastNavMesh.h"
+
+#if UE_ENABLE_DEBUG_DRAWING
 #include "Engine/Canvas.h"
+#include "Engine/Engine.h"
 #include "SceneManagement.h"
 #include "Materials/MaterialRenderProxy.h"
-#include "NavigationTestingActor.h"
-#include "NavMesh/RecastNavMesh.h"
 #include "Debug/DebugDrawService.h"
 
 #include UE_INLINE_GENERATED_CPP_BY_NAME(NavTestRenderingComponent)
@@ -191,6 +193,7 @@ void FNavTestSceneProxy::GatherPathStep()
 	BestNodeId = FSetElementId();
 
 #if WITH_EDITORONLY_DATA && WITH_RECAST
+	// DebugSteps are only available for: WITH_EDITORONLY_DATA && WITH_RECAST
 	if (NavTestActor && NavTestActor->DebugSteps.Num() && NavTestActor->ShowStepIndex >= 0)
 	{
 		const int32 ShowIdx = FMath::Min(NavTestActor->ShowStepIndex, NavTestActor->DebugSteps.Num() - 1);
@@ -294,7 +297,7 @@ void FNavTestSceneProxy::GatherPathStep()
 			BestPathId = MyDebugNode.ParentId;
 		}
 	}
-#endif
+#endif // WITH_EDITORONLY_DATA && WITH_RECAST
 }
 
 FPrimitiveViewRelevance FNavTestSceneProxy::GetViewRelevance(const FSceneView* View) const
@@ -323,7 +326,6 @@ uint32 FNavTestSceneProxy::GetAllocatedSizeInternal() const
 
 }
 
-#if WITH_RECAST && WITH_EDITOR
 void FNavTestDebugDrawDelegateHelper::SetupFromProxy(const FNavTestSceneProxy* InSceneProxy)
 {
 	PathPoints.Reset();
@@ -349,7 +351,6 @@ void FNavTestDebugDrawDelegateHelper::DrawDebugLabels(UCanvas* Canvas, APlayerCo
 	Canvas->SetDrawColor(FColor::White);
 	const FSceneView* View = Canvas->SceneView;
 
-#if WITH_EDITORONLY_DATA && WITH_RECAST
 	if (NodeDebug.Num())
 	{
 		const UFont* RenderFont = GEngine->GetSmallFont();
@@ -378,7 +379,6 @@ void FNavTestDebugDrawDelegateHelper::DrawDebugLabels(UCanvas* Canvas, APlayerCo
 	}
 	else
 	{
-#endif
 		for (int32 PointIndex = 0; PointIndex < PathPoints.Num(); ++PointIndex)
 		{
 			if (FNavTestSceneProxy::LocationInView(PathPoints[PointIndex], View))
@@ -389,25 +389,17 @@ void FNavTestDebugDrawDelegateHelper::DrawDebugLabels(UCanvas* Canvas, APlayerCo
 
 			}
 		}
-
-#if WITH_EDITORONLY_DATA && WITH_RECAST
 	}
-#endif
-	Canvas->SetDrawColor(OldDrawColor);
-}
-#endif //WITH_RECAST && WITH_EDITOR
 
-UNavTestRenderingComponent::UNavTestRenderingComponent(const FObjectInitializer& ObjectInitializer) : Super(ObjectInitializer)
-{
+	Canvas->SetDrawColor(OldDrawColor);
 }
+#endif // UE_ENABLE_DEBUG_DRAWING
 
 #if UE_ENABLE_DEBUG_DRAWING
-  FDebugRenderSceneProxy* UNavTestRenderingComponent::CreateDebugSceneProxy()
+FDebugRenderSceneProxy* UNavTestRenderingComponent::CreateDebugSceneProxy()
 {
 	FNavTestSceneProxy* NewSceneProxy = new FNavTestSceneProxy(this);
-#if WITH_RECAST && WITH_EDITOR
 	NavTestDebugDrawDelegateHelper.SetupFromProxy(NewSceneProxy);
-#endif
 	return NewSceneProxy;
 }
 #endif
@@ -429,6 +421,7 @@ FBoxSphereBounds UNavTestRenderingComponent::CalcBounds(const FTransform& LocalT
 			}
 		}
 #if WITH_EDITORONLY_DATA && WITH_RECAST
+		// DebugSteps are only available for: WITH_EDITORONLY_DATA && WITH_RECAST
 		if (TestActor->DebugSteps.Num() && TestActor->ShowStepIndex >= 0)
 		{
 			const int32 ShowIdx = FMath::Min(TestActor->ShowStepIndex, TestActor->DebugSteps.Num() - 1);
@@ -442,7 +435,7 @@ FBoxSphereBounds UNavTestRenderingComponent::CalcBounds(const FTransform& LocalT
 				}
 			}
 		}
-#endif
+#endif // WITH_EDITORONLY_DATA && WITH_RECAST
 	}
 
 	return FBoxSphereBounds(BoundingBox);
diff --git a/Engine/Source/Runtime/NavigationSystem/Public/NavMesh/NavTestRenderingComponent.h b/Engine/Source/Runtime/NavigationSystem/Public/NavMesh/NavTestRenderingComponent.h
--- a/Engine/Source/Runtime/NavigationSystem/Public/NavMesh/NavTestRenderingComponent.h
+++ b/Engine/Source/Runtime/NavigationSystem/Public/NavMesh/NavTestRenderingComponent.h
@@ -2,11 +2,19 @@
 
 #pragma once
 
-#include "PrimitiveViewRelevance.h"
 #include "Debug/DebugDrawComponent.h"
+
+#if UE_ENABLE_INCLUDE_ORDER_DEPRECATED_IN_5_4
+#include "PrimitiveViewRelevance.h"
 #include "DynamicMeshBuilder.h"
 #include "DebugRenderSceneProxy.h"
+#endif // UE_ENABLE_INCLUDE_ORDER_DEPRECATED_IN_5_4
+
+#include "NavTestRenderingComponent.generated.h"
+
+#if UE_ENABLE_DEBUG_DRAWING
 
+class FPrimitiveSceneProxy;
 class ANavigationTestingActor;
 class APlayerController;
 class FMeshElementCollector;
@@ -42,7 +50,7 @@ public:
 		}
 	};
 
-	FNavTestSceneProxy(const UNavTestRenderingComponent* InComponent);
+	explicit FNavTestSceneProxy(const UNavTestRenderingComponent* InComponent);
 
 	virtual void GetDynamicMeshElements(const TArray<const FSceneView*>& Views, const FSceneViewFamily& ViewFamily, uint32 VisibilityMap, FMeshElementCollector& Collector) const override;
 
@@ -69,8 +77,6 @@ private:
 	uint32 bShowDiff : 1;
 
 	ANavigationTestingActor* NavTestActor;
-	FDebugDrawDelegate DebugTextDrawingDelegate;
-	FDelegateHandle DebugTextDrawingDelegateHandle;
 	TArray<FVector> PathPoints;
 	TArray<FString> PathPointFlags;
 
@@ -84,7 +90,6 @@ private:
 	FVector ClosestWallLocation;
 };
 
-#if WITH_RECAST && WITH_EDITOR
 class FNavTestDebugDrawDelegateHelper : public FDebugDrawDelegateHelper
 {
 	typedef FDebugDrawDelegateHelper Super;
@@ -108,16 +113,13 @@ private:
 	uint32 bShowBestPath : 1;
 	uint32 bShowDiff : 1;
 };
-#endif //WITH_RECAST && WITH_EDITOR
-
-#include "NavTestRenderingComponent.generated.h"
 
-class FPrimitiveSceneProxy;
+#endif // UE_ENABLE_DEBUG_DRAWING
 
 UCLASS(ClassGroup = Debug)
 class UNavTestRenderingComponent: public UDebugDrawComponent
 {
-	GENERATED_UCLASS_BODY()
+	GENERATED_BODY()
 
 protected:
 
@@ -125,11 +127,9 @@ protected:
 
 #if UE_ENABLE_DEBUG_DRAWING
 	virtual FDebugRenderSceneProxy* CreateDebugSceneProxy() override;
-
-#if WITH_RECAST && WITH_EDITOR
 	virtual FDebugDrawDelegateHelper& GetDebugDrawDelegateHelper() override { return NavTestDebugDrawDelegateHelper; }
 private:
 	FNavTestDebugDrawDelegateHelper NavTestDebugDrawDelegateHelper;
-#endif // WITH_RECAST && WITH_EDITOR
 #endif // UE_ENABLE_DEBUG_DRAWING
 };
+
