Navmesh resolution from modifiers is now correctly read when used with in async tile generation combined with lazy modifier gathering.
#tests Tested in PIE and cook

diff --git a/Engine/Source/Runtime/NavigationSystem/Private/NavMesh/RecastNavMesh.cpp b/Engine/Source/Runtime/NavigationSystem/Private/NavMesh/RecastNavMesh.cpp
--- a/Engine/Source/Runtime/NavigationSystem/Private/NavMesh/RecastNavMesh.cpp
+++ b/Engine/Source/Runtime/NavigationSystem/Private/NavMesh/RecastNavMesh.cpp
@@ -979,16 +979,16 @@ void ARecastNavMesh::SortAreasForGenerator(TArray<FRecastAreaNavModifierElement>
 
 	struct FNavAreaSortPredicate
 	{
-		FORCEINLINE bool operator()(const FRecastAreaNavModifierElement& ElA, const FRecastAreaNavModifierElement& ElB) const
+		FORCEINLINE bool operator()(const FRecastAreaNavModifierElement& ElementA, const FRecastAreaNavModifierElement& ElementB) const
 		{
-			if (ElA.Areas.Num() == 0 || ElB.Areas.Num() == 0)
+			if (ElementA.Areas.Num() == 0 || ElementB.Areas.Num() == 0)
 			{
-				return ElA.Areas.Num() <= ElB.Areas.Num();
+				return ElementA.Areas.Num() <= ElementB.Areas.Num();
 			}
 
 			// assuming composite modifiers has same area type
-			const FAreaNavModifier& A = ElA.Areas[0];
-			const FAreaNavModifier& B = ElB.Areas[0];
+			const FAreaNavModifier& A = ElementA.Areas[0];
+			const FAreaNavModifier& B = ElementB.Areas[0];
 			
 			const bool bIsAReplacing = (A.GetAreaClassToReplace() != NULL);
 			const bool bIsBReplacing = (B.GetAreaClassToReplace() != NULL);
@@ -997,7 +997,17 @@ void ARecastNavMesh::SortAreasForGenerator(TArray<FRecastAreaNavModifierElement>
 				return bIsAReplacing;
 			}
 
-			return A.Cost != B.Cost ? A.Cost < B.Cost : A.FixedCost < B.FixedCost;
+			if (A.Cost != B.Cost)
+			{
+				return A.Cost < B.Cost;
+			}
+
+			if (A.FixedCost != B.FixedCost)
+			{
+				return A.FixedCost < B.FixedCost;
+			}
+
+			return ElementA.NavMeshResolution < ElementB.NavMeshResolution;
 		}
 	};
 
diff --git a/Engine/Source/Runtime/NavigationSystem/Private/NavMesh/RecastNavMeshGenerator.cpp b/Engine/Source/Runtime/NavigationSystem/Private/NavMesh/RecastNavMeshGenerator.cpp
--- a/Engine/Source/Runtime/NavigationSystem/Private/NavMesh/RecastNavMeshGenerator.cpp
+++ b/Engine/Source/Runtime/NavigationSystem/Private/NavMesh/RecastNavMeshGenerator.cpp
@@ -1823,6 +1823,10 @@ ETimeSliceWorkResult FRecastTileGenerator::DoWorkTimeSliced()
 			{
 				WorkResult = GatherGeometryFromSourcesTimeSliced();
 
+				// Needs to occur after DemandLazyDataGathering
+				TSharedPtr<FRecastNavMeshGenerator> RecastParentGenerator = StaticCastSharedPtr<FRecastNavMeshGenerator>(ParentGenerator);
+				SetupTileConfigFromHighestResolution(*RecastParentGenerator);	
+
 				if (WorkResult == ETimeSliceWorkResult::CallAgainNextTimeSlice)
 				{
 					break;
@@ -1867,6 +1871,10 @@ bool FRecastTileGenerator::DoWork()
 		if (InclusionBounds.Num())
 		{
 			GatherGeometryFromSources();
+			
+			// Needs to occur after DemandLazyDataGathering
+			TSharedPtr<FRecastNavMeshGenerator> RecastParentGenerator = StaticCastSharedPtr<FRecastNavMeshGenerator>(ParentGenerator);
+			SetupTileConfigFromHighestResolution(*RecastParentGenerator);	
 		}
 
 		bSuccess = GenerateTile();
@@ -1891,6 +1899,28 @@ void FRecastTileGenerator::DumpSyncData()
 	ensure(IsInGameThread());
 	NavigationRelevantData.Empty();
 }
+
+void FRecastTileGenerator::SetupTileConfigFromHighestResolution(const FRecastNavMeshGenerator& ParentGenerator)
+{
+	ENavigationDataResolution HighestResolution = ENavigationDataResolution::Low;
+	bool bNewResolutionFound = false;
+
+	for (const FRecastAreaNavModifierElement& Element : Modifiers)
+	{
+		if (Element.NavMeshResolution != ENavigationDataResolution::Invalid)
+		{
+			HighestResolution = FMath::Max(HighestResolution, Element.NavMeshResolution);
+			bNewResolutionFound = true;
+		}	
+	}
+	
+	check(HighestResolution != ENavigationDataResolution::Invalid);
+	if (bNewResolutionFound && ParentGenerator.GetOwner()->NavMeshResolutionParams[(uint8)HighestResolution].IsValid())
+	{
+		// Update the TileConfig
+		ParentGenerator.SetupTileConfig(HighestResolution, TileConfig);
+	}
+}
 	
 void FRecastTileGenerator::GatherGeometryFromSources()
 {
@@ -1959,12 +1989,9 @@ void FRecastTileGenerator::PrepareGeometrySources(const FRecastNavMeshGenerator&
 	
 	const ARecastNavMesh* const OwnerNav = ParentGenerator.GetOwner();
 	const bool bUseVirtualGeometryFilteringAndDirtying = OwnerNav->bUseVirtualGeometryFilteringAndDirtying;
-
-	ENavigationDataResolution HighestResolution = ENavigationDataResolution::Low;
-	bool bNewResolutionFound = false;
 	
 	NavOctreeInstance->FindElementsWithBoundsTest(ParentGenerator.GrowBoundingBox(TileBB, /*bIncludeAgentHeight*/ false),
-		[&HighestResolution, &bNewResolutionFound, &ParentGenerator, OwnerNav, this, bGeometryChanged, bUseVirtualGeometryFilteringAndDirtying](const FNavigationOctreeElement& Element)
+		[&ParentGenerator, this, bGeometryChanged, bUseVirtualGeometryFilteringAndDirtying](const FNavigationOctreeElement& Element)
 	{
 		const bool bShouldUse = bUseVirtualGeometryFilteringAndDirtying ?
 			ParentGenerator.ShouldGenerateGeometryForOctreeElement(Element, NavDataConfig) :
@@ -1977,27 +2004,10 @@ void FRecastTileGenerator::PrepareGeometrySources(const FRecastNavMeshGenerator&
 				Element.Data->Modifiers.HasMetaAreas() == true || 
 				Element.Data->Modifiers.IsEmpty() == false)
 			{
-				UE_SUPPRESS(LogNavigation, VeryVerbose, UE_VLOG_UELOG(OwnerNav, LogNavigation, VeryVerbose, TEXT("PrepareGeometrySources, adding: %s"), *GetNameSafe(Element.GetOwner())));
-				
-				// Keep highest resolution that is not the default.
-				const ENavigationDataResolution Resolution = Element.Data->Modifiers.GetNavMeshResolution();
-				if (Resolution != ENavigationDataResolution::Invalid)
-				{
-					HighestResolution = FMath::Max(HighestResolution, Resolution);
-					bNewResolutionFound = true;
-				}
-				
 				NavigationRelevantData.Add(Element.Data);
 			}
 		}
 	});
-	
-	check(HighestResolution != ENavigationDataResolution::Invalid);
-	if (bNewResolutionFound && ParentGenerator.GetOwner()->NavMeshResolutionParams[(uint8)HighestResolution].IsValid())
-	{
-		// Update the TileConfig
-		ParentGenerator.SetupTileConfig(HighestResolution, TileConfig);
-	}
 }
 
 void FRecastTileGenerator::GatherGeometry(const FRecastNavMeshGenerator& ParentGenerator, bool bGeometryChanged)
@@ -2029,7 +2039,6 @@ void FRecastTileGenerator::GatherGeometry(const FRecastNavMeshGenerator& ParentG
 		if (bShouldUse)
 		{
 			RelevantDataArray.Add(Element.Data);
-
 		}
 	});
 
@@ -2401,6 +2410,7 @@ void FRecastTileGenerator::AppendModifier(const FCompositeNavModifier& Modifier,
 		
 	ModifierElement.Areas = Modifier.GetAreas();
 	ModifierElement.bMaskFillCollisionUnderneathForNavmesh = Modifier.GetMaskFillCollisionUnderneathForNavmesh();
+	ModifierElement.NavMeshResolution = Modifier.GetNavMeshResolution();
 
 	Modifiers.Add(MoveTemp(ModifierElement));
 }
@@ -2458,6 +2468,8 @@ void FRecastTileGenerator::AppendGeometry(const FNavigationRelevantData& DataRef
 
 ETimeSliceWorkResult FRecastTileGenerator::GenerateTileTimeSliced()
 {
+	ensureMsgf(TileConfig.bIsTileSetupConfigCompleted, TEXT("SetupTileConfig must have been called before generating a tile."));
+	
 	FNavMeshBuildContext BuildContext(*this);
 	ETimeSliceWorkResult WorkResult = ETimeSliceWorkResult::Succeeded;
 
@@ -2527,6 +2539,8 @@ bool FRecastTileGenerator::GenerateTile()
 	const double StartStamp = FPlatformTime::Seconds();
 	double PostCompressLayerStamp = StartStamp;
 #endif // RECAST_INTERNAL_DEBUG_DATA
+
+	ensureMsgf(TileConfig.bIsTileSetupConfigCompleted, TEXT("SetupTileConfig must have been called before generating a tile."));
 	
 	FNavMeshBuildContext BuildContext(*this);
 	bool bSuccess = true;
@@ -4647,6 +4661,8 @@ void FRecastNavMeshGenerator::SetupTileConfig(const ENavigationDataResolution Ti
 
 	// Update all settings that depends directly or indirectly of AgentMaxStepHeight
 	OutConfig.AgentMaxClimb = AgentMaxStepHeight;
+
+	OutConfig.bIsTileSetupConfigCompleted = true;
 }
 
 void FRecastNavMeshGenerator::ConfigureBuildProperties(FRecastBuildConfig& OutConfig)
diff --git a/Engine/Source/Runtime/NavigationSystem/Public/NavMesh/RecastNavMeshGenerator.h b/Engine/Source/Runtime/NavigationSystem/Public/NavMesh/RecastNavMeshGenerator.h
--- a/Engine/Source/Runtime/NavigationSystem/Public/NavMesh/RecastNavMeshGenerator.h
+++ b/Engine/Source/Runtime/NavigationSystem/Public/NavMesh/RecastNavMeshGenerator.h
@@ -83,7 +83,9 @@ struct FRecastBuildConfig : public rcConfig
 	ENavigationDataResolution TileResolution;
 	/** Ledge filtering mode */
 	ENavigationLedgeSlopeFilterMode LedgeSlopeFilterMode;
-
+	/** Is the config completely setup */
+	bool bIsTileSetupConfigCompleted = false;
+	
 	FRecastBuildConfig()
 	{
 		Reset();
@@ -104,6 +106,7 @@ struct FRecastBuildConfig : public rcConfig
 		AgentIndex = 0;
 		TileResolution = ENavigationDataResolution::Default;
 		LedgeSlopeFilterMode = ENavigationLedgeSlopeFilterMode::Recast;
+		bIsTileSetupConfigCompleted = false;
 	}
 
 	rcReal GetTileSizeUU() const { return tileSize * cs; }
@@ -182,6 +185,8 @@ struct FRecastAreaNavModifierElement
 	// When empty, areas are in world space
 	TArray<FTransform>	PerInstanceTransform;
 
+	ENavigationDataResolution NavMeshResolution = ENavigationDataResolution::Invalid;
+	
 	bool bMaskFillCollisionUnderneathForNavmesh = false;
 };
 
@@ -385,6 +390,9 @@ protected:
 	NAVIGATIONSYSTEM_API bool GenerateTile();
 
 	NAVIGATIONSYSTEM_API void Setup(const FRecastNavMeshGenerator& ParentGenerator, const TArray<FBox>& DirtyAreas);
+
+	/** Find highest navmesh resolution from Modifiers and use it to update the tile configuration. */
+	void SetupTileConfigFromHighestResolution(const FRecastNavMeshGenerator& ParentGenerator);
 	
 	/** Gather geometry */
 	NAVIGATIONSYSTEM_API virtual void GatherGeometry(const FRecastNavMeshGenerator& ParentGenerator, bool bGeometryChanged);
