Rollback recently added ensures, they get triggered by preexisting conditions that needs investigation.

diff --git a/Engine/Source/Runtime/NavigationSystem/Private/NavMesh/RecastNavMeshGenerator.cpp b/Engine/Source/Runtime/NavigationSystem/Private/NavMesh/RecastNavMeshGenerator.cpp
--- a/Engine/Source/Runtime/NavigationSystem/Private/NavMesh/RecastNavMeshGenerator.cpp
+++ b/Engine/Source/Runtime/NavigationSystem/Private/NavMesh/RecastNavMeshGenerator.cpp
@@ -2468,8 +2468,6 @@ void FRecastTileGenerator::AppendGeometry(const FNavigationRelevantData& DataRef
 
 ETimeSliceWorkResult FRecastTileGenerator::GenerateTileTimeSliced()
 {
-	ensureMsgf(TileConfig.bIsTileSetupConfigCompleted, TEXT("SetupTileConfig must have been called before generating a tile."));
-	
 	FNavMeshBuildContext BuildContext(*this);
 	ETimeSliceWorkResult WorkResult = ETimeSliceWorkResult::Succeeded;
 
@@ -2539,8 +2537,6 @@ bool FRecastTileGenerator::GenerateTile()
 	const double StartStamp = FPlatformTime::Seconds();
 	double PostCompressLayerStamp = StartStamp;
 #endif // RECAST_INTERNAL_DEBUG_DATA
-
-	ensureMsgf(TileConfig.bIsTileSetupConfigCompleted, TEXT("SetupTileConfig must have been called before generating a tile."));
 	
 	FNavMeshBuildContext BuildContext(*this);
 	bool bSuccess = true;
