[Navigation] Prevent addition of objects with empty bounds to the navigation octree.

diff --git a/Engine/Source/Runtime/NavigationSystem/Private/NavigationOctree.cpp b/Engine/Source/Runtime/NavigationSystem/Private/NavigationOctree.cpp
--- a/Engine/Source/Runtime/NavigationSystem/Private/NavigationOctree.cpp
+++ b/Engine/Source/Runtime/NavigationSystem/Private/NavigationOctree.cpp
@@ -143,9 +143,9 @@ void FNavigationOctree::AddNode(UObject* ElementOb, INavRelevantInterface* NavEl
 {
 	LLM_SCOPE_BYTAG(NavigationOctree);
 
-	if (!Bounds.IsValid)
+	if (UNLIKELY(!Bounds.IsValid || Bounds.GetSize().IsNearlyZero())) 
 	{
-		UE_LOG(LogNavigation, Warning, TEXT("%hs: Ignoring %s because bounds are invalid."), __FUNCTION__, *GetFullNameSafe(ElementOb));
+		UE_LOG(LogNavigation, Warning, TEXT("%hs: %s bounds, ignoring %s."), __FUNCTION__, !Bounds.IsValid ? TEXT("Invalid") : TEXT("Empty"), *GetFullNameSafe(ElementOb));
 		return;
 	}
 	
