Add a way to register a Navigation Invoker based on UObject with a INavigationInvokerInterface instead of an AActor.

diff --git a/Engine/Source/Runtime/Engine/Classes/AI/Navigation/NavigationInvokerInterface.h b/Engine/Source/Runtime/Engine/Classes/AI/Navigation/NavigationInvokerInterface.h
new file mode 100644
--- /dev/null
+++ b/Engine/Source/Runtime/Engine/Classes/AI/Navigation/NavigationInvokerInterface.h
@@ -0,0 +1,23 @@
+// Copyright Epic Games, Inc. All Rights Reserved.
+
+#pragma once
+
+#include "AI/Navigation/NavigationTypes.h"
+#include "UObject/ObjectMacros.h"
+#include "UObject/Interface.h"
+
+#include "NavigationInvokerInterface.generated.h"
+
+UINTERFACE(MinimalAPI, meta=(CannotImplementInterfaceInBlueprint))
+class UNavigationInvokerInterface : public UInterface
+{
+	GENERATED_UINTERFACE_BODY()
+};
+
+class INavigationInvokerInterface
+{
+	GENERATED_IINTERFACE_BODY()
+
+	/** Called to retrieve the center of the area that require navmesh. Must be implemented by child class. */
+	virtual FVector GetNavigationInvokerLocation() const PURE_VIRTUAL(, return FNavigationSystem::InvalidLocation;);
+};
diff --git a/Engine/Source/Runtime/Engine/Private/AI/Navigation/NavigationInterfaces.cpp b/Engine/Source/Runtime/Engine/Private/AI/Navigation/NavigationInterfaces.cpp
--- a/Engine/Source/Runtime/Engine/Private/AI/Navigation/NavigationInterfaces.cpp
+++ b/Engine/Source/Runtime/Engine/Private/AI/Navigation/NavigationInterfaces.cpp
@@ -7,6 +7,7 @@
 #include "AI/Navigation/NavigationTypes.h"
 #include "AI/RVOAvoidanceInterface.h"
 #include "AI/Navigation/NavigationDataInterface.h"
+#include "AI/Navigation/NavigationInvokerInterface.h"
 #include "AI/Navigation/PathFollowingAgentInterface.h"
 #include "UObject/Interface.h"
 
@@ -49,3 +50,8 @@ UPathFollowingAgentInterface::UPathFollowingAgentInterface(const FObjectInitiali
 	: Super(ObjectInitializer)
 {
 }
+
+UNavigationInvokerInterface::UNavigationInvokerInterface(const FObjectInitializer& ObjectInitializer)
+	: Super(ObjectInitializer)
+{
+}
diff --git a/Engine/Source/Runtime/NavigationSystem/Private/NavigationSystem.cpp b/Engine/Source/Runtime/NavigationSystem/Private/NavigationSystem.cpp
--- a/Engine/Source/Runtime/NavigationSystem/Private/NavigationSystem.cpp
+++ b/Engine/Source/Runtime/NavigationSystem/Private/NavigationSystem.cpp
@@ -10,6 +10,7 @@
 #include "GameFramework/Controller.h"
 #include "AI/Navigation/NavRelevantInterface.h"
 #include "AI/Navigation/NavigationDirtyElement.h"
+#include "AI/Navigation/NavigationInvokerInterface.h"
 #include "AI/Navigation/NavigationInvokerPriority.h"
 #include "UObject/UObjectIterator.h"
 #include "EngineUtils.h"
@@ -233,6 +234,7 @@ FNavigationInvokerRaw::FNavigationInvokerRaw(const FVector& InLocation, float Mi
 //----------------------------------------------------------------------//
 FNavigationInvoker::FNavigationInvoker()
 : Actor(nullptr)
+, Object(nullptr)
 , GenerationRadius(0)
 , RemovalRadius(0)
 , Priority(ENavigationInvokerPriority::Default)
@@ -242,6 +244,7 @@ FNavigationInvoker::FNavigationInvoker()
 
 FNavigationInvoker::FNavigationInvoker(AActor& InActor, float InGenerationRadius, float InRemovalRadius, const FNavAgentSelector& InSupportedAgents, ENavigationInvokerPriority InPriority)
 : Actor(&InActor)
+, Object(nullptr)
 , GenerationRadius(InGenerationRadius)
 , RemovalRadius(InRemovalRadius)
 , SupportedAgents(InSupportedAgents)
@@ -250,6 +253,52 @@ FNavigationInvoker::FNavigationInvoker(AActor& InActor, float InGenerationRadius
 	SupportedAgents.MarkInitialized();
 }
 
+FNavigationInvoker::FNavigationInvoker(INavigationInvokerInterface& InObject, float InGenerationRadius, float InRemovalRadius, const FNavAgentSelector& InSupportedAgents, ENavigationInvokerPriority InPriority)
+: Actor(nullptr)
+, Object(&InObject)
+, GenerationRadius(InGenerationRadius)
+, RemovalRadius(InRemovalRadius)
+, SupportedAgents(InSupportedAgents)
+, Priority(InPriority)
+{
+}
+
+FString FNavigationInvoker::GetName() const
+{
+	/** We are using IsExplicitlyNull to know which one of the Actor or the Object was set at construction */
+	if (!Actor.IsExplicitlyNull())
+	{
+		return GetNameSafe(Actor.Get());
+	}
+	else
+	{
+		return GetNameSafe(Object.GetObject());
+	}
+}
+
+bool FNavigationInvoker::GetLocation(FVector& OutLocation) const
+{
+	/** We are using IsExplicitlyNull to know which one of the Actor or the Object was set at construction */
+	if (!Actor.IsExplicitlyNull())
+	{
+		if (const AActor* ActorPtr = Actor.Get())
+		{
+			OutLocation = ActorPtr->GetActorLocation();
+			return true;
+		}
+	}
+	else
+	{
+		if (const INavigationInvokerInterface* InvokerInterface = Object.Get())
+		{
+			OutLocation = InvokerInterface->GetNavigationInvokerLocation();
+			return true;
+		}
+	}
+
+	return false;
+}
+
 //----------------------------------------------------------------------//
 // helpers
 //----------------------------------------------------------------------//
@@ -5043,6 +5092,38 @@ void UNavigationSystemV1::SetGeometryGatheringMode(ENavDataGatheringModeConfig N
 	}
 }
 
+namespace UE::Navigation::Private
+{
+	void LogNavInvokerRegistration(const UNavigationSystemV1& NavSystem, const FNavigationInvoker& Data)
+	{
+		UE_SUPPRESS(LogNavInvokers, Log,
+		{
+			TStringBuilder<128> InvokerNavData;
+			for (int32 NavDataIndex = 0; NavDataIndex < NavSystem.NavDataSet.Num(); NavDataIndex++)
+			{
+				const ANavigationData* NavData = NavSystem.NavDataSet[NavDataIndex].Get();
+				if (NavData)
+				{
+					const int32 NavDataSupportedAgentIndex = NavSystem.GetSupportedAgentIndex(NavData);
+					if (Data.SupportedAgents.Contains(NavDataSupportedAgentIndex))
+					{
+						InvokerNavData.Append(FString::Printf(TEXT("%s "), *NavData->GetName()));
+					}
+				}
+			}
+
+			const FString RegisterText = FString::Printf(TEXT("Register invoker r: %.0f, r area: %.0f m2, removal r: %.0f, priority: %s, (%s %s) "),
+				Data.GenerationRadius, UE_PI*FMath::Square(Data.GenerationRadius/100.f), Data.RemovalRadius, *UEnum::GetDisplayValueAsText(Data.Priority).ToString(), *Data.GetName(), *InvokerNavData);
+			UE_LOG(LogNavInvokers, Log, TEXT("%s"), *RegisterText);
+
+			FVector InvokerLocation = FVector::ZeroVector;
+			Data.GetLocation(InvokerLocation);
+			UE_VLOG_CYLINDER(&NavSystem, LogNavInvokers, Log, InvokerLocation, InvokerLocation + FVector(0, 0, 20), Data.GenerationRadius, FColorList::LimeGreen, TEXT("%s"), *RegisterText);
+			UE_VLOG_CYLINDER(&NavSystem, LogNavInvokers, Log, InvokerLocation, InvokerLocation + FVector(0, 0, 20), Data.RemovalRadius, FColorList::IndianRed, TEXT(""));
+		});
+	}
+}
+
 // Deprecated
 void UNavigationSystemV1::RegisterInvoker(AActor& Invoker, float TileGenerationRadius, float TileRemovalRadius)
 {
@@ -5074,32 +5155,44 @@ void UNavigationSystemV1::RegisterInvoker(AActor& Invoker, float TileGenerationR
 	Data.SupportedAgents.MarkInitialized();
 	Data.Priority = InPriority;
 
-	UE_SUPPRESS(LogNavInvokers, Log,
-	{
-		TStringBuilder<128> InvokerNavData;
-		for (int32 NavDataIndex = 0; NavDataIndex < NavDataSet.Num(); NavDataIndex++)
-		{
-			const ANavigationData* NavData = NavDataSet[NavDataIndex].Get();
-			if (NavData)
-			{
-				const int32 NavDataSupportedAgentIndex = GetSupportedAgentIndex(NavData);	
-				if (Data.SupportedAgents.Contains(NavDataSupportedAgentIndex))
-				{
-					InvokerNavData.Append(FString::Printf(TEXT("%s "), *NavData->GetName()));
-				}
-			}
-		}
+	UE::Navigation::Private::LogNavInvokerRegistration(*this, Data);
+}
 
-		const FString RegisterText = FString::Printf(TEXT("Register invoker r: %.0f, r area: %.0f m2, removal r: %.0f, priority: %s, (%s %s) "),
-			TileGenerationRadius, UE_PI*FMath::Square(TileGenerationRadius/100.f), TileRemovalRadius, *UEnum::GetDisplayValueAsText(InPriority).ToString(), *Invoker.GetName(), *InvokerNavData);
-		UE_LOG(LogNavInvokers, Log, TEXT("%s"), *RegisterText);
+void UNavigationSystemV1::RegisterInvoker(const TWeakInterfacePtr<INavigationInvokerInterface>& Invoker, float TileGenerationRadius, float TileRemovalRadius, const FNavAgentSelector& Agents, ENavigationInvokerPriority InPriority)
+{
+	UE_CVLOG(bGenerateNavigationOnlyAroundNavigationInvokers == false, this, LogNavInvokers, Warning
+		, TEXT("Trying to register %s as invoker, but NavigationSystem is not set up for invoker-centric generation. See GenerateNavigationOnlyAroundNavigationInvokers in NavigationSystem's properties")
+		, *GetNameSafe(Invoker.GetObject()));
 
-		UE_VLOG_CYLINDER(this, LogNavInvokers, Log, Invoker.GetActorLocation(), Invoker.GetActorLocation() + FVector(0, 0, 20), TileGenerationRadius, FColorList::LimeGreen, TEXT("%s"), *RegisterText);
-		UE_VLOG_CYLINDER(this, LogNavInvokers, Log, Invoker.GetActorLocation(), Invoker.GetActorLocation() + FVector(0, 0, 20), TileRemovalRadius, FColorList::IndianRed, TEXT(""));
-	});
+	UObject* InvokerObject = Invoker.GetObject();
+	if (ensure(InvokerObject != nullptr))
+	{
+		FNavigationInvoker& Data = Invokers.FindOrAdd(InvokerObject);
+		Data.Object = Invoker;
+		Data.GenerationRadius = TileGenerationRadius;
+		Data.RemovalRadius = TileRemovalRadius;
+		Data.SupportedAgents = Agents;
+		Data.SupportedAgents.MarkInitialized();
+		Data.Priority = InPriority;
+
+		UE::Navigation::Private::LogNavInvokerRegistration(*this, Data);
+	}
 }
 
 void UNavigationSystemV1::UnregisterInvoker(AActor& Invoker)
+{
+	UnregisterInvoker_Internal(Invoker);
+}
+
+void UNavigationSystemV1::UnregisterInvoker(const TWeakInterfacePtr<INavigationInvokerInterface>& Invoker)
+{
+	if (const UObject* InvokerObject = Invoker.GetObject())
+	{
+		UnregisterInvoker_Internal(*InvokerObject);
+	}
+}
+
+void UNavigationSystemV1::UnregisterInvoker_Internal(const UObject& Invoker)
 {
 	UE_VLOG(this, LogNavInvokers, Log, TEXT("Removing %s from invokers list"), *Invoker.GetName());
 	Invokers.Remove(&Invoker);
@@ -5134,60 +5227,55 @@ void UNavigationSystemV1::UpdateInvokers()
 						FVector(SeedLocation.X+InvokersMaximumDistanceFromSeed, SeedLocation.Y+InvokersMaximumDistanceFromSeed, SeedLocation.Z+InvokersMaximumDistanceFromSeed));
 				}
 			}
-			
+
+#if ENABLE_VISUAL_LOG
 			const double StartTime = FPlatformTime::Seconds();
+#endif // ENABLE_VISUAL_LOG
 
 			InvokerLocations.Reserve(Invokers.Num());
 
 			for (auto ItemIterator = Invokers.CreateIterator(); ItemIterator; ++ItemIterator)
 			{
-				const AActor* Actor = ItemIterator->Value.Actor.Get();
-				if (Actor != nullptr
-#if WITH_EDITOR
-					// Would like to ignore objects in transactional buffer here, but there's no flag for it
-					//&& (GIsEditor == false || Item.Actor->HasAnyFlags(RF_Transactional | RF_PendingKill) == false)
-#endif //WITH_EDITOR
-					)
+				FVector InvokerLocation;
+				if (!ItemIterator->Value.GetLocation(InvokerLocation))
 				{
-					const FVector ActorLocation = Actor->GetActorLocation();
-					const float GenerationRadius = ItemIterator->Value.GenerationRadius;
-					bool bKeep = !bCheckMaximumDistanceFromSeeds;
+					ItemIterator.RemoveCurrent();
+					continue;
+				}
 
-					double ClosestDistanceSq = DBL_MAX;
-					if (bCheckMaximumDistanceFromSeeds)
-					{
-						const double CheckDistanceSq = FMath::Square(InvokersMaximumDistanceFromSeed + GenerationRadius);
+				const float GenerationRadius = ItemIterator->Value.GenerationRadius;
+				bool bKeep = !bCheckMaximumDistanceFromSeeds;
 
-						// Check if the invoker is close enough
-						for (const FVector SeedLocation : SeedLocations)
-						{
-							const double InvokerDistanceToSeedSq = FVector::DistSquared(SeedLocation, ActorLocation);
-							if (InvokerDistanceToSeedSq <= CheckDistanceSq)
-							{
-								bKeep = true;
-								break;
-							}
-							else
-							{
-								ClosestDistanceSq = FMath::Min(InvokerDistanceToSeedSq, ClosestDistanceSq);
-							}
-						}	
-					}
+				double ClosestDistanceSq = DBL_MAX;
+				if (bCheckMaximumDistanceFromSeeds)
+				{
+					const double CheckDistanceSq = FMath::Square(InvokersMaximumDistanceFromSeed + GenerationRadius);
 
-					if (bKeep)
-					{
-						InvokerLocations.Add(FNavigationInvokerRaw(ActorLocation, GenerationRadius, ItemIterator->Value.RemovalRadius,
-							ItemIterator->Value.SupportedAgents, ItemIterator->Value.Priority));
-					}
-					else
+					// Check if the invoker is close enough
+					for (const FVector SeedLocation : SeedLocations)
 					{
-						UE_LOG(LogNavInvokers, Verbose, TEXT("Invoker %s ignored because it's too far from any seed location. Closest seed at %.0f."),
-							*Actor->GetName(), FMath::Sqrt(ClosestDistanceSq));
+						const double InvokerDistanceToSeedSq = FVector::DistSquared(SeedLocation, InvokerLocation);
+						if (InvokerDistanceToSeedSq <= CheckDistanceSq)
+						{
+							bKeep = true;
+							break;
+						}
+						else
+						{
+							ClosestDistanceSq = FMath::Min(InvokerDistanceToSeedSq, ClosestDistanceSq);
+						}
 					}
 				}
+
+				if (bKeep)
+				{
+					InvokerLocations.Add(FNavigationInvokerRaw(InvokerLocation, GenerationRadius, ItemIterator->Value.RemovalRadius,
+						ItemIterator->Value.SupportedAgents, ItemIterator->Value.Priority));
+				}
 				else
 				{
-					ItemIterator.RemoveCurrent();
+					UE_LOG(LogNavInvokers, Verbose, TEXT("Invoker %s ignored because it's too far from any seed location. Closest seed at %.0f."),
+						*ItemIterator->Value.GetName(), FMath::Sqrt(ClosestDistanceSq));
 				}
 			}
 
@@ -5205,7 +5293,7 @@ void UNavigationSystemV1::UpdateInvokers()
 
 #if WITH_RECAST
 		const double UpdateStartTime = FPlatformTime::Seconds();
-		for (TActorIterator<ARecastNavMesh> It(GetWorld()); It; ++It)
+		for (TActorIterator<ARecastNavMesh> It(World); It; ++It)
 		{
 			It->UpdateActiveTiles(InvokerLocations);
 		}
diff --git a/Engine/Source/Runtime/NavigationSystem/Public/NavigationSystem.h b/Engine/Source/Runtime/NavigationSystem/Public/NavigationSystem.h
--- a/Engine/Source/Runtime/NavigationSystem/Public/NavigationSystem.h
+++ b/Engine/Source/Runtime/NavigationSystem/Public/NavigationSystem.h
@@ -36,6 +36,7 @@ class ANavMeshBoundsVolume;
 class AWorldSettings;
 class FEdMode;
 class FNavDataGenerator;
+class INavigationInvokerInterface;
 class INavLinkCustomInterface;
 class INavRelevantInterface;
 class UCrowdManagerBase;
@@ -453,10 +454,10 @@ protected:
 	TSet<FNavigationBounds> RegisteredNavBounds;
 
 private:
-	TMap<AActor*, FNavigationInvoker> Invokers;
+	TMap<UObject*, FNavigationInvoker> Invokers;
 	/** Contains pre-digested and cached invokers' info. Generated by UpdateInvokers */
 	TArray<FNavigationInvokerRaw> InvokerLocations;
-	
+
 	TArray<FBox> InvokersSeedBounds;
 
 	double NextInvokersUpdateTime;
@@ -464,6 +465,8 @@ private:
 
 	NAVIGATIONSYSTEM_API void DirtyTilesInBuildBounds();
 
+	void UnregisterInvoker_Internal(const UObject& Invoker);
+
 public:
 	//----------------------------------------------------------------------//
 	// Blueprint functions
@@ -814,8 +817,12 @@ public:
 	NAVIGATIONSYSTEM_API virtual void RegisterInvoker(AActor& Invoker, float TileGenerationRadius, float TileRemovalRadius, const FNavAgentSelector& Agents);
 
 	NAVIGATIONSYSTEM_API virtual void RegisterInvoker(AActor& Invoker, float TileGenerationRadius, float TileRemovalRadius, const FNavAgentSelector& Agents, ENavigationInvokerPriority InPriority);
-	
+
+	NAVIGATIONSYSTEM_API virtual void RegisterInvoker(const TWeakInterfacePtr<INavigationInvokerInterface>& Invoker, float TileGenerationRadius, float TileRemovalRadius, const FNavAgentSelector& Agents, ENavigationInvokerPriority InPriority);
+
 	NAVIGATIONSYSTEM_API virtual void UnregisterInvoker(AActor& Invoker);
+	
+	NAVIGATIONSYSTEM_API virtual void UnregisterInvoker(const TWeakInterfacePtr<INavigationInvokerInterface>& Invoker);
 
 	static NAVIGATIONSYSTEM_API void RegisterNavigationInvoker(AActor& Invoker, float TileGenerationRadius, float TileRemovalRadius, 
 		const FNavAgentSelector& Agents = FNavAgentSelector(), ENavigationInvokerPriority Priority = ENavigationInvokerPriority::Default);
diff --git a/Engine/Source/Runtime/NavigationSystem/Public/NavigationSystemTypes.h b/Engine/Source/Runtime/NavigationSystem/Public/NavigationSystemTypes.h
--- a/Engine/Source/Runtime/NavigationSystem/Public/NavigationSystemTypes.h
+++ b/Engine/Source/Runtime/NavigationSystem/Public/NavigationSystemTypes.h
@@ -8,6 +8,7 @@
 #include "AI/Navigation/NavLinkDefinition.h"
 #include "Math/GenericOctreePublic.h"
 #include "AI/NavigationModifier.h"
+#include "UObject/WeakInterfacePtr.h"
 
 #define NAVSYS_DEBUG (0 && UE_BUILD_DEBUG)
 
@@ -20,6 +21,7 @@ struct FKAggregateGeom;
 class FNavigationOctree;
 class UNavigationPath;
 class ANavigationData;
+class INavigationInvokerInterface;
 
 struct FPathFindingQueryData
 {
@@ -105,7 +107,9 @@ struct FNavigationInvokerRaw
 
 struct FNavigationInvoker
 {
+	/** The Invoker source should be either an Actor or an Object. Thus only 1 of those member should be set. We'll use IsExplicitlyNull to know which one to use */
 	TWeakObjectPtr<AActor> Actor;
+	TWeakInterfacePtr<INavigationInvokerInterface> Object;
 
 	/** tiles GenerationRadius away or close will be generated if they're not already present */
 	float GenerationRadius;
@@ -119,9 +123,13 @@ struct FNavigationInvoker
 
 	/** invoker Priority used when dirtying tiles */
 	ENavigationInvokerPriority Priority;
-	
+
 	FNavigationInvoker();
 	FNavigationInvoker(AActor& InActor, float InGenerationRadius, float InRemovalRadius, const FNavAgentSelector& InSupportedAgents, ENavigationInvokerPriority InPriority);
+	FNavigationInvoker(INavigationInvokerInterface& InObject, float InGenerationRadius, float InRemovalRadius, const FNavAgentSelector& InSupportedAgents, ENavigationInvokerPriority InPriority);
+
+	FString GetName() const;
+	bool GetLocation(FVector& OutLocation) const;
 };
 
 namespace NavigationHelper
