[Navmesh] Tile generation debug improvements:
- addition of an options to skip contour simplification and display raw contours;
- addition of text labels to FRecastInternalDebugData;
- addition of indexes and coordinates to contour debug edges, allowing to identify them individually.

diff --git a/Engine/Source/Runtime/NavigationSystem/Private/NavMesh/NavMeshRenderingComponent.cpp b/Engine/Source/Runtime/NavigationSystem/Private/NavMesh/NavMeshRenderingComponent.cpp
--- a/Engine/Source/Runtime/NavigationSystem/Private/NavMesh/NavMeshRenderingComponent.cpp
+++ b/Engine/Source/Runtime/NavigationSystem/Private/NavMesh/NavMeshRenderingComponent.cpp
@@ -1352,6 +1352,16 @@ void FNavMeshSceneProxyData::AddMeshForInternalData(const FRecastInternalDebugDa
 			}
 		}
 	}
+
+	if (InInternalData.LabelVertices.Num() > 0)
+	{
+		const TArray<FVector>& Vertices = InInternalData.LabelVertices;
+		const TArray<FString>& Labels = InInternalData.Labels;
+		for (int32 i = 0; i < Vertices.Num(); i++)
+		{
+			DebugLabels.Emplace(Vertices[i] + NavMeshDrawOffset, Labels[i]);
+		}
+	}
 }
 #endif //RECAST_INTERNAL_DEBUG_DATA
 
diff --git a/Engine/Source/Runtime/NavigationSystem/Private/NavMesh/RecastNavMesh.cpp b/Engine/Source/Runtime/NavigationSystem/Private/NavMesh/RecastNavMesh.cpp
--- a/Engine/Source/Runtime/NavigationSystem/Private/NavMesh/RecastNavMesh.cpp
+++ b/Engine/Source/Runtime/NavigationSystem/Private/NavMesh/RecastNavMesh.cpp
@@ -501,6 +501,7 @@ FRecastNavMeshTileGenerationDebug::FRecastNavMeshTileGenerationDebug()
 	bCompactHeightfieldDistances = false;
 	bTileCacheLayerAreas = false;
 	bTileCacheLayerRegions = false;
+	bSkipContourSimplification = false;
 	bTileCacheContours = false;
 	bTileCachePolyMesh = false;
 	bTileCacheDetailMesh = false;
diff --git a/Engine/Source/Runtime/NavigationSystem/Private/NavMesh/RecastNavMeshGenerator.cpp b/Engine/Source/Runtime/NavigationSystem/Private/NavMesh/RecastNavMeshGenerator.cpp
--- a/Engine/Source/Runtime/NavigationSystem/Private/NavMesh/RecastNavMeshGenerator.cpp
+++ b/Engine/Source/Runtime/NavigationSystem/Private/NavMesh/RecastNavMeshGenerator.cpp
@@ -2468,6 +2468,8 @@ void FRecastTileGenerator::AppendGeometry(const FNavigationRelevantData& DataRef
 
 ETimeSliceWorkResult FRecastTileGenerator::GenerateTileTimeSliced()
 {
+	UE_LOG(LogNavigation, Verbose, TEXT("Building tile (time sliced): (%i,%i)"), TileX, TileY);
+
 	FNavMeshBuildContext BuildContext(*this);
 	ETimeSliceWorkResult WorkResult = ETimeSliceWorkResult::Succeeded;
 
@@ -2537,6 +2539,8 @@ bool FRecastTileGenerator::GenerateTile()
 	const double StartStamp = FPlatformTime::Seconds();
 	double PostCompressLayerStamp = StartStamp;
 #endif // RECAST_INTERNAL_DEBUG_DATA
+
+	UE_LOG(LogNavigation, Verbose, TEXT("Building tile: (%i,%i)"), TileX, TileY);
 	
 	FNavMeshBuildContext BuildContext(*this);
 	bool bSuccess = true;
@@ -3710,6 +3714,11 @@ bool FRecastTileGenerator::GenerateNavigationDataLayer(FNavMeshBuildContext& Bui
 			return false;
 		}
 
+		bool bSkipContourSimplification = false;
+#if RECAST_INTERNAL_DEBUG_DATA
+		bSkipContourSimplification = IsTileDebugActive() && TileDebugSettings.bSkipContourSimplification;
+#endif
+		
 #if WITH_NAVMESH_CLUSTER_LINKS
 		GenerationContext.ClusterSet = dtAllocTileCacheClusterSet(&GenNavAllocator);
 		if (GenerationContext.ClusterSet == nullptr)
@@ -3720,11 +3729,11 @@ bool FRecastTileGenerator::GenerateNavigationDataLayer(FNavMeshBuildContext& Bui
 
 		status = dtBuildTileCacheContours(&GenNavAllocator, *GenerationContext.Layer,
 			TileConfig.walkableClimb, TileConfig.maxVerticalMergeError, TileConfig.maxSimplificationError, TileConfig.simplificationElevationRatio,
-			TileConfig.cs, TileConfig.ch,*GenerationContext.ContourSet, *GenerationContext.ClusterSet);
+			TileConfig.cs, TileConfig.ch,*GenerationContext.ContourSet, *GenerationContext.ClusterSet, bSkipContourSimplification);
 #else
 		status = dtBuildTileCacheContours(&GenNavAllocator, *GenerationContext.Layer,
 			TileConfig.walkableClimb, TileConfig.maxVerticalMergeError, TileConfig.maxSimplificationError, TileConfig.simplificationElevationRatio,
-			TileConfig.cs, TileConfig.ch, *GenerationContext.ContourSet);
+			TileConfig.cs, TileConfig.ch, *GenerationContext.ContourSet, bSkipContourSimplification);
 #endif //WITH_NAVMESH_CLUSTER_LINKS
 		
 		if (dtStatusFailed(status))
@@ -3743,7 +3752,7 @@ bool FRecastTileGenerator::GenerateNavigationDataLayer(FNavMeshBuildContext& Bui
 #if RECAST_INTERNAL_DEBUG_DATA
 	if (IsTileDebugActive() && TileDebugSettings.bTileCacheContours)
 	{
-		duDebugDrawTileCacheContours(&BuildContext.InternalDebugData, *GenerationContext.ContourSet, GenerationContext.Layer->header->bmin, TileConfig.cs, TileConfig.ch);
+		duDebugDrawTileCacheContours(&BuildContext.InternalDebugData, *GenerationContext.ContourSet, LayerIdx, GenerationContext.Layer->header->bmin, TileConfig.cs, TileConfig.ch);
 	}
 #endif
 
diff --git a/Engine/Source/Runtime/NavigationSystem/Public/NavMesh/RecastInternalDebugData.cpp b/Engine/Source/Runtime/NavigationSystem/Public/NavMesh/RecastInternalDebugData.cpp
--- a/Engine/Source/Runtime/NavigationSystem/Public/NavMesh/RecastInternalDebugData.cpp
+++ b/Engine/Source/Runtime/NavigationSystem/Public/NavMesh/RecastInternalDebugData.cpp
@@ -12,7 +12,7 @@
 
 void FRecastInternalDebugData::vertex(const FVector::FReal x, const FVector::FReal y, const FVector::FReal z, unsigned int color, const FVector::FReal u, const FVector::FReal v)
 {
-	FVector::FReal RecastPos[3] = { x,y,z };
+	const FVector::FReal RecastPos[3] = { x,y,z };
 	const FVector Pos = Recast2UnrealPoint(RecastPos);
 	const FColor Color = Recast2UnrealColor(color);
 	switch(CurrentPrim)
@@ -34,6 +34,14 @@ void FRecastInternalDebugData::vertex(const FVector::FReal x, const FVector::FRe
 	}
 }
 
+void FRecastInternalDebugData::text(const FVector::FReal x, const FVector::FReal y, const FVector::FReal z, const char* text)
+{
+	const FVector::FReal RecastPos[3] = { x,y,z };
+	const FVector Pos = Recast2UnrealPoint(RecastPos);
+	LabelVertices.Push(Pos);
+	Labels.Push(FString(text));
+}
+
 void FRecastInternalDebugData::end()
 {
 	if (CurrentPrim == DU_DRAW_QUADS)
diff --git a/Engine/Source/Runtime/NavigationSystem/Public/NavMesh/RecastInternalDebugData.h b/Engine/Source/Runtime/NavigationSystem/Public/NavMesh/RecastInternalDebugData.h
--- a/Engine/Source/Runtime/NavigationSystem/Public/NavMesh/RecastInternalDebugData.h
+++ b/Engine/Source/Runtime/NavigationSystem/Public/NavMesh/RecastInternalDebugData.h
@@ -27,6 +27,9 @@ struct FRecastInternalDebugData : public duDebugDraw
 	TArray<FVector> PointVertices;
 	TArray<FColor>  PointColors;
 
+	TArray<FVector> LabelVertices;
+	TArray<FString> Labels;
+
 	double BuildTime = 0.;
 	double BuildCompressedLayerTime = 0.;
 	double BuildNavigationDataTime = 0.;
@@ -63,6 +66,8 @@ struct FRecastInternalDebugData : public duDebugDraw
 
 	virtual void vertex(const FVector::FReal x, const FVector::FReal y, const FVector::FReal z, unsigned int color, const FVector::FReal u, const FVector::FReal v) override;
 
+	virtual void text(const FVector::FReal x, const FVector::FReal y, const FVector::FReal z, const char* text) override;
+
 	virtual void end() override;
 };
 #endif // WITH_RECAST
diff --git a/Engine/Source/Runtime/NavigationSystem/Public/NavMesh/RecastNavMesh.h b/Engine/Source/Runtime/NavigationSystem/Public/NavMesh/RecastNavMesh.h
--- a/Engine/Source/Runtime/NavigationSystem/Public/NavMesh/RecastNavMesh.h
+++ b/Engine/Source/Runtime/NavigationSystem/Public/NavMesh/RecastNavMesh.h
@@ -504,6 +504,10 @@ struct FRecastNavMeshTileGenerationDebug
 	UPROPERTY(EditAnywhere, Category = Debug)
 	uint32 bTileCacheLayerRegions : 1;
 
+	/** If set, the contour simplification step will be skipped. Beware that enabling this changes the way navmesh will generate when Tile Generation Debug is enabled. */
+	UPROPERTY(EditAnywhere, Category = Debug)
+	uint32 bSkipContourSimplification : 1;
+	
 	UPROPERTY(EditAnywhere, Category = Debug)
 	uint32 bTileCacheContours : 1;
 
diff --git a/Engine/Source/Runtime/Navmesh/Private/DebugUtils/DetourDebugDraw.cpp b/Engine/Source/Runtime/Navmesh/Private/DebugUtils/DetourDebugDraw.cpp
--- a/Engine/Source/Runtime/Navmesh/Private/DebugUtils/DetourDebugDraw.cpp
+++ b/Engine/Source/Runtime/Navmesh/Private/DebugUtils/DetourDebugDraw.cpp
@@ -818,7 +818,7 @@ struct dtTileCacheContourSet
 };*/
 
 void duDebugDrawTileCacheContours(duDebugDraw* dd, const struct dtTileCacheContourSet& lcset,
-								  const dtReal* orig, const dtReal cs, const dtReal ch)
+								  const int32 layerIdx, const dtReal* orig, const dtReal cs, const dtReal ch)
 {
 	if (!dd) return;
 	
@@ -840,11 +840,12 @@ void duDebugDrawTileCacheContours(duDebugDraw* dd, const struct dtTileCacheConto
 			const int k = (j+1) % c.nverts;
 			const unsigned short* va = &c.verts[j*4];
 			const unsigned short* vb = &c.verts[k*4];
+			const dtReal verticalOffset = 1 + (0.05*i) + (0.3*layerIdx);
 			const dtReal ax = orig[0] + va[0]*cs;
-			const dtReal ay = orig[1] + (va[1]+1+(i&1))*ch;
+			const dtReal ay = orig[1] + (va[1]+verticalOffset)*ch;
 			const dtReal az = orig[2] + va[2]*cs;
 			const dtReal bx = orig[0] + vb[0]*cs;
-			const dtReal by = orig[1] + (vb[1]+1+(i&1))*ch;
+			const dtReal by = orig[1] + (vb[1]+verticalOffset)*ch;
 			const dtReal bz = orig[2] + vb[2]*cs;
 			unsigned int col = color;
 			if ((va[3] & 0xf) != 0xf)
@@ -866,7 +867,7 @@ void duDebugDrawTileCacheContours(duDebugDraw* dd, const struct dtTileCacheConto
 				dd->vertex(dx,dy,dz,duRGBA(255,0,0,255));
 			}
 			
-			duAppendArrow(dd, ax,ay,az, bx,by,bz, 0.0f, cs*0.5f, col);
+			duAppendArrow(dd, ax,ay,az, bx,by,bz, 0.0, cs*0.2, col);
 		}
 	}
 	dd->end();
@@ -890,9 +891,16 @@ void duDebugDrawTileCacheContours(duDebugDraw* dd, const struct dtTileCacheConto
 			}
 			
 			dtReal fx = orig[0] + va[0]*cs;
-			dtReal fy = orig[1] + (va[1]+1+(i&1))*ch;
+			const dtReal verticalOffset = 1 + (0.05*i) + (0.3*layerIdx);
+			dtReal fy = orig[1] + (va[1]+verticalOffset)*ch;
 			dtReal fz = orig[2] + va[2]*cs;
 			dd->vertex(fx,fy,fz, color);
+
+			// Add label for contour edge.
+			constexpr int bufferSize = 128;
+			char text[bufferSize];
+			snprintf(text, bufferSize, "%i (%i:%i) (%u,%u,%u) neiReg:%u r:%u a:%u", layerIdx, i, j, va[0], va[1], va[2], va[3], c.reg, c.area);
+			dd->text(fx,fy,fz, text);
 		}
 	}
 	dd->end();
diff --git a/Engine/Source/Runtime/Navmesh/Private/DetourTileCache/DetourTileCacheBuilder.cpp b/Engine/Source/Runtime/Navmesh/Private/DetourTileCache/DetourTileCacheBuilder.cpp
--- a/Engine/Source/Runtime/Navmesh/Private/DetourTileCache/DetourTileCacheBuilder.cpp
+++ b/Engine/Source/Runtime/Navmesh/Private/DetourTileCache/DetourTileCacheBuilder.cpp
@@ -22,6 +22,7 @@
 #include "DetourTileCache/DetourTileCacheBuilder.h"
 #include "Detour/DetourCommon.h"
 #include "Detour/DetourAssert.h"
+#include "DebugUtils/DetourDebugDraw.h"
 #include "Stats/Stats.h"
 #include <limits>
 
@@ -275,7 +276,7 @@ static void getNeighbourRegAndArea(dtTileCacheLayer& layer,
 	}
 }
 
-static bool walkContour(dtTileCacheLayer& layer, int x, int y, int idx, const int maxVerticalMergeError, unsigned char* flags, dtTempContour& cont) // UE
+static bool walkContour(dtTileCacheLayer& layer, int x, int y, int idx, const int maxVerticalMergeError, unsigned char* flags, dtTempContour& cont, int& contourIndex) // UE
 {
 	const int w = (int)layer.header->width;
 	const int h = (int)layer.header->height;
@@ -405,6 +406,8 @@ static bool walkContour(dtTileCacheLayer& layer, int x, int y, int idx, const in
 	}	
 //@UE END
 	
+	contourIndex++;
+
 	return true;
 }	
 
@@ -794,6 +797,7 @@ dtStatus dtBuildTileCacheContours(dtTileCacheAlloc* alloc, dtTileCacheLayer& lay
 #if WITH_NAVMESH_CLUSTER_LINKS
 	, dtTileCacheClusterSet& clusters
 #endif //WITH_NAVMESH_CLUSTER_LINKS
+	, const bool skipContourSimplification /*=false*/
 	//@UE END
 	)
 {
@@ -869,6 +873,7 @@ dtStatus dtBuildTileCacheContours(dtTileCacheAlloc* alloc, dtTileCacheLayer& lay
 	dtIntArray linksBase(maxConts);
 
 	// Find contours.
+	int contourIndex = 0;	// UE
 	for (int y = 0; y < h; ++y)
 	{
 		for (int x = 0; x < w; ++x)
@@ -883,14 +888,17 @@ dtStatus dtBuildTileCacheContours(dtTileCacheAlloc* alloc, dtTileCacheLayer& lay
 			if (ri == 0xffff || ri == 0)
 				continue;
 
-			if (!walkContour(layer, x, y, idx, maxVerticalMergeError, flags, temp)) // UE
+			if (!walkContour(layer, x, y, idx, maxVerticalMergeError, flags, temp, contourIndex)) // UE
 			{
 				// Too complex contour.
 				// Note: If you hit here often, try increasing 'maxTempVerts'.
 				return DT_FAILURE | DT_BUFFER_TOO_SMALL;
 			}
 
-			simplifyContour(layer.areas[idx], ri, temp, maxError, simplificationElevationRatio, cs, ch); // UE
+			if (!skipContourSimplification)
+			{
+				simplifyContour(layer.areas[idx], ri, temp, maxError, simplificationElevationRatio, cs, ch); // UE
+			}
 
 			// Store contour.
 			if (lcset.nconts >= maxConts)
diff --git a/Engine/Source/Runtime/Navmesh/Public/DebugUtils/DebugDraw.h b/Engine/Source/Runtime/Navmesh/Public/DebugUtils/DebugDraw.h
--- a/Engine/Source/Runtime/Navmesh/Public/DebugUtils/DebugDraw.h
+++ b/Engine/Source/Runtime/Navmesh/Public/DebugUtils/DebugDraw.h
@@ -134,6 +134,17 @@ struct duDebugDraw
 #endif
 //@UE END
 
+	/// Submit a text label
+	///  @param x,y,z [in] 3D position of the text.
+	///  @param text  [in] text to display.
+	virtual void text(const duReal x, const duReal y, const duReal z, const char* text) = 0;
+//@UE BEGIN Adding support for LWCoords.
+#if !DU_LARGE_WORLD_COORDINATES_DISABLED
+	// float version of this function is now deprecated
+	virtual void text(const float x, const float y, const float z, const char* text) final {}
+#endif
+//@UE END
+	
 	/// End drawing primitives.
 	virtual void end() = 0;
 };
diff --git a/Engine/Source/Runtime/Navmesh/Public/DebugUtils/DetourDebugDraw.h b/Engine/Source/Runtime/Navmesh/Public/DebugUtils/DetourDebugDraw.h
--- a/Engine/Source/Runtime/Navmesh/Public/DebugUtils/DetourDebugDraw.h
+++ b/Engine/Source/Runtime/Navmesh/Public/DebugUtils/DetourDebugDraw.h
@@ -60,7 +60,7 @@ NAVMESH_API void duDebugDrawTileCacheLayerDistance(struct duDebugDraw* dd, const
 									   const dtTileCacheDistanceField& dfield, const dtReal cs, const dtReal ch);
 NAVMESH_API void duDebugDrawTileCacheLayerRegions(struct duDebugDraw* dd, const dtTileCacheLayer& layer, const dtReal cs, const dtReal ch);
 NAVMESH_API void duDebugDrawTileCacheContours(duDebugDraw* dd, const struct dtTileCacheContourSet& lcset,
-								  const dtReal* orig, const dtReal cs, const dtReal ch);
+								  const int32 layerIdx, const dtReal* orig, const dtReal cs, const dtReal ch);
 NAVMESH_API void duDebugDrawTileCachePolyMesh(duDebugDraw* dd, const struct dtTileCachePolyMesh& lmesh,
 								  const dtReal* orig, const dtReal cs, const dtReal ch);
 NAVMESH_API void duDebugDrawTileCacheDetailMesh(duDebugDraw* dd, const struct dtTileCachePolyMeshDetail& dmesh);
diff --git a/Engine/Source/Runtime/Navmesh/Public/DetourTileCache/DetourTileCacheBuilder.h b/Engine/Source/Runtime/Navmesh/Public/DetourTileCache/DetourTileCacheBuilder.h
--- a/Engine/Source/Runtime/Navmesh/Public/DetourTileCache/DetourTileCacheBuilder.h
+++ b/Engine/Source/Runtime/Navmesh/Public/DetourTileCache/DetourTileCacheBuilder.h
@@ -270,6 +270,7 @@ NAVMESH_API dtStatus dtBuildTileCacheContours(dtTileCacheAlloc* alloc,
 #if WITH_NAVMESH_CLUSTER_LINKS
 								, dtTileCacheClusterSet& lclusters
 #endif // WITH_NAVMESH_CLUSTER_LINKS
+								, const bool skipContourSimplification = false
 								//@UE END
 								);
 
