[Navmesh] Update navmesh display lines to better distinguish edges borders and polygons borders.

diff --git a/Engine/Source/Runtime/NavigationSystem/Private/NavMesh/NavMeshRenderingComponent.cpp b/Engine/Source/Runtime/NavigationSystem/Private/NavMesh/NavMeshRenderingComponent.cpp
--- a/Engine/Source/Runtime/NavigationSystem/Private/NavMesh/NavMeshRenderingComponent.cpp
+++ b/Engine/Source/Runtime/NavigationSystem/Private/NavMesh/NavMeshRenderingComponent.cpp
@@ -41,8 +41,8 @@ static const FColor NavMeshRenderColor_PolyForbidden(FColorList::Black);
 
 static constexpr float DefaultEdges_LineThickness = 0.0f;
 static constexpr float TileResolution_LineThickness = 5.f;
-static constexpr float PolyEdges_LineThickness = 1.5f;
-static constexpr float NavMeshEdges_LineThickness = 3.5f;
+static constexpr float PolyEdges_LineThickness = 1.1f;
+static constexpr float NavMeshEdges_LineThickness = 4.f;
 static constexpr float LinkLines_LineThickness = 2.0f;
 static constexpr float ClusterLinkLines_LineThickness = 2.0f;
 
