PR #11080: Change bit-field type from int32 to uint32 in FRecastDebugGeometry

diff --git a/Engine/Source/Runtime/NavigationSystem/Public/NavMesh/RecastNavMesh.h b/Engine/Source/Runtime/NavigationSystem/Public/NavMesh/RecastNavMesh.h
--- a/Engine/Source/Runtime/NavigationSystem/Public/NavMesh/RecastNavMesh.h
+++ b/Engine/Source/Runtime/NavigationSystem/Public/NavMesh/RecastNavMesh.h
@@ -253,10 +253,10 @@ struct FRecastDebugGeometry
 	TArray<FIntPoint> TilesToDisplayInternalData;
 #endif
 
-	int32 bGatherPolyEdges : 1;
-	int32 bGatherNavMeshEdges : 1;
-	int32 bMarkForbiddenPolys : 1;
-	int32 bGatherTileBuildTimesHeatMap : 1;
+	uint32 bGatherPolyEdges : 1;
+	uint32 bGatherNavMeshEdges : 1;
+	uint32 bMarkForbiddenPolys : 1;
+	uint32 bGatherTileBuildTimesHeatMap : 1;
 
 	double MinTileBuildTime = DBL_MAX;
 	double MaxTileBuildTime = 0.;
