Added a new type of NavModifierComponent which assigns a given nav area to all parts of the nav mesh in the vicinity of a spline.

diff --git a/Engine/Source/Runtime/NavigationSystem/Private/NavModifierComponent.cpp b/Engine/Source/Runtime/NavigationSystem/Private/NavModifierComponent.cpp
--- a/Engine/Source/Runtime/NavigationSystem/Private/NavModifierComponent.cpp
+++ b/Engine/Source/Runtime/NavigationSystem/Private/NavModifierComponent.cpp
@@ -123,9 +123,72 @@ void UNavModifierComponent::PopulateComponentBounds(FTransform InParentTransform
 	}
 }
 
+void UNavModifierComponent::CalculateBounds() const
+{
+	const AActor* MyOwner = GetOwner();
+	if (!MyOwner)
+	{
+		return;
+	}
+
+	Bounds = FBox(ForceInit);
+    ComponentBounds.Reset();
+    for (UActorComponent* Component : MyOwner->GetComponents())
+    {
+    	UPrimitiveComponent* PrimComp = Cast<UPrimitiveComponent>(Component);
+    	if (PrimComp && PrimComp->IsRegistered() && PrimComp->IsCollisionEnabled() && PrimComp->CanEverAffectNavigation())
+    	{
+    		UBodySetup* BodySetup = PrimComp->GetBodySetup();
+    		if (BodySetup)
+    		{
+    			Bounds += PrimComp->Bounds.GetBox();
+    				
+    			const FTransform& ParentTM = PrimComp->GetComponentTransform();
+    			PopulateComponentBounds(ParentTM, *BodySetup);
+    		}
+    		else if (const UGeometryCollectionComponent* GeometryCollection = Cast<UGeometryCollectionComponent>(PrimComp))
+    		{
+    			// If it's a GC, use the bodySetups from the proxyMeshes.
+    			if (const TObjectPtr<const UGeometryCollection> RestCollection = GeometryCollection->RestCollection)
+    			{
+    				Bounds += GeometryCollection->Bounds.GetBox();
+
+    				for (const TObjectPtr<UStaticMesh>& ProxyMesh : RestCollection->RootProxyData.ProxyMeshes)
+    				{
+    					if (ProxyMesh != nullptr)
+    					{
+    						const UBodySetup* Body = ProxyMesh->GetBodySetup();
+    						if (Body)
+    						{
+    							const FTransform& ParentTM = PrimComp->GetComponentTransform();
+    							PopulateComponentBounds(ParentTM, *Body);	
+    						}
+    					}
+    				}
+    			}
+    		}	
+    	}
+    }
+
+    if (ComponentBounds.Num() == 0)
+    {
+    	Bounds = FBox::BuildAABB(MyOwner->GetActorLocation(), FailsafeExtent);
+    	ComponentBounds.Add(FRotatedBox(Bounds, MyOwner->GetActorQuat()));
+    }
+
+    for (int32 Idx = 0; Idx < ComponentBounds.Num(); Idx++)
+    {
+    	const FVector BoxOrigin = ComponentBounds[Idx].Box.GetCenter();
+    	const FVector BoxExtent = ComponentBounds[Idx].Box.GetExtent();
+
+    	const FVector NavModBoxOrigin = FTransform(ComponentBounds[Idx].Quat).InverseTransformPosition(BoxOrigin);
+    	ComponentBounds[Idx].Box = FBox::BuildAABB(NavModBoxOrigin, BoxExtent);
+    }
+}
+
 void UNavModifierComponent::CalcAndCacheBounds() const
 {
-	AActor* MyOwner = GetOwner();
+	const AActor* MyOwner = GetOwner();
 	if (MyOwner)
 	{
 		CachedTransform = MyOwner->GetActorTransform();
@@ -139,62 +202,10 @@ void UNavModifierComponent::CalcAndCacheBounds() const
 			// We're filtering for nav relevancy in OnTransformUpdated.
 			TransformUpdateHandle = MyOwner->GetRootComponent()->TransformUpdated.AddUObject(const_cast<UNavModifierComponent*>(this), &UNavModifierComponent::OnTransformUpdated);
 		}
-
-		Bounds = FBox(ForceInit);
-		ComponentBounds.Reset();
-		for (UActorComponent* Component : MyOwner->GetComponents())
-		{
-			UPrimitiveComponent* PrimComp = Cast<UPrimitiveComponent>(Component);
-			if (PrimComp && PrimComp->IsRegistered() && PrimComp->IsCollisionEnabled() && PrimComp->CanEverAffectNavigation())
-			{
-				UBodySetup* BodySetup = PrimComp->GetBodySetup();
-				if (BodySetup)
-				{
-					Bounds += PrimComp->Bounds.GetBox();
-					
-					const FTransform& ParentTM = PrimComp->GetComponentTransform();
-					PopulateComponentBounds(ParentTM, *BodySetup);
-				}
-				else if (const UGeometryCollectionComponent* GeometryCollection = Cast<UGeometryCollectionComponent>(PrimComp))
-				{
-					// If it's a GC, use the bodySetups from the proxyMeshes.
-					if (const TObjectPtr<const UGeometryCollection> RestCollection = GeometryCollection->RestCollection)
-					{
-						Bounds += GeometryCollection->Bounds.GetBox();
-
-						for (const TObjectPtr<UStaticMesh>& ProxyMesh : RestCollection->RootProxyData.ProxyMeshes)
-						{
-							if (ProxyMesh != nullptr)
-							{
-								const UBodySetup* Body = ProxyMesh->GetBodySetup();
-								if (Body)
-								{
-									const FTransform& ParentTM = PrimComp->GetComponentTransform();
-									PopulateComponentBounds(ParentTM, *Body);	
-								}
-							}
-						}
-					}
-				}	
-			}
-		}
-
-		if (ComponentBounds.Num() == 0)
-		{
-			Bounds = FBox::BuildAABB(MyOwner->GetActorLocation(), FailsafeExtent);
-			ComponentBounds.Add(FRotatedBox(Bounds, MyOwner->GetActorQuat()));
-		}
-
-		for (int32 Idx = 0; Idx < ComponentBounds.Num(); Idx++)
-		{
-			const FVector BoxOrigin = ComponentBounds[Idx].Box.GetCenter();
-			const FVector BoxExtent = ComponentBounds[Idx].Box.GetExtent();
-
-			const FVector NavModBoxOrigin = FTransform(ComponentBounds[Idx].Quat).InverseTransformPosition(BoxOrigin);
-			ComponentBounds[Idx].Box = FBox::BuildAABB(NavModBoxOrigin, BoxExtent);
-		}
 	}
 
+	CalculateBounds();
+
 	UE_SUPPRESS(LogNavigation, VeryVerbose,
 	{
 		TArray<FAreaNavModifier> Areas;
diff --git a/Engine/Source/Runtime/NavigationSystem/Private/SplineNavModifierComponent.cpp b/Engine/Source/Runtime/NavigationSystem/Private/SplineNavModifierComponent.cpp
new file mode 100644
--- /dev/null
+++ b/Engine/Source/Runtime/NavigationSystem/Private/SplineNavModifierComponent.cpp
@@ -0,0 +1,85 @@
+// Copyright Epic Games, Inc. All Rights Reserved.
+
+#include "SplineNavModifierComponent.h"
+
+#include "AI/NavigationSystemBase.h"
+#include "Components/SplineComponent.h"
+#include "VisualLogger/VisualLogger.h"
+
+#include UE_INLINE_GENERATED_CPP_BY_NAME(SplineNavModifierComponent)
+
+namespace
+{
+	const USplineComponent* GetSpline(const AActor* Owner)
+	{
+		if (!Owner)
+		{
+			UE_LOG(LogNavigation, Warning, TEXT("USplineNavModifierComponent has no owner, cannot proceed"));
+			return nullptr;
+		}
+
+		const USplineComponent* Spline = Owner->GetComponentByClass<USplineComponent>();
+		UE_CVLOG_UELOG(Spline, Owner, LogNavigation, Warning, TEXT("USplineNavModifierComponent attached to \"%s\" could not find a spline component, cannot proceed"), *Owner->GetName());
+
+		return Spline;
+	}
+}
+
+void USplineNavModifierComponent::CalculateBounds() const
+{
+	if (const USplineComponent* Spline = GetSpline(GetOwner()))
+	{
+		Bounds = Spline->CalcBounds(FTransform::Identity).GetBox();
+	}
+}
+
+void USplineNavModifierComponent::GetNavigationData(FNavigationRelevantData& Data) const
+{
+	const USplineComponent* Spline = GetSpline(GetOwner());
+	if (!Spline)
+	{
+		return;
+	}
+
+	// Square in the YZ plane used to sample the spline at each cross section
+	constexpr int32 SampleVertices = 4;
+	TStaticArray<FVector, SampleVertices> SampleSquare;
+	SampleSquare[0] = FVector(0, -SplineExtent, -SplineExtent);
+	SampleSquare[1] = FVector(0,  SplineExtent, -SplineExtent);
+	SampleSquare[2] = FVector(0,  SplineExtent,  SplineExtent);
+	SampleSquare[3] = FVector(0, -SplineExtent,  SplineExtent);
+
+	// Vertices (in no particular order) of a prism which will enclose each segment of the spline
+	// @Note The only reason Tube isn't a TStaticArray is because FAreaNavModifier expects a TArray
+	TArray<FVector> Tube;
+	Tube.SetNum(SampleSquare.Num() * 2);
+
+	// Always sample at least the start and end points
+	const int32 NumPoints = FMath::Max(NumSplineSamples, 2);
+	
+	const float SplineLength = Spline->GetSplineLength();
+	FTransform PreviousTransform;
+
+	for (int32 SampleIndex = 0; SampleIndex < NumPoints; SampleIndex++)
+	{
+		// Sample a point on the spline at the current distance
+		const double Distance = SplineLength * SampleIndex / (NumPoints - 1);
+		const FTransform CurrentTransform = Spline->GetTransformAtDistanceAlongSpline(Distance, ESplineCoordinateSpace::World);
+
+		if (SampleIndex > 0)
+		{
+			// Compute the vertices of the current tube segment
+			for (int i = 0; i < SampleVertices; i++)
+			{
+				Tube[i] = PreviousTransform.TransformPosition(SampleSquare[i]);
+				Tube[i + SampleVertices] = CurrentTransform.TransformPosition(SampleSquare[i]);
+			}
+
+			// From the tube construct a convex hull whose volume will be used to mark the nav mesh with the selected AreaClass
+			const FAreaNavModifier NavModifier(Tube, ENavigationCoordSystem::Type::Unreal, FTransform::Identity, AreaClass);
+			Data.Modifiers.Add(NavModifier);
+		}
+		
+		PreviousTransform = CurrentTransform;
+	}
+}
\ No newline at end of file
diff --git a/Engine/Source/Runtime/NavigationSystem/Public/NavModifierComponent.h b/Engine/Source/Runtime/NavigationSystem/Public/NavModifierComponent.h
--- a/Engine/Source/Runtime/NavigationSystem/Public/NavModifierComponent.h
+++ b/Engine/Source/Runtime/NavigationSystem/Public/NavModifierComponent.h
@@ -33,6 +33,10 @@ class UNavModifierComponent : public UNavRelevantComponent
 	UPROPERTY(config, EditAnywhere, Category = Navigation)
 	uint8 bIncludeAgentHeight : 1;
 
+
+	// Does the actual calculating and caching of the bounds when called by CalcAndCacheBounds
+	NAVIGATIONSYSTEM_API virtual void CalculateBounds() const;
+	// @Note We might make this function non-virtual in the future in favor of child classes overriding CalculateBounds, see #jira UE-202451
 	NAVIGATIONSYSTEM_API virtual void CalcAndCacheBounds() const override;
 	NAVIGATIONSYSTEM_API virtual void GetNavigationData(FNavigationRelevantData& Data) const override;
 
diff --git a/Engine/Source/Runtime/NavigationSystem/Public/SplineNavModifierComponent.h b/Engine/Source/Runtime/NavigationSystem/Public/SplineNavModifierComponent.h
new file mode 100644
--- /dev/null
+++ b/Engine/Source/Runtime/NavigationSystem/Public/SplineNavModifierComponent.h
@@ -0,0 +1,29 @@
+// Copyright Epic Games, Inc. All Rights Reserved.
+#pragma once
+
+#include "NavModifierComponent.h"
+
+#include "SplineNavModifierComponent.generated.h"
+
+struct FNavigationRelevantData;
+
+/**
+ *	Assign a chosen NavArea in the vicinity of a spline
+ *	This component must only be attached to an actor with a USplineComponent
+ */
+UCLASS(Blueprintable, Meta = (BlueprintSpawnableComponent), hidecategories = (Variable, Tags, Cooking, Collision), MInimalAPI)
+class USplineNavModifierComponent : public UNavModifierComponent
+{
+	GENERATED_BODY()
+
+	/** Radius of the tube surrounding the spline which will be used to generate the Nav Modifiers */
+	UPROPERTY(EditAnywhere, Category = Navigation, Meta=(UIMin="1", ClampMin="1"))
+	double SplineExtent = 1000.0f;
+
+	/** How many sections the spline will be divided into for generating overlap volumes. More samples result in finer detail */
+	UPROPERTY(EditAnywhere, Category = Navigation, Meta=(UIMin="2", ClampMin="2"))
+	int32 NumSplineSamples = 100;
+	
+	virtual void CalculateBounds() const override;
+	virtual void GetNavigationData(FNavigationRelevantData& Data) const override;
+};
