Two small fixes to SplineNavModifierComponent: flipped a condition for error logging, and accounted for the owner transform when calculating bounds

<saved by Perforce>

diff --git a/Engine/Source/Runtime/NavigationSystem/Private/SplineNavModifierComponent.cpp b/Engine/Source/Runtime/NavigationSystem/Private/SplineNavModifierComponent.cpp
--- a/Engine/Source/Runtime/NavigationSystem/Private/SplineNavModifierComponent.cpp
+++ b/Engine/Source/Runtime/NavigationSystem/Private/SplineNavModifierComponent.cpp
@@ -19,7 +19,7 @@ namespace
 		}
 
 		const USplineComponent* Spline = Owner->GetComponentByClass<USplineComponent>();
-		UE_CVLOG_UELOG(Spline, Owner, LogNavigation, Warning, TEXT("USplineNavModifierComponent attached to \"%s\" could not find a spline component, cannot proceed"), *Owner->GetName());
+		UE_CVLOG_UELOG(!Spline, Owner, LogNavigation, Warning, TEXT("USplineNavModifierComponent attached to \"%s\" could not find a spline component, cannot proceed"), *Owner->GetName());
 
 		return Spline;
 	}
@@ -29,7 +29,7 @@ void USplineNavModifierComponent::CalculateBounds() const
 {
 	if (const USplineComponent* Spline = GetSpline(GetOwner()))
 	{
-		Bounds = Spline->CalcBounds(FTransform::Identity).GetBox();
+		Bounds = Spline->CalcBounds(GetOwner()->GetTransform()).GetBox();
 	}
 }
 
