Second part to fixing DynamicModifierOnly navmesh missing tiles when using nav modifiers.
Don't add an unused area modifier to ANavModifierVolume when using navmesh resolution.
When adding dirty areas make sure to prevent bRebuildGeometry for game world static navmeshes.

diff --git a/Engine/Source/Runtime/NavigationSystem/Private/NavMesh/RecastNavMeshGenerator.cpp b/Engine/Source/Runtime/NavigationSystem/Private/NavMesh/RecastNavMeshGenerator.cpp
--- a/Engine/Source/Runtime/NavigationSystem/Private/NavMesh/RecastNavMeshGenerator.cpp
+++ b/Engine/Source/Runtime/NavigationSystem/Private/NavMesh/RecastNavMeshGenerator.cpp
@@ -1717,6 +1717,8 @@ void FRecastTileGenerator::Setup(const FRecastNavMeshGenerator& ParentGenerator,
 		}
 	}
 
+	// If there are no DirtyAreas, we expect there is a geometry change (also see usage of bRegenerateCompressedLayers).
+	// Else it's a modifier change.
 	const bool bGeometryChanged = (DirtyAreas.Num() == 0);
 	if (!bGeometryChanged)
 	{
@@ -2388,7 +2390,8 @@ void FRecastTileGenerator::AppendModifier(const FCompositeNavModifier& Modifier,
 		OffmeshLinks.Add(SimpleLinkCollection);
 	}
 
-	if (Modifier.GetAreas().Num() == 0)
+	// Navmesh resolutions is a modifier without area, if present, it must not be skipped.
+	if (Modifier.GetAreas().Num() == 0 && Modifier.GetNavMeshResolution() == ENavigationDataResolution::Invalid)
 	{
 		return;
 	}
@@ -6039,7 +6042,7 @@ void FRecastNavMeshGenerator::MarkDirtyTiles(const TArray<FNavigationDirtyArea>&
 			continue;
 		}
 
-		// Static navmeshes accept only area modifiers updates
+		// Game world static navmeshes accept only area modifiers updates
 		if (bGameStaticNavMesh && (!DirtyArea.HasFlag(ENavigationDirtyFlag::DynamicModifier) || DirtyArea.HasFlag(ENavigationDirtyFlag::NavigationBounds)))
 		{
 			continue;
@@ -6144,7 +6147,9 @@ void FRecastNavMeshGenerator::MarkDirtyTiles(const TArray<FNavigationDirtyArea>&
 											
 					FPendingTileElement Element;
 					Element.Coord = FIntPoint(TileX, TileY);
-					Element.bRebuildGeometry = DirtyArea.HasFlag(ENavigationDirtyFlag::Geometry) || DirtyArea.HasFlag(ENavigationDirtyFlag::NavigationBounds);
+					// Make sure to prevent bRebuildGeometry for game world static navmeshes.
+					// Game world static navmeshes accept only area modifiers updates. Rebuilding geometry would bRegenerateCompressedLayers without having the geometry for them.
+					Element.bRebuildGeometry = !bGameStaticNavMesh && (DirtyArea.HasFlag(ENavigationDirtyFlag::Geometry) || DirtyArea.HasFlag(ENavigationDirtyFlag::NavigationBounds));
 					Element.CreationTime = CurrentTimeSeconds;
 					if (Element.bRebuildGeometry == false)
 					{
diff --git a/Engine/Source/Runtime/NavigationSystem/Private/NavModifierVolume.cpp b/Engine/Source/Runtime/NavigationSystem/Private/NavModifierVolume.cpp
--- a/Engine/Source/Runtime/NavigationSystem/Private/NavModifierVolume.cpp
+++ b/Engine/Source/Runtime/NavigationSystem/Private/NavModifierVolume.cpp
@@ -124,10 +124,7 @@ void ANavModifierVolume::GetNavigationData(FNavigationRelevantData& Data) const
 
 		if (NavMeshResolution != ENavigationDataResolution::Invalid)
 		{
-			const FBox& Box = GetBrushComponent()->Brush->Bounds.GetBox();
-			const FAreaNavModifier AreaMod(Box, GetBrushComponent()->GetComponentTransform(), AreaClass);
 			Data.Modifiers.SetNavMeshResolution(NavMeshResolution);
-			Data.Modifiers.Add(AreaMod);
 		}
 	}
 }
