[NavigationSystem] Added cheat command 'ai.debug.nav.DirtyAreaAroundPlayer' that can be used with dynamic navmesh to dirty tiles in a square area around the local player using provided value as extent (in cm) or 10 meters if not specified.

diff --git a/Engine/Source/Runtime/NavigationSystem/Private/NavigationSystem.cpp b/Engine/Source/Runtime/NavigationSystem/Private/NavigationSystem.cpp
--- a/Engine/Source/Runtime/NavigationSystem/Private/NavigationSystem.cpp
+++ b/Engine/Source/Runtime/NavigationSystem/Private/NavigationSystem.cpp
@@ -6,6 +6,7 @@
 #include "Stats/StatsMisc.h"
 #include "Modules/ModuleManager.h"
 #include "AI/Navigation/NavAgentInterface.h"
+#include "Engine/LocalPlayer.h"
 #include "Engine/World.h"
 #include "GameFramework/Controller.h"
 #include "AI/Navigation/NavRelevantInterface.h"
@@ -169,6 +170,39 @@ CSV_DEFINE_CATEGORY(NavInvokers, true);
 //----------------------------------------------------------------------//
 namespace FNavigationSystem
 {
+
+static FAutoConsoleCommandWithWorldArgsAndOutputDevice CmdNavDirtyAreaAroundPlayer(
+	TEXT("ai.debug.nav.DirtyAreaAroundPlayer"),
+	TEXT("Dirty all tiles in a square area around the local player using provided value as extent (in cm), using 10 meters if not specified."),
+	FConsoleCommandWithWorldArgsAndOutputDeviceDelegate::CreateLambda([](const TArray<FString>& Args, const UWorld* World, FOutputDevice& OutputDevice)
+		{
+			if (const ULocalPlayer* LocalPlayer = World->GetFirstLocalPlayerFromController<ULocalPlayer>())
+			{
+				const FVector Center = LocalPlayer->LastViewLocation;
+
+				FVector::FReal Extent = 1000;
+				if (Args.Num() > 0)
+				{
+					if (FCString::IsNumeric(*Args[0]))
+					{
+						Extent = FCString::Atod(*Args[0]);
+					}
+					else
+					{
+						OutputDevice.Log(ELogVerbosity::Error, TEXT("Command failed since first parameter is not a valid numerical value"));
+						return;
+					}
+				}
+
+				UNavigationSystemV1::NavigationDirtyEvent.Broadcast(FBox(Center - FVector(Extent), Center + FVector(Extent)));
+			}
+			else
+			{
+				OutputDevice.Log(ELogVerbosity::Error, TEXT("Command failed since it was unable to find a local player"));
+			}
+		}
+	));
+
 	const FNavDataConfig& GetFallbackNavDataConfig()
 	{
 		static FNavDataConfig FallbackNavDataConfig(FNavigationSystem::FallbackAgentRadius, FNavigationSystem::FallbackAgentHeight);
