[NavigationSystem] fixed UpdateActorAndComponentsInNavOctree to properly update attached actors

diff --git a/Engine/Source/Runtime/NavigationSystem/Private/NavigationSystem.cpp b/Engine/Source/Runtime/NavigationSystem/Private/NavigationSystem.cpp
--- a/Engine/Source/Runtime/NavigationSystem/Private/NavigationSystem.cpp
+++ b/Engine/Source/Runtime/NavigationSystem/Private/NavigationSystem.cpp
@@ -3469,7 +3469,8 @@ void UNavigationSystemV1::UpdateActorAndComponentsInNavOctree(AActor& Actor, boo
 			{
 				for (AActor* AttachedActor : UniqueAttachedActors)
 				{
-					DataHandler.UpdateActorAndComponentsInNavOctree(Actor);
+					checkf(AttachedActor, TEXT("GetAllAttachedActors should only return unique, non-null ptrs."));
+					DataHandler.UpdateActorAndComponentsInNavOctree(*AttachedActor);
 				}
 			}
 		}
