[Navmesh] Optimizations when marking dirty tiles
Extracted finding the seeds and building sub areas to do it once for all NavData
ActiveTiles now use a set instead of an array
When using active tiles, don't try to mark dirtytiles if there is no ActiveTiles
Fixed warning message for oversized dirtyarea to use the AdjustedAreaSize
Addition of profile tags

diff --git a/Engine/Source/Runtime/NavigationSystem/Private/NavMesh/NavMeshRenderingComponent.cpp b/Engine/Source/Runtime/NavigationSystem/Private/NavMesh/NavMeshRenderingComponent.cpp
--- a/Engine/Source/Runtime/NavigationSystem/Private/NavMesh/NavMeshRenderingComponent.cpp
+++ b/Engine/Source/Runtime/NavigationSystem/Private/NavMesh/NavMeshRenderingComponent.cpp
@@ -590,9 +590,9 @@ void FNavMeshSceneProxyData::GatherData(const ARecastNavMesh* NavMesh, int32 InN
 			DebugLabels.Add(FDebugText(FString::Printf(TEXT("Region part %s, Layer part %s"), *GetPartitioningString(NavMesh->RegionPartitioning), *GetPartitioningString(NavMesh->LayerPartitioning))));
 			DebugLabels.Add(FDebugText(TEXT(""))); // empty line
 
-			if (NavMesh->GetGenerator() && NavMesh->GetActiveTiles().Num() != 0)
+			if (NavMesh->GetGenerator() && !NavMesh->GetActiveTileSet().IsEmpty())
 			{
-				DebugLabels.Add(FDebugText(FString::Printf(TEXT("Active tiles: %i"), NavMesh->GetActiveTiles().Num())));	
+				DebugLabels.Add(FDebugText(FString::Printf(TEXT("Active tiles: %i"), NavMesh->GetActiveTileSet().Num())));	
 			}	
 			
 			// Navigation system
diff --git a/Engine/Source/Runtime/NavigationSystem/Private/NavMesh/PImplRecastNavMesh.cpp b/Engine/Source/Runtime/NavigationSystem/Private/NavMesh/PImplRecastNavMesh.cpp
--- a/Engine/Source/Runtime/NavigationSystem/Private/NavMesh/PImplRecastNavMesh.cpp
+++ b/Engine/Source/Runtime/NavigationSystem/Private/NavMesh/PImplRecastNavMesh.cpp
@@ -2725,7 +2725,7 @@ bool FPImplRecastNavMesh::GetDebugGeometryForTile(FRecastDebugGeometry& OutGeome
 		: 0;
 
 	const FRecastNavMeshGenerator* Generator = static_cast<const FRecastNavMeshGenerator*>(NavMeshOwner->GetGenerator());
-	const bool bIsGenerationRestrictedToActiveTiles = Generator && Generator->IsBuildingRestrictedToActiveTiles() && NavMeshOwner->GetActiveTiles().Num() > 0;
+	const bool bIsGenerationRestrictedToActiveTiles = Generator && Generator->IsBuildingRestrictedToActiveTiles() && !NavMeshOwner->GetActiveTileSet().IsEmpty();
 
 	auto ComputeSizeToReserve = [](dtMeshTile const* const Tile, int32& OutNumVertsToReserve, int32& OutNumIndicesToReserve)
 	{
@@ -2772,7 +2772,7 @@ bool FPImplRecastNavMesh::GetDebugGeometryForTile(FRecastDebugGeometry& OutGeome
 	else if (bIsGenerationRestrictedToActiveTiles)
 	{
 		TArray<const dtMeshTile*> Tiles;
-		const TArray<FIntPoint>& ActiveTiles = NavMeshOwner->GetActiveTiles();
+		const TSet<FIntPoint>& ActiveTiles = NavMeshOwner->GetActiveTileSet();
 		for (const FIntPoint& TileLocation : ActiveTiles)
 		{
 			Tiles.Reset();
diff --git a/Engine/Source/Runtime/NavigationSystem/Private/NavMesh/RecastNavMesh.cpp b/Engine/Source/Runtime/NavigationSystem/Private/NavMesh/RecastNavMesh.cpp
--- a/Engine/Source/Runtime/NavigationSystem/Private/NavMesh/RecastNavMesh.cpp
+++ b/Engine/Source/Runtime/NavigationSystem/Private/NavMesh/RecastNavMesh.cpp
@@ -1015,18 +1015,40 @@ void ARecastNavMesh::SortAreasForGenerator(TArray<FRecastAreaNavModifierElement>
 	Modifiers.Sort(FNavAreaSortPredicate());
 }
 
+#if WITH_EDITORONLY_DATA
+// Deprecated
 const TArray<FIntPoint>& ARecastNavMesh::GetActiveTiles() const
 {
 	const FRecastNavMeshGenerator* MyGenerator = static_cast<const FRecastNavMeshGenerator*>(GetGenerator());
 	check(MyGenerator);
+PRAGMA_DISABLE_DEPRECATION_WARNINGS
 	return MyGenerator->ActiveTiles;
+PRAGMA_ENABLE_DEPRECATION_WARNINGS
 }
 
+// Deprecated
 TArray<FIntPoint>& ARecastNavMesh::GetActiveTiles()
 {
 	FRecastNavMeshGenerator* MyGenerator = static_cast<FRecastNavMeshGenerator*>(GetGenerator());
 	check(MyGenerator);
+PRAGMA_DISABLE_DEPRECATION_WARNINGS	
 	return MyGenerator->ActiveTiles;
+PRAGMA_ENABLE_DEPRECATION_WARNINGS
+}
+#endif // WITH_EDITORONLY_DATA
+
+const TSet<FIntPoint>& ARecastNavMesh::GetActiveTileSet() const
+{
+	const FRecastNavMeshGenerator* MyGenerator = static_cast<const FRecastNavMeshGenerator*>(GetGenerator());
+	check(MyGenerator);
+	return MyGenerator->ActiveTileSet;
+}
+
+TSet<FIntPoint>& ARecastNavMesh::GetActiveTileSet()
+{
+	FRecastNavMeshGenerator* MyGenerator = static_cast<FRecastNavMeshGenerator*>(GetGenerator());
+	check(MyGenerator);
+	return MyGenerator->ActiveTileSet;
 }
 
 void ARecastNavMesh::LogRecastTile(const TCHAR* Caller, const FName& Prefix, const FName& OperationName, const dtNavMesh& DetourMesh, const int32 TileX, const int32 TileY, const int32 LayerIndex, const dtTileRef TileRef) const
@@ -3550,10 +3572,10 @@ void ARecastNavMesh::UpdateActiveTiles(const TArray<FNavigationInvokerRaw>& Invo
 	const FVector NavmeshOrigin = Recast2UnrealPoint(NavParams->orig);
 	const FVector::FReal TileDim = Config.GetTileSizeUU();
 
-	TArray<FIntPoint>& ActiveTiles = GetActiveTiles();
-	TArray<FIntPoint> OldActiveSet = ActiveTiles;
+	TSet<FIntPoint>& ActiveTiles = GetActiveTileSet();
+	TSet<FIntPoint> OldActiveSet = ActiveTiles;
 	TArray<FNavMeshDirtyTileElement> TilesInMinDistance;
-	TArray<FIntPoint> TilesInMaxDistance;
+	TSet<FIntPoint> TilesInMaxDistance;
 	TArray<FIntPoint> TileToAppend;
 	TilesInMinDistance.Reserve(ActiveTiles.Num());
 	TilesInMaxDistance.Reserve(ActiveTiles.Num());
@@ -3585,7 +3607,7 @@ void ARecastNavMesh::UpdateActiveTiles(const TArray<FNavigationInvokerRaw>& Invo
 					const FVector::FReal DistanceSq = (InvokerRelativeLocation - FVector(X * TileDim + TileDim / 2, Y * TileDim + TileDim / 2, 0.f)).SizeSquared2D();
 					if (DistanceSq < TileCenterDistanceToRemoveSq)
 					{
-						TilesInMaxDistance.AddUnique(FIntPoint(X, Y));
+						TilesInMaxDistance.FindOrAdd(FIntPoint(X, Y));
 
 						if (DistanceSq < TileCenterDistanceToAddSq)
 						{
@@ -3614,18 +3636,18 @@ void ARecastNavMesh::UpdateActiveTiles(const TArray<FNavigationInvokerRaw>& Invo
 
 	{
 		TRACE_CPUPROFILER_EVENT_SCOPE(ARecastNavMesh::CategorizeTiles);
+
 		TilesToRemove.Reserve(OldActiveSet.Num());
-		for (int32 Index = OldActiveSet.Num() - 1; Index >= 0; --Index)
+		for(TSet<FIntPoint>::TIterator TileIt = OldActiveSet.CreateIterator(); TileIt; ++TileIt)
 		{
-			const FIntPoint& Tile = OldActiveSet[Index];
-			if (TilesInMaxDistance.Find(Tile) == INDEX_NONE)
+			if(!TilesInMaxDistance.Contains(*TileIt))
 			{
-				TilesToRemove.Add(Tile);
-				OldActiveSet.RemoveAtSwap(Index, 1, /*bAllowShrinking=*/false);
+				TilesToRemove.Add(*TileIt);
+				TileIt.RemoveCurrent();
 			}
 			else
 			{
-				ActiveTiles.AddUnique(Tile);
+				ActiveTiles.FindOrAdd(*TileIt);
 			}
 		}
 	}
@@ -3633,11 +3655,10 @@ void ARecastNavMesh::UpdateActiveTiles(const TArray<FNavigationInvokerRaw>& Invo
 	// Find tiles to update
 	TArray<FNavMeshDirtyTileElement> TilesToUpdate;
 	TilesToUpdate.Reserve(ActiveTiles.Num());
-	for (int32 Index = TilesInMinDistance.Num() - 1; Index >= 0; --Index)
+	for (const FNavMeshDirtyTileElement& Tile : TilesInMinDistance)
 	{
 		// Check if it's a new tile (not in the active set)
-		const FNavMeshDirtyTileElement& Tile = TilesInMinDistance[Index];
-		if (OldActiveSet.Find(Tile.Coordinates) == INDEX_NONE)
+		if (!OldActiveSet.Contains(Tile.Coordinates))
 		{
 			TilesToUpdate.Add(Tile);
 		}
diff --git a/Engine/Source/Runtime/NavigationSystem/Private/NavMesh/RecastNavMeshDataChunk.cpp b/Engine/Source/Runtime/NavigationSystem/Private/NavMesh/RecastNavMeshDataChunk.cpp
--- a/Engine/Source/Runtime/NavigationSystem/Private/NavMesh/RecastNavMeshDataChunk.cpp
+++ b/Engine/Source/Runtime/NavigationSystem/Private/NavMesh/RecastNavMeshDataChunk.cpp
@@ -271,10 +271,11 @@ TArray<FNavTileRef> URecastNavMeshDataChunk::AttachTiles(ARecastNavMesh& NavMesh
 
 	if (DetourNavMesh != nullptr)
 	{
-		TArray<FIntPoint>* ActiveTiles = nullptr;
+		TSet<FIntPoint>* ActiveTiles = nullptr;
 		if (UE::NavMesh::Private::IsUsingActiveTileGeneration(NavMesh))
 		{
-			ActiveTiles = &NavMesh.GetActiveTiles();
+			ActiveTiles = &NavMesh.GetActiveTileSet();
+			ActiveTiles->Reserve(ActiveTiles->Num() + Tiles.Num());
 		}
 		
 		for (FRecastTileData& TileData : Tiles)
@@ -335,7 +336,7 @@ TArray<FNavTileRef> URecastNavMeshDataChunk::AttachTiles(ARecastNavMesh& NavMesh
 				
 				if (ActiveTiles)
 				{
-					ActiveTiles->AddUnique(FIntPoint(TileData.X, TileData.Y));					
+					ActiveTiles->FindOrAdd(FIntPoint(TileData.X, TileData.Y));
 				}
 				
 				if (bKeepCopyOfData == false)
@@ -403,10 +404,10 @@ TArray<FNavTileRef> URecastNavMeshDataChunk::DetachTiles(ARecastNavMesh& NavMesh
 
 	if (DetourNavMesh != nullptr)
 	{
-		TArray<FIntPoint>* ActiveTiles = nullptr;
+		TSet<FIntPoint>* ActiveTiles = nullptr;
 		if (UE::NavMesh::Private::IsUsingActiveTileGeneration(NavMesh))
 		{
-			ActiveTiles = &NavMesh.GetActiveTiles();
+			ActiveTiles = &NavMesh.GetActiveTileSet();
 		}
 
 		TArray<const dtMeshTile*> ExtraMeshTiles;
@@ -450,7 +451,7 @@ TArray<FNavTileRef> URecastNavMeshDataChunk::DetachTiles(ARecastNavMesh& NavMesh
 
 					if (ActiveTiles)
 					{
-						ActiveTiles->RemoveSwap(FIntPoint(TileData.X, TileData.Y));	
+						ActiveTiles->Remove(FIntPoint(TileData.X, TileData.Y));
 					}
 						
 					Result.Add(FNavTileRef(TileRef));
diff --git a/Engine/Source/Runtime/NavigationSystem/Private/NavMesh/RecastNavMeshGenerator.cpp b/Engine/Source/Runtime/NavigationSystem/Private/NavMesh/RecastNavMeshGenerator.cpp
--- a/Engine/Source/Runtime/NavigationSystem/Private/NavMesh/RecastNavMeshGenerator.cpp
+++ b/Engine/Source/Runtime/NavigationSystem/Private/NavMesh/RecastNavMeshGenerator.cpp
@@ -4628,12 +4628,14 @@ PRAGMA_ENABLE_DEPRECATION_WARNINGS
 	INC_DWORD_STAT_BY(STAT_NavigationMemory, sizeof(*this));
 }
 
+PRAGMA_DISABLE_DEPRECATION_WARNINGS	// Needed for ActiveTiles
 FRecastNavMeshGenerator::~FRecastNavMeshGenerator()
 {
 	UE_CLOG(RunningDirtyTiles.Num() > 0, LogNavigation, Log, TEXT("Discarding %d build tasks"), RunningDirtyTiles.Num());
 	CancelBuild();
 	DEC_DWORD_STAT_BY( STAT_NavigationMemory, sizeof(*this) );
 }
+PRAGMA_ENABLE_DEPRECATION_WARNINGS
 
 void FRecastNavMeshGenerator::SetupTileConfig(const ENavigationDataResolution TileResolution, FRecastBuildConfig& OutConfig) const
 {
@@ -5268,20 +5270,20 @@ void FRecastNavMeshGenerator::RestrictBuildingToActiveTiles(bool InRestrictBuild
 		bRestrictBuildingToActiveTiles = InRestrictBuildingToActiveTiles;
 		if (InRestrictBuildingToActiveTiles)
 		{
-			// gather non-empty tiles and add them to ActiveTiles
+			// gather non-empty tiles and add them to ActiveTileSet
 
 			const dtNavMesh* DetourMesh = DestNavMesh->GetRecastNavMeshImpl()->GetRecastMesh();
 
 			if (DetourMesh != nullptr && DetourMesh->isEmpty() == false)
 			{
-				ActiveTiles.Reset();
+				ActiveTileSet.Reset();
 				int32 TileCount = DetourMesh->getMaxTiles();
 				for (int32 TileIndex = 0; TileIndex < TileCount; ++TileIndex)
 				{
 					const dtMeshTile* Tile = DetourMesh->getTile(TileIndex);
 					if (Tile != nullptr && Tile->header != nullptr && Tile->header->polyCount > 0)
 					{
-						ActiveTiles.AddUnique(FIntPoint(Tile->header->x, Tile->header->y));
+						ActiveTileSet.FindOrAdd(FIntPoint(Tile->header->x, Tile->header->y));
 					}
 				}
 			}
@@ -5291,8 +5293,7 @@ void FRecastNavMeshGenerator::RestrictBuildingToActiveTiles(bool InRestrictBuild
 
 bool FRecastNavMeshGenerator::IsInActiveSet(const FIntPoint& Tile) const
 {
-	// @TODO checking if given tile is in active tiles needs to be faster
-	return bRestrictBuildingToActiveTiles == false || ActiveTiles.Find(Tile) != INDEX_NONE;
+	return bRestrictBuildingToActiveTiles == false || ActiveTileSet.Contains(Tile);
 }
 
 void FRecastNavMeshGenerator::ResetTimeSlicedTileGeneratorSync()
@@ -6014,12 +6015,6 @@ void FRecastNavMeshGenerator::MarkDirtyTiles(const TArray<FNavigationDirtyArea>&
 	const bool bGameStaticNavMesh = IsGameStaticNavMesh(DestNavMesh);
 
 	const UNavigationSystemV1* NavSys = FNavigationSystem::GetCurrent<UNavigationSystemV1>(GetWorld());
-	const TArray<FBox>* SeedsBoundsArrayPtr = nullptr;
-	if (NavSys && NavSys->IsActiveTilesGenerationEnabled())
-	{
-		SeedsBoundsArrayPtr = &NavSys->GetInvokersSeedBounds();
-	}
-
 	const ARecastNavMesh* const OwnerNav = GetOwner();
 	const bool bUseVirtualGeometryFilteringAndDirtying = OwnerNav != nullptr && OwnerNav->bUseVirtualGeometryFilteringAndDirtying;
 	// Those are set only if bUseVirtualGeometryFilteringAndDirtying is enabled since we do not use them for anything else
@@ -6035,188 +6030,172 @@ void FRecastNavMeshGenerator::MarkDirtyTiles(const TArray<FNavigationDirtyArea>&
 	// Used for debug purposes to track the number of new dirty tiles per area; Updated only if LogNavigationDirtyArea is VeryVerbose
 	TMap<FPendingTileElement, TArray<FNavigationDirtyAreaPerTileDebugInformation>> DirtyAreasDebugging;
 #endif
-	for (const FNavigationDirtyArea& DirtyArea : DirtyAreas)
+
+	if (!bRestrictBuildingToActiveTiles || !ActiveTileSet.IsEmpty())
 	{
-		if (!ensureMsgf(DirtyArea.Bounds.IsValid, TEXT("%hs Attempting to use DirtyArea.Bounds which are not valid. SourceObject: %s"), __FUNCTION__, *GetFullNameSafe(DirtyArea.OptionalSourceObject.Get())))
+		for (const FNavigationDirtyArea& DirtyArea : DirtyAreas)
 		{
-			continue;
-		}
+			QUICK_SCOPE_CYCLE_COUNTER(STAT_RecastNavMeshGenerator_DirtyArea);
+		
+			if (!ensureMsgf(DirtyArea.Bounds.IsValid, TEXT("%hs Attempting to use DirtyArea.Bounds which are not valid. SourceObject: %s"), __FUNCTION__, *GetFullNameSafe(DirtyArea.OptionalSourceObject.Get())))
+			{
+				continue;
+			}
 
-		// Game world static navmeshes accept only area modifiers updates
-		if (bGameStaticNavMesh && (!DirtyArea.HasFlag(ENavigationDirtyFlag::DynamicModifier) || DirtyArea.HasFlag(ENavigationDirtyFlag::NavigationBounds)))
-		{
-			continue;
-		}
+			// Game world static navmeshes accept only area modifiers updates
+			if (bGameStaticNavMesh && (!DirtyArea.HasFlag(ENavigationDirtyFlag::DynamicModifier) || DirtyArea.HasFlag(ENavigationDirtyFlag::NavigationBounds)))
+			{
+				continue;
+			}
 
-		UE_VLOG_BOX(OwnerNav, LogNavigation, VeryVerbose, DirtyArea.Bounds, FColor::Blue, TEXT("DirtyArea %s"), *GetNameSafe(DirtyArea.OptionalSourceObject.Get()));
+			UE_VLOG_BOX(OwnerNav, LogNavigation, VeryVerbose, DirtyArea.Bounds, FColor::Blue, TEXT("DirtyArea %s"), *GetNameSafe(DirtyArea.OptionalSourceObject.Get()));
 		
-		// (if bUseVirtualGeometryFilteringAndDirtying is true) Ignore dirty areas flagged by a source object that is not supposed to apply to this navmesh
-		if (bUseVirtualGeometryFilteringAndDirtying && NavSys && NavOctreeInstance && NavDataConfig)
-		{
-			if (const UObject* const SourceObject = DirtyArea.OptionalSourceObject.Get())
+			// (if bUseVirtualGeometryFilteringAndDirtying is true) Ignore dirty areas flagged by a source object that is not supposed to apply to this navmesh
+			if (bUseVirtualGeometryFilteringAndDirtying && NavSys && NavOctreeInstance && NavDataConfig)
 			{
-				if (!ShouldDirtyTilesRequestedByObject(*NavSys, *NavOctreeInstance, *SourceObject, *NavDataConfig))
+				if (const UObject* const SourceObject = DirtyArea.OptionalSourceObject.Get())
 				{
-					continue;
+					if (!ShouldDirtyTilesRequestedByObject(*NavSys, *NavOctreeInstance, *SourceObject, *NavDataConfig))
+					{
+						continue;
+					}
 				}
 			}
-		}
 		
-		bool bDoTileInclusionTest = false;
-		FBox AdjustedAreaBounds = DirtyArea.Bounds;
+			bool bDoTileInclusionTest = false;
+			FBox AdjustedAreaBounds = DirtyArea.Bounds;
 		
-		// if it's not expanding the navigable area
-		if (DirtyArea.HasFlag(ENavigationDirtyFlag::NavigationBounds) == false)
-		{
-			// and is outside of current bounds
-			if (GetTotalBounds().Intersect(DirtyArea.Bounds) == false)
-			{
-				// skip it
-				continue;
-			}
-
-			const FBox CutDownArea = CalculateBoxIntersection(GetTotalBounds(), DirtyArea.Bounds);
-			AdjustedAreaBounds = GrowBoundingBox(CutDownArea, DirtyArea.HasFlag(ENavigationDirtyFlag::UseAgentHeight));
-
-			// @TODO this and the following test share some work in common
-			if (IntersectBounds(AdjustedAreaBounds, InclusionBounds) == false)
+			// if it's not expanding the navigable area
+			if (DirtyArea.HasFlag(ENavigationDirtyFlag::NavigationBounds) == false)
 			{
-				continue;
-			}
-
-			// check if any of inclusion volumes encapsulates this box
-			// using CutDownArea not AdjustedAreaBounds since if the area is on the border of navigable space
-			// then FindInclusionBoundEncapsulatingBox can produce false negative
-			bDoTileInclusionTest = (FindInclusionBoundEncapsulatingBox(CutDownArea) == INDEX_NONE);
-		}
+				// and is outside of current bounds
+				if (GetTotalBounds().Intersect(DirtyArea.Bounds) == false)
+				{
+					// skip it
+					continue;
+				}
 
-		QUICK_SCOPE_CYCLE_COUNTER(STAT_RecastNavMeshGenerator_DirtyArea);
+				const FBox CutDownArea = CalculateBoxIntersection(GetTotalBounds(), DirtyArea.Bounds);
+				AdjustedAreaBounds = GrowBoundingBox(CutDownArea, DirtyArea.HasFlag(ENavigationDirtyFlag::UseAgentHeight));
 
-		TArray<FBox, TInlineAllocator<32>> SubAreaBoundsArray;
-		if (SeedsBoundsArrayPtr != nullptr && SeedsBoundsArrayPtr->Num() > 0)
-		{
-			for (const FBox& SeedBounds : *SeedsBoundsArrayPtr)
-			{
-				// Compute sub area bound
-				const FBox OverlapBox = AdjustedAreaBounds.Overlap(SeedBounds);
-				if (OverlapBox.IsValid)
+				// @TODO this and the following test share some work in common
+				if (IntersectBounds(AdjustedAreaBounds, InclusionBounds) == false)
 				{
-					SubAreaBoundsArray.Add(OverlapBox);
+					continue;
 				}
+
+				// check if any of inclusion volumes encapsulates this box
+				// using CutDownArea not AdjustedAreaBounds since if the area is on the border of navigable space
+				// then FindInclusionBoundEncapsulatingBox can produce false negative
+				bDoTileInclusionTest = (FindInclusionBoundEncapsulatingBox(CutDownArea) == INDEX_NONE);
 			}
-		}
-		else
-		{
-			if (ensureMsgf(AdjustedAreaBounds.IsValid, TEXT("%hs Attempting to use AdjustedAreaBounds which are not valid"), __FUNCTION__))
+
 			{
-				SubAreaBoundsArray.Add(AdjustedAreaBounds);
-			}
-		}
+				QUICK_SCOPE_CYCLE_COUNTER(STAT_RecastNavMeshGenerator_CheckTilesInBounds);
 
-		uint32 PendingTilesMarked = 0;
-		
-		for (const FBox& SubArea : SubAreaBoundsArray)
-		{
-			const FRcTileBox TileBox(SubArea, RcNavMeshOrigin, TileSizeInWorldUnits);
+				uint32 PendingTilesMarked = 0;
+			
+				const FRcTileBox TileBox(AdjustedAreaBounds, RcNavMeshOrigin, TileSizeInWorldUnits);
 
-			for (int32 TileY = TileBox.YMin; TileY <= TileBox.YMax; ++TileY)
-			{
-				for (int32 TileX = TileBox.XMin; TileX <= TileBox.XMax; ++TileX)
+				for (int32 TileY = TileBox.YMin; TileY <= TileBox.YMax; ++TileY)
 				{
-					if (IsInActiveSet(FIntPoint(TileX, TileY)) == false)
+					for (int32 TileX = TileBox.XMin; TileX <= TileBox.XMax; ++TileX)
 					{
-						UE_SUPPRESS(LogNavigation, VeryVerbose,
+						if (IsInActiveSet(FIntPoint(TileX, TileY)) == false)
 						{
-							const FBox TileBounds = FRecastTileGenerator::CalculateTileBounds(TileX, TileY, RcNavMeshOrigin, TotalNavBounds, TileSizeInWorldUnits);
-							UE_VLOG_BOX(OwnerNav, LogNavigation, VeryVerbose, TileBounds, FColor::Red, TEXT("Not in active set"));
-						});
-						continue;
-					}
-
-					if (bDoTileInclusionTest == true && DirtyArea.HasFlag(ENavigationDirtyFlag::NavigationBounds) == false)
-					{
-						const FBox TileBounds = FRecastTileGenerator::CalculateTileBounds(TileX, TileY, RcNavMeshOrigin, TotalNavBounds, TileSizeInWorldUnits);
+							UE_SUPPRESS(LogNavigation, VeryVerbose,
+							{
+								const FBox TileBounds = FRecastTileGenerator::CalculateTileBounds(TileX, TileY, RcNavMeshOrigin, TotalNavBounds, TileSizeInWorldUnits);
+								UE_VLOG_BOX(OwnerNav, LogNavigation, VeryVerbose, TileBounds, FColor::Red, TEXT("Not in active set"));
+							});
+							continue;
+						}
 
-						// do per tile check since we can have lots of tiles in between navigable bounds volumes
-						if (IntersectBounds(TileBounds, InclusionBounds) == false)
+						if (bDoTileInclusionTest == true && DirtyArea.HasFlag(ENavigationDirtyFlag::NavigationBounds) == false)
 						{
-							// Skip this tile
-							continue;
+							const FBox TileBounds = FRecastTileGenerator::CalculateTileBounds(TileX, TileY, RcNavMeshOrigin, TotalNavBounds, TileSizeInWorldUnits);
+
+							// do per tile check since we can have lots of tiles in between navigable bounds volumes
+							if (IntersectBounds(TileBounds, InclusionBounds) == false)
+							{
+								// Skip this tile
+								continue;
+							}
 						}
-					}
 											
-					FPendingTileElement Element;
-					Element.Coord = FIntPoint(TileX, TileY);
-					// Make sure to prevent bRebuildGeometry for game world static navmeshes.
-					// Game world static navmeshes accept only area modifiers updates. Rebuilding geometry would bRegenerateCompressedLayers without having the geometry for them.
-					Element.bRebuildGeometry = !bGameStaticNavMesh && (DirtyArea.HasFlag(ENavigationDirtyFlag::Geometry) || DirtyArea.HasFlag(ENavigationDirtyFlag::NavigationBounds));
-					Element.CreationTime = CurrentTimeSeconds;
-					if (Element.bRebuildGeometry == false)
-					{
-						Element.DirtyAreas.Add(SubArea);
-					}
-					PendingTilesMarked++;
+						FPendingTileElement Element;
+						Element.Coord = FIntPoint(TileX, TileY);
+						// Make sure to prevent bRebuildGeometry for game world static navmeshes.
+						// Game world static navmeshes accept only area modifiers updates. Rebuilding geometry would bRegenerateCompressedLayers without having the geometry for them.
+						Element.bRebuildGeometry = !bGameStaticNavMesh && (DirtyArea.HasFlag(ENavigationDirtyFlag::Geometry) || DirtyArea.HasFlag(ENavigationDirtyFlag::NavigationBounds));
+						Element.CreationTime = CurrentTimeSeconds;
+						if (Element.bRebuildGeometry == false)
+						{
+							Element.DirtyAreas.Add(AdjustedAreaBounds);
+						}
+						PendingTilesMarked++;
 			
-					FPendingTileElement* ExistingElement = DirtyTiles.Find(Element);
-					if (ExistingElement)
-					{
-						ExistingElement->bRebuildGeometry |= Element.bRebuildGeometry;
-						// Append area bounds to existing list 
-						if (ExistingElement->bRebuildGeometry == false)
+						FPendingTileElement* ExistingElement = DirtyTiles.Find(Element);
+						if (ExistingElement)
 						{
-							ExistingElement->DirtyAreas.Append(Element.DirtyAreas);
+							ExistingElement->bRebuildGeometry |= Element.bRebuildGeometry;
+							// Append area bounds to existing list 
+							if (ExistingElement->bRebuildGeometry == false)
+							{
+								ExistingElement->DirtyAreas.Append(Element.DirtyAreas);
+							}
+							else
+							{
+								ExistingElement->DirtyAreas.Empty();
+							}
 						}
 						else
 						{
-							ExistingElement->DirtyAreas.Empty();
+							DirtyTiles.Add(Element);
 						}
-					}
-					else
-					{
-						DirtyTiles.Add(Element);
-					}
 			
 #if !UE_BUILD_SHIPPING
-					UE_SUPPRESS(LogNavigationDirtyArea, VeryVerbose, 
-					{
-						const bool bAlreadyAdded = ExistingElement != nullptr;
-						DirtyAreasDebugging.FindOrAdd(Element).Add({DirtyArea, bAlreadyAdded});
-					});
+						UE_SUPPRESS(LogNavigationDirtyArea, VeryVerbose, 
+						{
+							const bool bAlreadyAdded = ExistingElement != nullptr;
+							DirtyAreasDebugging.FindOrAdd(Element).Add({DirtyArea, bAlreadyAdded});
+						});
 #endif
+					}
 				}
-			}
 
 #if !UE_BUILD_SHIPPING
-			// Warn if this is from a big dirty area
-			UE_SUPPRESS(LogNavigationDirtyArea, Warning, 
-			{
-				if (PendingTilesMarked > 0)
+				// Warn if this is from a big dirty area
+				UE_SUPPRESS(LogNavigationDirtyArea, Warning, 
 				{
-					if (NavSys == nullptr)
+					if (PendingTilesMarked > 0)
 					{
-						// NavSys might not have been initialized yet if not using bUseVirtualGeometryFilteringAndDirtying
-						NavSys = FNavigationSystem::GetCurrent<UNavigationSystemV1>(GetWorld());
-					}
+						if (NavSys == nullptr)
+						{
+							// NavSys might not have been initialized yet if not using bUseVirtualGeometryFilteringAndDirtying
+							NavSys = FNavigationSystem::GetCurrent<UNavigationSystemV1>(GetWorld());
+						}
 			
-					// If not using active tile generation, those are reported earlier in FNavigationDirtyAreasController::AddArea
-					if (NavSys && NavSys->GetOperationMode() == FNavigationSystemRunMode::GameMode && NavSys->IsActiveTilesGenerationEnabled() &&
-						AdjustedAreaBounds.GetSize().GetMax() > NavSys->GetDirtyAreaWarningSizeThreshold())
-					{
-						const UObject* const SourceObject = DirtyArea.OptionalSourceObject.Get();
-						const UActorComponent* const ObjectAsComponent = Cast<UActorComponent>(SourceObject);
-						const AActor* const ComponentOwner = ObjectAsComponent ? ObjectAsComponent->GetOwner() : nullptr;
-						const FVector2D BoundsSize(DirtyArea.Bounds.GetSize());
+						// If not using active tile generation, those are reported earlier in FNavigationDirtyAreasController::AddArea
+						if (NavSys && NavSys->GetOperationMode() == FNavigationSystemRunMode::GameMode && NavSys->IsActiveTilesGenerationEnabled() &&
+							AdjustedAreaBounds.GetSize().GetMax() > NavSys->GetDirtyAreaWarningSizeThreshold())
+						{
+							const UObject* const SourceObject = DirtyArea.OptionalSourceObject.Get();
+							const UActorComponent* const ObjectAsComponent = Cast<UActorComponent>(SourceObject);
+							const AActor* const ComponentOwner = ObjectAsComponent ? ObjectAsComponent->GetOwner() : nullptr;
+							const FVector2D AdjustedAreaBoundsSize(AdjustedAreaBounds.GetSize());
 		
-						UE_LOG(LogNavigationDirtyArea, Warning,
-							TEXT("(navmesh: %-30s) Added an oversized dirty area | Tiles marked: %2u | Source object = %s | Potential comp owner = %s | Bounds size = %s | Threshold: %.0f"),
-							*GetNameSafe(GetOwner()), PendingTilesMarked, *GetFullNameSafe(SourceObject), *GetFullNameSafe(ComponentOwner),
-							*BoundsSize.ToString(), NavSys->GetDirtyAreaWarningSizeThreshold());
+							UE_LOG(LogNavigationDirtyArea, Warning,
+								TEXT("(navmesh: %-30s) Added an oversized dirty area | Tiles marked: %2u | Source object = %s | Potential comp owner = %s | Bounds size = %s | Threshold: %.0f"),
+								*GetNameSafe(GetOwner()), PendingTilesMarked, *GetFullNameSafe(SourceObject), *GetFullNameSafe(ComponentOwner),
+								*AdjustedAreaBoundsSize.ToString(), NavSys->GetDirtyAreaWarningSizeThreshold());
+						}
 					}
-				}
-			});
+				});
 #endif // !UE_BUILD_SHIPPING
-			
-		} // SubAreaBoundsArray loop	
+
+			} // QUICK_SCOPE_CYCLE_COUNTER(STAT_RecastNavMeshGenerator_DirtyArea);
+		}
 	}
 	
 	int32 NumTilesMarked = DirtyTiles.Num();
diff --git a/Engine/Source/Runtime/NavigationSystem/Private/NavigationDirtyAreasController.cpp b/Engine/Source/Runtime/NavigationSystem/Private/NavigationDirtyAreasController.cpp
--- a/Engine/Source/Runtime/NavigationSystem/Private/NavigationDirtyAreasController.cpp
+++ b/Engine/Source/Runtime/NavigationSystem/Private/NavigationDirtyAreasController.cpp
@@ -2,6 +2,7 @@
 
 #include "NavigationDirtyAreasController.h"
 #include "NavigationData.h"
+#include "NavigationSystem.h"
 #include "VisualLogger/VisualLogger.h"
 #include "AI/Navigation/NavigationDirtyElement.h"
 
@@ -28,18 +29,79 @@ void FNavigationDirtyAreasController::ForceRebuildOnNextTick()
 	DirtyAreasUpdateTime = FMath::Max(DirtyAreasUpdateTime, MinTimeForUpdate);
 }
 
-void FNavigationDirtyAreasController::Tick(const float DeltaSeconds, const TArray<ANavigationData*>& NavDataSet, bool bForceRebuilding)
+namespace UE::Navigation::Private
+{
+	const UNavigationSystemV1* FindNavigationSystem(const TArray<ANavigationData*>& NavDataSet)
+	{
+		const UNavigationSystemV1* NavSys = nullptr;
+		for (const ANavigationData* NavData : NavDataSet)
+		{
+			if (NavData)
+			{
+				NavSys = FNavigationSystem::GetCurrent<UNavigationSystemV1>(NavData->GetWorld());
+				if (NavSys)
+				{
+					return NavSys;
+				}
+			}
+		}
+
+		return NavSys;
+	}
+}
+
+void FNavigationDirtyAreasController::Tick(const float DeltaSeconds, const TArray<ANavigationData*>& NavDataSet, bool bForceRebuilding /*= false*/)
 {
 	DirtyAreasUpdateTime += DeltaSeconds;
 	const bool bCanRebuildNow = bForceRebuilding || (DirtyAreasUpdateFreq != 0.f && DirtyAreasUpdateTime >= (1.0f / DirtyAreasUpdateFreq));
 
 	if (DirtyAreas.Num() > 0 && bCanRebuildNow)
 	{
+		bool bIsUsingActiveTileGeneration = false;
+		TArray<FNavigationDirtyArea> SubAreaArray;
+		SubAreaArray.Reserve(DirtyAreas.Num());
+		
+		{
+			QUICK_SCOPE_CYCLE_COUNTER(STAT_RecastNavMeshGenerator_MakingSubAreas);
+
+			// Find the relevant navigation system
+			const UNavigationSystemV1* NavSys = UE::Navigation::Private::FindNavigationSystem(NavDataSet);
+			
+			const TArray<FBox>* SeedsBoundsArrayPtr = nullptr;
+			bIsUsingActiveTileGeneration = NavSys && NavSys->IsActiveTilesGenerationEnabled(); 
+			if (bIsUsingActiveTileGeneration)
+			{
+				SeedsBoundsArrayPtr = &NavSys->GetInvokersSeedBounds();
+			}
+
+			if (SeedsBoundsArrayPtr != nullptr && SeedsBoundsArrayPtr->Num() > 0)
+			{
+				for (const FNavigationDirtyArea& DirtyArea : DirtyAreas)
+				{
+					const FBox& AreaBound = DirtyArea.Bounds;
+					if (!ensureMsgf(AreaBound.IsValid, TEXT("%hs Attempting to use DirtyArea.Bounds which are not valid. SourceObject: %s"), __FUNCTION__, *GetFullNameSafe(DirtyArea.OptionalSourceObject.Get())))
+					{
+						continue;
+					}
+					
+					for (const FBox& SeedBounds : *SeedsBoundsArrayPtr)
+					{
+						// Compute sub area bound
+						const FBox OverlapBox = AreaBound.Overlap(SeedBounds);
+						if (OverlapBox.IsValid)
+						{
+							SubAreaArray.Emplace(OverlapBox, DirtyArea.Flags, DirtyArea.OptionalSourceObject.Get());
+						}
+					}
+				}
+			}
+		}
+		
 		for (ANavigationData* NavData : NavDataSet)
 		{
 			if (NavData)
 			{
-				NavData->RebuildDirtyAreas(DirtyAreas);
+				NavData->RebuildDirtyAreas(bIsUsingActiveTileGeneration ? SubAreaArray : DirtyAreas);
 			}
 		}
 
diff --git a/Engine/Source/Runtime/NavigationSystem/Public/NavMesh/RecastNavMesh.h b/Engine/Source/Runtime/NavigationSystem/Public/NavMesh/RecastNavMesh.h
--- a/Engine/Source/Runtime/NavigationSystem/Public/NavMesh/RecastNavMesh.h
+++ b/Engine/Source/Runtime/NavigationSystem/Public/NavMesh/RecastNavMesh.h
@@ -1105,9 +1105,14 @@ public:
 	NAVIGATIONSYSTEM_API virtual void AttachNavMeshDataChunk(URecastNavMeshDataChunk& NavDataChunk);
 	NAVIGATIONSYSTEM_API virtual void DetachNavMeshDataChunk(URecastNavMeshDataChunk& NavDataChunk);
 
+	UE_DEPRECATED(5.4, "Use GetActiveTileSet instead.")
 	NAVIGATIONSYSTEM_API const TArray<FIntPoint>& GetActiveTiles() const;
+	UE_DEPRECATED(5.4, "Use GetActiveTileSet instead.")
 	NAVIGATIONSYSTEM_API TArray<FIntPoint>& GetActiveTiles(); 
 
+	NAVIGATIONSYSTEM_API const TSet<FIntPoint>& GetActiveTileSet() const;
+	NAVIGATIONSYSTEM_API TSet<FIntPoint>& GetActiveTileSet(); 
+	
 	NAVIGATIONSYSTEM_API void LogRecastTile(const TCHAR* Caller, const FName& Prefix, const FName& OperationName, const dtNavMesh& DetourMesh, const int32 TileX, const int32 TileY, const int32 LayerIndex, const uint64 TileRef) const;
 	
 protected:
diff --git a/Engine/Source/Runtime/NavigationSystem/Public/NavMesh/RecastNavMeshGenerator.h b/Engine/Source/Runtime/NavigationSystem/Public/NavMesh/RecastNavMeshGenerator.h
--- a/Engine/Source/Runtime/NavigationSystem/Public/NavMesh/RecastNavMeshGenerator.h
+++ b/Engine/Source/Runtime/NavigationSystem/Public/NavMesh/RecastNavMeshGenerator.h
@@ -1063,8 +1063,13 @@ protected:
 	/** List of tiles that were recently regenerated */
 	TNavStatArray<FTileTimestamp> RecentlyBuiltTiles;
 #endif// WITH_EDITOR
-	
+
+#if WITH_EDITORONLY_DATA	
+	UE_DEPRECATED(5.4, "Use ActiveTileSet instead.")
 	TArray<FIntPoint> ActiveTiles;
+#endif // WITH_EDITORONLY_DATA
+	
+	TSet<FIntPoint> ActiveTileSet;
 
 	/** */
 	FRecastNavMeshCachedData AdditionalCachedData;
