Still fills the SubAreaArray with the DirtyArray in the case we have no seeds bound

#tests pie

diff --git a/Engine/Source/Runtime/NavigationSystem/Private/NavigationDirtyAreasController.cpp b/Engine/Source/Runtime/NavigationSystem/Private/NavigationDirtyAreasController.cpp
--- a/Engine/Source/Runtime/NavigationSystem/Private/NavigationDirtyAreasController.cpp
+++ b/Engine/Source/Runtime/NavigationSystem/Private/NavigationDirtyAreasController.cpp
@@ -74,16 +74,16 @@ void FNavigationDirtyAreasController::Tick(const float DeltaSeconds, const TArra
 				SeedsBoundsArrayPtr = &NavSys->GetInvokersSeedBounds();
 			}
 
-			if (SeedsBoundsArrayPtr != nullptr && SeedsBoundsArrayPtr->Num() > 0)
+			for (const FNavigationDirtyArea& DirtyArea : DirtyAreas)
 			{
-				for (const FNavigationDirtyArea& DirtyArea : DirtyAreas)
+				const FBox& AreaBound = DirtyArea.Bounds;
+				if (!ensureMsgf(AreaBound.IsValid, TEXT("%hs Attempting to use DirtyArea.Bounds which are not valid. SourceObject: %s"), __FUNCTION__, *GetFullNameSafe(DirtyArea.OptionalSourceObject.Get())))
+				{
+					continue;
+				}
+
+				if (SeedsBoundsArrayPtr != nullptr && SeedsBoundsArrayPtr->Num() > 0)
 				{
-					const FBox& AreaBound = DirtyArea.Bounds;
-					if (!ensureMsgf(AreaBound.IsValid, TEXT("%hs Attempting to use DirtyArea.Bounds which are not valid. SourceObject: %s"), __FUNCTION__, *GetFullNameSafe(DirtyArea.OptionalSourceObject.Get())))
-					{
-						continue;
-					}
-					
 					for (const FBox& SeedBounds : *SeedsBoundsArrayPtr)
 					{
 						// Compute sub area bound
@@ -94,6 +94,10 @@ void FNavigationDirtyAreasController::Tick(const float DeltaSeconds, const TArra
 						}
 					}
 				}
+				else
+				{
+					SubAreaArray.Emplace(DirtyArea);
+				}
 			}
 		}
 		
