#UE Lower verbosity of a message about a full queue to a warning and add more info to it so it is actionable

diff --git a/Engine/Source/Runtime/NavigationSystem/Private/NavigationSystem.cpp b/Engine/Source/Runtime/NavigationSystem/Private/NavigationSystem.cpp
--- a/Engine/Source/Runtime/NavigationSystem/Private/NavigationSystem.cpp
+++ b/Engine/Source/Runtime/NavigationSystem/Private/NavigationSystem.cpp
@@ -2535,7 +2535,7 @@ void UNavigationSystemV1::RequestRegistrationDeferred(ANavigationData& NavData)
 	}
 	else
 	{
-		UE_LOG(LogNavigation, Error, TEXT("Navigation System: registration queue full!"));
+		UE_LOG(LogNavigation, Warning, TEXT("Navigation System: registration queue full! System:%s NavData:%s"), *GetPathNameSafe(this), *GetPathNameSafe(&NavData));
 	}
 }
 
