Remove usage of CoreMinimal.h from navigation modules
Removed few other includes as well

diff --git a/Engine/Source/Runtime/Engine/Classes/AI/Navigation/AvoidanceManager.h b/Engine/Source/Runtime/Engine/Classes/AI/Navigation/AvoidanceManager.h
--- a/Engine/Source/Runtime/Engine/Classes/AI/Navigation/AvoidanceManager.h
+++ b/Engine/Source/Runtime/Engine/Classes/AI/Navigation/AvoidanceManager.h
@@ -2,7 +2,9 @@
 
 #pragma once
 
+#if UE_ENABLE_INCLUDE_ORDER_DEPRECATED_IN_5_4
 #include "CoreMinimal.h"
+#endif
 #include "Stats/Stats.h"
 #include "UObject/ObjectMacros.h"
 #include "UObject/Object.h"
diff --git a/Engine/Source/Runtime/Engine/Classes/AI/Navigation/NavCollisionBase.h b/Engine/Source/Runtime/Engine/Classes/AI/Navigation/NavCollisionBase.h
--- a/Engine/Source/Runtime/Engine/Classes/AI/Navigation/NavCollisionBase.h
+++ b/Engine/Source/Runtime/Engine/Classes/AI/Navigation/NavCollisionBase.h
@@ -2,7 +2,9 @@
 
 #pragma once
 
+#if UE_ENABLE_INCLUDE_ORDER_DEPRECATED_IN_5_4
 #include "CoreMinimal.h"
+#endif
 #include "AI/NavigationSystemHelpers.h"
 #include "AI/Navigation/NavigationTypes.h"
 #include "NavCollisionBase.generated.h"
diff --git a/Engine/Source/Runtime/Engine/Classes/AI/Navigation/NavEdgeProviderInterface.h b/Engine/Source/Runtime/Engine/Classes/AI/Navigation/NavEdgeProviderInterface.h
--- a/Engine/Source/Runtime/Engine/Classes/AI/Navigation/NavEdgeProviderInterface.h
+++ b/Engine/Source/Runtime/Engine/Classes/AI/Navigation/NavEdgeProviderInterface.h
@@ -2,7 +2,9 @@
 
 #pragma once
 
+#if UE_ENABLE_INCLUDE_ORDER_DEPRECATED_IN_5_4
 #include "CoreMinimal.h"
+#endif
 #include "UObject/ObjectMacros.h"
 #include "UObject/Interface.h"
 #include "NavEdgeProviderInterface.generated.h"
diff --git a/Engine/Source/Runtime/Engine/Classes/AI/Navigation/NavLinkDefinition.h b/Engine/Source/Runtime/Engine/Classes/AI/Navigation/NavLinkDefinition.h
--- a/Engine/Source/Runtime/Engine/Classes/AI/Navigation/NavLinkDefinition.h
+++ b/Engine/Source/Runtime/Engine/Classes/AI/Navigation/NavLinkDefinition.h
@@ -2,7 +2,9 @@
 
 #pragma once
 
+#if UE_ENABLE_INCLUDE_ORDER_DEPRECATED_IN_5_4
 #include "CoreMinimal.h"
+#endif
 #include "UObject/ObjectMacros.h"
 #include "UObject/Object.h"
 #include "UObject/Class.h"
diff --git a/Engine/Source/Runtime/Engine/Classes/AI/Navigation/NavPathObserverInterface.h b/Engine/Source/Runtime/Engine/Classes/AI/Navigation/NavPathObserverInterface.h
--- a/Engine/Source/Runtime/Engine/Classes/AI/Navigation/NavPathObserverInterface.h
+++ b/Engine/Source/Runtime/Engine/Classes/AI/Navigation/NavPathObserverInterface.h
@@ -2,7 +2,9 @@
 
 #pragma once
 
+#if UE_ENABLE_INCLUDE_ORDER_DEPRECATED_IN_5_4
 #include "CoreMinimal.h"
+#endif
 #include "UObject/ObjectMacros.h"
 #include "UObject/Interface.h"
 #include "NavPathObserverInterface.generated.h"
diff --git a/Engine/Source/Runtime/Engine/Classes/AI/Navigation/NavigationAvoidanceTypes.h b/Engine/Source/Runtime/Engine/Classes/AI/Navigation/NavigationAvoidanceTypes.h
--- a/Engine/Source/Runtime/Engine/Classes/AI/Navigation/NavigationAvoidanceTypes.h
+++ b/Engine/Source/Runtime/Engine/Classes/AI/Navigation/NavigationAvoidanceTypes.h
@@ -2,7 +2,9 @@
 
 #pragma once
 
+#if UE_ENABLE_INCLUDE_ORDER_DEPRECATED_IN_5_4
 #include "CoreMinimal.h"
+#endif
 #include "UObject/ObjectMacros.h"
 #include "NavigationAvoidanceTypes.generated.h"
 
diff --git a/Engine/Source/Runtime/Engine/Classes/AI/Navigation/NavigationDataChunk.h b/Engine/Source/Runtime/Engine/Classes/AI/Navigation/NavigationDataChunk.h
--- a/Engine/Source/Runtime/Engine/Classes/AI/Navigation/NavigationDataChunk.h
+++ b/Engine/Source/Runtime/Engine/Classes/AI/Navigation/NavigationDataChunk.h
@@ -2,7 +2,9 @@
 
 #pragma once
 
+#if UE_ENABLE_INCLUDE_ORDER_DEPRECATED_IN_5_4
 #include "CoreMinimal.h"
+#endif
 #include "UObject/ObjectMacros.h"
 #include "UObject/Object.h"
 #include "NavigationDataChunk.generated.h"
diff --git a/Engine/Source/Runtime/Engine/Classes/AI/Navigation/NavigationDataInterface.h b/Engine/Source/Runtime/Engine/Classes/AI/Navigation/NavigationDataInterface.h
--- a/Engine/Source/Runtime/Engine/Classes/AI/Navigation/NavigationDataInterface.h
+++ b/Engine/Source/Runtime/Engine/Classes/AI/Navigation/NavigationDataInterface.h
@@ -2,7 +2,9 @@
 
 #pragma once
 
+#if UE_ENABLE_INCLUDE_ORDER_DEPRECATED_IN_5_4
 #include "CoreMinimal.h"
+#endif
 #include "UObject/ObjectMacros.h"
 #include "UObject/Interface.h"
 #include "AI/Navigation/NavigationTypes.h"
diff --git a/Engine/Source/Runtime/Engine/Classes/AI/Navigation/PathFollowingAgentInterface.h b/Engine/Source/Runtime/Engine/Classes/AI/Navigation/PathFollowingAgentInterface.h
--- a/Engine/Source/Runtime/Engine/Classes/AI/Navigation/PathFollowingAgentInterface.h
+++ b/Engine/Source/Runtime/Engine/Classes/AI/Navigation/PathFollowingAgentInterface.h
@@ -2,7 +2,9 @@
 
 #pragma once
 
+#if UE_ENABLE_INCLUDE_ORDER_DEPRECATED_IN_5_4
 #include "CoreMinimal.h"
+#endif
 #include "UObject/ObjectMacros.h"
 #include "UObject/Interface.h"
 #include "PathFollowingAgentInterface.generated.h"
diff --git a/Engine/Source/Runtime/NavigationSystem/Public/AbstractNavData.h b/Engine/Source/Runtime/NavigationSystem/Public/AbstractNavData.h
--- a/Engine/Source/Runtime/NavigationSystem/Public/AbstractNavData.h
+++ b/Engine/Source/Runtime/NavigationSystem/Public/AbstractNavData.h
@@ -3,12 +3,18 @@
 #pragma once
 
 #include "AI/Navigation/NavQueryFilter.h"
+#if UE_ENABLE_INCLUDE_ORDER_DEPRECATED_IN_5_4
 #include "AI/Navigation/NavigationTypes.h"
+#endif
 #include "Containers/Array.h"
 #include "Containers/ContainersFwd.h"
+#if UE_ENABLE_INCLUDE_ORDER_DEPRECATED_IN_5_4
 #include "CoreMinimal.h"
+#endif
 #include "Math/Box.h"
+#if UE_ENABLE_INCLUDE_ORDER_DEPRECATED_IN_5_4
 #include "Math/UnrealMathSSE.h"
+#endif
 #include "Math/Vector.h"
 #if UE_ENABLE_INCLUDE_ORDER_DEPRECATED_IN_5_2
 #include "NavFilters/NavigationQueryFilter.h"
diff --git a/Engine/Source/Runtime/NavigationSystem/Public/NavAreas/NavArea.h b/Engine/Source/Runtime/NavigationSystem/Public/NavAreas/NavArea.h
--- a/Engine/Source/Runtime/NavigationSystem/Public/NavAreas/NavArea.h
+++ b/Engine/Source/Runtime/NavigationSystem/Public/NavAreas/NavArea.h
@@ -2,7 +2,9 @@
 
 #pragma once
 
+#if UE_ENABLE_INCLUDE_ORDER_DEPRECATED_IN_5_4
 #include "CoreMinimal.h"
+#endif
 #include "UObject/ObjectMacros.h"
 #include "UObject/Object.h"
 #include "Templates/SubclassOf.h"
diff --git a/Engine/Source/Runtime/NavigationSystem/Public/NavAreas/NavAreaMeta.h b/Engine/Source/Runtime/NavigationSystem/Public/NavAreas/NavAreaMeta.h
--- a/Engine/Source/Runtime/NavigationSystem/Public/NavAreas/NavAreaMeta.h
+++ b/Engine/Source/Runtime/NavigationSystem/Public/NavAreas/NavAreaMeta.h
@@ -2,7 +2,9 @@
 
 #pragma once
 
+#if UE_ENABLE_INCLUDE_ORDER_DEPRECATED_IN_5_4
 #include "CoreMinimal.h"
+#endif
 #include "NavAreas/NavArea.h"
 #include "Templates/SubclassOf.h"
 #include "UObject/ObjectMacros.h"
diff --git a/Engine/Source/Runtime/NavigationSystem/Public/NavAreas/NavAreaMeta_SwitchByAgent.h b/Engine/Source/Runtime/NavigationSystem/Public/NavAreas/NavAreaMeta_SwitchByAgent.h
--- a/Engine/Source/Runtime/NavigationSystem/Public/NavAreas/NavAreaMeta_SwitchByAgent.h
+++ b/Engine/Source/Runtime/NavigationSystem/Public/NavAreas/NavAreaMeta_SwitchByAgent.h
@@ -2,7 +2,9 @@
 
 #pragma once
 
+#if UE_ENABLE_INCLUDE_ORDER_DEPRECATED_IN_5_4
 #include "CoreMinimal.h"
+#endif
 #include "UObject/ObjectMacros.h"
 #include "Templates/SubclassOf.h"
 #include "NavAreas/NavAreaMeta.h"
diff --git a/Engine/Source/Runtime/NavigationSystem/Public/NavAreas/NavArea_Default.h b/Engine/Source/Runtime/NavigationSystem/Public/NavAreas/NavArea_Default.h
--- a/Engine/Source/Runtime/NavigationSystem/Public/NavAreas/NavArea_Default.h
+++ b/Engine/Source/Runtime/NavigationSystem/Public/NavAreas/NavArea_Default.h
@@ -2,7 +2,9 @@
 
 #pragma once
 
+#if UE_ENABLE_INCLUDE_ORDER_DEPRECATED_IN_5_4
 #include "CoreMinimal.h"
+#endif
 #include "NavAreas/NavArea.h"
 #include "UObject/ObjectMacros.h"
 #include "UObject/UObjectGlobals.h"
diff --git a/Engine/Source/Runtime/NavigationSystem/Public/NavAreas/NavArea_Null.h b/Engine/Source/Runtime/NavigationSystem/Public/NavAreas/NavArea_Null.h
--- a/Engine/Source/Runtime/NavigationSystem/Public/NavAreas/NavArea_Null.h
+++ b/Engine/Source/Runtime/NavigationSystem/Public/NavAreas/NavArea_Null.h
@@ -2,7 +2,9 @@
 
 #pragma once
 
+#if UE_ENABLE_INCLUDE_ORDER_DEPRECATED_IN_5_4
 #include "CoreMinimal.h"
+#endif
 #include "NavAreas/NavArea.h"
 #include "UObject/ObjectMacros.h"
 #include "UObject/UObjectGlobals.h"
diff --git a/Engine/Source/Runtime/NavigationSystem/Public/NavAreas/NavArea_Obstacle.h b/Engine/Source/Runtime/NavigationSystem/Public/NavAreas/NavArea_Obstacle.h
--- a/Engine/Source/Runtime/NavigationSystem/Public/NavAreas/NavArea_Obstacle.h
+++ b/Engine/Source/Runtime/NavigationSystem/Public/NavAreas/NavArea_Obstacle.h
@@ -2,7 +2,9 @@
 
 #pragma once
 
+#if UE_ENABLE_INCLUDE_ORDER_DEPRECATED_IN_5_4
 #include "CoreMinimal.h"
+#endif
 #include "NavAreas/NavArea.h"
 #include "UObject/ObjectMacros.h"
 #include "UObject/UObjectGlobals.h"
diff --git a/Engine/Source/Runtime/NavigationSystem/Public/NavCollision.h b/Engine/Source/Runtime/NavigationSystem/Public/NavCollision.h
--- a/Engine/Source/Runtime/NavigationSystem/Public/NavCollision.h
+++ b/Engine/Source/Runtime/NavigationSystem/Public/NavCollision.h
@@ -2,7 +2,9 @@
 
 #pragma once
 
+#if UE_ENABLE_INCLUDE_ORDER_DEPRECATED_IN_5_4
 #include "CoreMinimal.h"
+#endif
 #include "UObject/ObjectMacros.h"
 #include "UObject/Object.h"
 #include "Misc/Guid.h"
diff --git a/Engine/Source/Runtime/NavigationSystem/Public/NavDebugTypes.h b/Engine/Source/Runtime/NavigationSystem/Public/NavDebugTypes.h
--- a/Engine/Source/Runtime/NavigationSystem/Public/NavDebugTypes.h
+++ b/Engine/Source/Runtime/NavigationSystem/Public/NavDebugTypes.h
@@ -2,7 +2,9 @@
 
 #pragma once
 
+#if UE_ENABLE_INCLUDE_ORDER_DEPRECATED_IN_5_4
 #include "CoreMinimal.h"
+#endif
 #include "EngineDefines.h"
 
 
diff --git a/Engine/Source/Runtime/NavigationSystem/Public/NavFilters/NavigationQueryFilter.h b/Engine/Source/Runtime/NavigationSystem/Public/NavFilters/NavigationQueryFilter.h
--- a/Engine/Source/Runtime/NavigationSystem/Public/NavFilters/NavigationQueryFilter.h
+++ b/Engine/Source/Runtime/NavigationSystem/Public/NavFilters/NavigationQueryFilter.h
@@ -2,7 +2,9 @@
 
 #pragma once
 
+#if UE_ENABLE_INCLUDE_ORDER_DEPRECATED_IN_5_4
 #include "CoreMinimal.h"
+#endif
 #include "UObject/ObjectMacros.h"
 #include "UObject/Object.h"
 #include "Templates/SubclassOf.h"
diff --git a/Engine/Source/Runtime/NavigationSystem/Public/NavFilters/RecastFilter_UseDefaultArea.h b/Engine/Source/Runtime/NavigationSystem/Public/NavFilters/RecastFilter_UseDefaultArea.h
--- a/Engine/Source/Runtime/NavigationSystem/Public/NavFilters/RecastFilter_UseDefaultArea.h
+++ b/Engine/Source/Runtime/NavigationSystem/Public/NavFilters/RecastFilter_UseDefaultArea.h
@@ -2,7 +2,9 @@
 
 #pragma once
 
+#if UE_ENABLE_INCLUDE_ORDER_DEPRECATED_IN_5_4
 #include "CoreMinimal.h"
+#endif
 #include "NavFilters/NavigationQueryFilter.h"
 #include "UObject/ObjectMacros.h"
 #include "UObject/UObjectGlobals.h"
diff --git a/Engine/Source/Runtime/NavigationSystem/Public/NavGraph/NavigationGraph.h b/Engine/Source/Runtime/NavigationSystem/Public/NavGraph/NavigationGraph.h
--- a/Engine/Source/Runtime/NavigationSystem/Public/NavGraph/NavigationGraph.h
+++ b/Engine/Source/Runtime/NavigationSystem/Public/NavGraph/NavigationGraph.h
@@ -2,7 +2,9 @@
 
 #pragma once
 
+#if UE_ENABLE_INCLUDE_ORDER_DEPRECATED_IN_5_4
 #include "CoreMinimal.h"
+#endif
 #include "UObject/ObjectMacros.h"
 #include "NavigationData.h"
 #include "NavigationGraph.generated.h"
diff --git a/Engine/Source/Runtime/NavigationSystem/Public/NavGraph/NavigationGraphNode.h b/Engine/Source/Runtime/NavigationSystem/Public/NavGraph/NavigationGraphNode.h
--- a/Engine/Source/Runtime/NavigationSystem/Public/NavGraph/NavigationGraphNode.h
+++ b/Engine/Source/Runtime/NavigationSystem/Public/NavGraph/NavigationGraphNode.h
@@ -2,7 +2,9 @@
 
 #pragma once
 
+#if UE_ENABLE_INCLUDE_ORDER_DEPRECATED_IN_5_4
 #include "CoreMinimal.h"
+#endif
 #include "UObject/ObjectMacros.h"
 #include "GameFramework/Actor.h"
 
diff --git a/Engine/Source/Runtime/NavigationSystem/Public/NavGraph/NavigationGraphNodeComponent.h b/Engine/Source/Runtime/NavigationSystem/Public/NavGraph/NavigationGraphNodeComponent.h
--- a/Engine/Source/Runtime/NavigationSystem/Public/NavGraph/NavigationGraphNodeComponent.h
+++ b/Engine/Source/Runtime/NavigationSystem/Public/NavGraph/NavigationGraphNodeComponent.h
@@ -2,7 +2,9 @@
 
 #pragma once
 
+#if UE_ENABLE_INCLUDE_ORDER_DEPRECATED_IN_5_4
 #include "CoreMinimal.h"
+#endif
 #include "UObject/ObjectMacros.h"
 #include "Components/SceneComponent.h"
 #include "NavGraph/NavigationGraph.h"
diff --git a/Engine/Source/Runtime/NavigationSystem/Public/NavLinkComponent.h b/Engine/Source/Runtime/NavigationSystem/Public/NavLinkComponent.h
--- a/Engine/Source/Runtime/NavigationSystem/Public/NavLinkComponent.h
+++ b/Engine/Source/Runtime/NavigationSystem/Public/NavLinkComponent.h
@@ -2,7 +2,9 @@
 
 #pragma once
 
+#if UE_ENABLE_INCLUDE_ORDER_DEPRECATED_IN_5_4
 #include "CoreMinimal.h"
+#endif
 #include "UObject/ObjectMacros.h"
 #include "Templates/SubclassOf.h"
 #include "Components/PrimitiveComponent.h"
diff --git a/Engine/Source/Runtime/NavigationSystem/Public/NavLinkCustomComponent.h b/Engine/Source/Runtime/NavigationSystem/Public/NavLinkCustomComponent.h
--- a/Engine/Source/Runtime/NavigationSystem/Public/NavLinkCustomComponent.h
+++ b/Engine/Source/Runtime/NavigationSystem/Public/NavLinkCustomComponent.h
@@ -2,7 +2,9 @@
 
 #pragma once
 
+#if UE_ENABLE_INCLUDE_ORDER_DEPRECATED_IN_5_4
 #include "CoreMinimal.h"
+#endif
 #include "UObject/ObjectMacros.h"
 #include "Templates/SubclassOf.h"
 #include "Engine/EngineTypes.h"
diff --git a/Engine/Source/Runtime/NavigationSystem/Public/NavLinkCustomInterface.h b/Engine/Source/Runtime/NavigationSystem/Public/NavLinkCustomInterface.h
--- a/Engine/Source/Runtime/NavigationSystem/Public/NavLinkCustomInterface.h
+++ b/Engine/Source/Runtime/NavigationSystem/Public/NavLinkCustomInterface.h
@@ -2,7 +2,9 @@
 
 #pragma once
 
+#if UE_ENABLE_INCLUDE_ORDER_DEPRECATED_IN_5_4
 #include "CoreMinimal.h"
+#endif
 #include "UObject/ObjectMacros.h"
 #include "Templates/SubclassOf.h"
 #include "UObject/Interface.h"
diff --git a/Engine/Source/Runtime/NavigationSystem/Public/NavLinkHostInterface.h b/Engine/Source/Runtime/NavigationSystem/Public/NavLinkHostInterface.h
--- a/Engine/Source/Runtime/NavigationSystem/Public/NavLinkHostInterface.h
+++ b/Engine/Source/Runtime/NavigationSystem/Public/NavLinkHostInterface.h
@@ -2,7 +2,9 @@
 
 #pragma once
 
+#if UE_ENABLE_INCLUDE_ORDER_DEPRECATED_IN_5_4
 #include "CoreMinimal.h"
+#endif
 #include "UObject/ObjectMacros.h"
 #include "Templates/SubclassOf.h"
 #include "UObject/Interface.h"
diff --git a/Engine/Source/Runtime/NavigationSystem/Public/NavLinkRenderingComponent.h b/Engine/Source/Runtime/NavigationSystem/Public/NavLinkRenderingComponent.h
--- a/Engine/Source/Runtime/NavigationSystem/Public/NavLinkRenderingComponent.h
+++ b/Engine/Source/Runtime/NavigationSystem/Public/NavLinkRenderingComponent.h
@@ -2,7 +2,9 @@
 
 #pragma once
 
+#if UE_ENABLE_INCLUDE_ORDER_DEPRECATED_IN_5_4
 #include "CoreMinimal.h"
+#endif
 #include "UObject/ObjectMacros.h"
 #include "Components/PrimitiveComponent.h"
 #include "NavLinkRenderingComponent.generated.h"
diff --git a/Engine/Source/Runtime/NavigationSystem/Public/NavLinkRenderingProxy.h b/Engine/Source/Runtime/NavigationSystem/Public/NavLinkRenderingProxy.h
--- a/Engine/Source/Runtime/NavigationSystem/Public/NavLinkRenderingProxy.h
+++ b/Engine/Source/Runtime/NavigationSystem/Public/NavLinkRenderingProxy.h
@@ -2,7 +2,9 @@
 
 #pragma once
 
+#if UE_ENABLE_INCLUDE_ORDER_DEPRECATED_IN_5_4
 #include "CoreMinimal.h"
+#endif
 #include "NavAreas/NavArea.h"
 #include "PrimitiveViewRelevance.h"
 #include "PrimitiveSceneProxy.h"
diff --git a/Engine/Source/Runtime/NavigationSystem/Public/NavLinkTrivial.h b/Engine/Source/Runtime/NavigationSystem/Public/NavLinkTrivial.h
--- a/Engine/Source/Runtime/NavigationSystem/Public/NavLinkTrivial.h
+++ b/Engine/Source/Runtime/NavigationSystem/Public/NavLinkTrivial.h
@@ -2,7 +2,9 @@
 
 #pragma once
 
+#if UE_ENABLE_INCLUDE_ORDER_DEPRECATED_IN_5_4
 #include "CoreMinimal.h"
+#endif
 #include "UObject/ObjectMacros.h"
 #include "AI/Navigation/NavLinkDefinition.h"
 #include "NavLinkTrivial.generated.h"
diff --git a/Engine/Source/Runtime/NavigationSystem/Public/NavMesh/NavMeshBoundsVolume.h b/Engine/Source/Runtime/NavigationSystem/Public/NavMesh/NavMeshBoundsVolume.h
--- a/Engine/Source/Runtime/NavigationSystem/Public/NavMesh/NavMeshBoundsVolume.h
+++ b/Engine/Source/Runtime/NavigationSystem/Public/NavMesh/NavMeshBoundsVolume.h
@@ -7,9 +7,13 @@
 
 #pragma once
 
+#if UE_ENABLE_INCLUDE_ORDER_DEPRECATED_IN_5_4
 #include "CoreMinimal.h"
+#endif
 #include "UObject/ObjectMacros.h"
+#if UE_ENABLE_INCLUDE_ORDER_DEPRECATED_IN_5_4
 #include "AI/Navigation/NavigationTypes.h"
+#endif
 #include "GameFramework/Volume.h"
 #include "AI/Navigation/NavAgentSelector.h"
 #include "NavMeshBoundsVolume.generated.h"
diff --git a/Engine/Source/Runtime/NavigationSystem/Public/NavMesh/NavMeshPath.h b/Engine/Source/Runtime/NavigationSystem/Public/NavMesh/NavMeshPath.h
--- a/Engine/Source/Runtime/NavigationSystem/Public/NavMesh/NavMeshPath.h
+++ b/Engine/Source/Runtime/NavigationSystem/Public/NavMesh/NavMeshPath.h
@@ -2,7 +2,9 @@
 
 #pragma once
 
+#if UE_ENABLE_INCLUDE_ORDER_DEPRECATED_IN_5_4
 #include "CoreMinimal.h"
+#endif
 #include "UObject/ObjectMacros.h"
 #include "UObject/Object.h"
 #include "NavigationPath.h"
diff --git a/Engine/Source/Runtime/NavigationSystem/Public/NavMesh/PImplRecastNavMesh.h b/Engine/Source/Runtime/NavigationSystem/Public/NavMesh/PImplRecastNavMesh.h
--- a/Engine/Source/Runtime/NavigationSystem/Public/NavMesh/PImplRecastNavMesh.h
+++ b/Engine/Source/Runtime/NavigationSystem/Public/NavMesh/PImplRecastNavMesh.h
@@ -8,7 +8,9 @@
 
 #pragma once 
 
+#if UE_ENABLE_INCLUDE_ORDER_DEPRECATED_IN_5_4
 #include "CoreMinimal.h"
+#endif
 #include "UObject/WeakObjectPtr.h"
 #if UE_ENABLE_INCLUDE_ORDER_DEPRECATED_IN_5_4
 #include "NavFilters/NavigationQueryFilter.h"
diff --git a/Engine/Source/Runtime/NavigationSystem/Public/NavMesh/RecastHelpers.h b/Engine/Source/Runtime/NavigationSystem/Public/NavMesh/RecastHelpers.h
--- a/Engine/Source/Runtime/NavigationSystem/Public/NavMesh/RecastHelpers.h
+++ b/Engine/Source/Runtime/NavigationSystem/Public/NavMesh/RecastHelpers.h
@@ -10,11 +10,15 @@
 
 #pragma once
 
+#if UE_ENABLE_INCLUDE_ORDER_DEPRECATED_IN_5_4
 #include "CoreMinimal.h"
+#endif
 #include "Math/Box.h"
 #include "Math/Color.h"
 #include "Math/Matrix.h"
+#if UE_ENABLE_INCLUDE_ORDER_DEPRECATED_IN_5_4
 #include "Math/UnrealMathSSE.h"
+#endif
 #include "Math/Vector.h"
 
 
diff --git a/Engine/Source/Runtime/NavigationSystem/Public/NavMesh/RecastNavMesh.h b/Engine/Source/Runtime/NavigationSystem/Public/NavMesh/RecastNavMesh.h
--- a/Engine/Source/Runtime/NavigationSystem/Public/NavMesh/RecastNavMesh.h
+++ b/Engine/Source/Runtime/NavigationSystem/Public/NavMesh/RecastNavMesh.h
@@ -2,7 +2,9 @@
 
 #pragma once
 
+#if UE_ENABLE_INCLUDE_ORDER_DEPRECATED_IN_5_4
 #include "CoreMinimal.h"
+#endif
 #include "UObject/ObjectMacros.h"
 #include "Templates/SubclassOf.h"
 #include "EngineDefines.h"
diff --git a/Engine/Source/Runtime/NavigationSystem/Public/NavMesh/RecastNavMeshDataChunk.h b/Engine/Source/Runtime/NavigationSystem/Public/NavMesh/RecastNavMeshDataChunk.h
--- a/Engine/Source/Runtime/NavigationSystem/Public/NavMesh/RecastNavMeshDataChunk.h
+++ b/Engine/Source/Runtime/NavigationSystem/Public/NavMesh/RecastNavMeshDataChunk.h
@@ -2,9 +2,13 @@
 
 #pragma once
 
+#if UE_ENABLE_INCLUDE_ORDER_DEPRECATED_IN_5_4
 #include "CoreMinimal.h"
+#endif
 #include "UObject/ObjectMacros.h"
+#if UE_ENABLE_INCLUDE_ORDER_DEPRECATED_IN_5_4
 #include "AI/Navigation/NavigationTypes.h"
+#endif
 #include "AI/Navigation/NavigationDataChunk.h"
 #include "NavMesh/RecastNavMesh.h"
 #include "RecastNavMeshDataChunk.generated.h"
diff --git a/Engine/Source/Runtime/NavigationSystem/Public/NavMesh/RecastNavMeshGenerator.h b/Engine/Source/Runtime/NavigationSystem/Public/NavMesh/RecastNavMeshGenerator.h
--- a/Engine/Source/Runtime/NavigationSystem/Public/NavMesh/RecastNavMeshGenerator.h
+++ b/Engine/Source/Runtime/NavigationSystem/Public/NavMesh/RecastNavMeshGenerator.h
@@ -1,7 +1,9 @@
 // Copyright Epic Games, Inc. All Rights Reserved.
 #pragma once 
 
+#if UE_ENABLE_INCLUDE_ORDER_DEPRECATED_IN_5_4
 #include "CoreMinimal.h"
+#endif
 #include "Stats/Stats.h"
 #include "AI/Navigation/NavigationTypes.h"
 #include "AI/Navigation/NavigationInvokerPriority.h"
diff --git a/Engine/Source/Runtime/NavigationSystem/Public/NavMesh/RecastQueryFilter.h b/Engine/Source/Runtime/NavigationSystem/Public/NavMesh/RecastQueryFilter.h
--- a/Engine/Source/Runtime/NavigationSystem/Public/NavMesh/RecastQueryFilter.h
+++ b/Engine/Source/Runtime/NavigationSystem/Public/NavMesh/RecastQueryFilter.h
@@ -2,7 +2,9 @@
 
 #pragma once 
 
+#if UE_ENABLE_INCLUDE_ORDER_DEPRECATED_IN_5_4
 #include "CoreMinimal.h"
+#endif
 #if UE_ENABLE_INCLUDE_ORDER_DEPRECATED_IN_5_2
 #include "NavFilters/NavigationQueryFilter.h"
 #endif //UE_ENABLE_INCLUDE_ORDER_DEPRECATED_IN_5_2
@@ -10,7 +12,10 @@
 #include "AI/Navigation/NavQueryFilter.h"
 
 #if WITH_RECAST
+
+#if UE_ENABLE_INCLUDE_ORDER_DEPRECATED_IN_5_4
 #include "Detour/DetourNavMesh.h"
+#endif
 #include "Detour/DetourNavMeshQuery.h"
 
 #define RECAST_VERY_SMALL_AGENT_RADIUS 0.0f
diff --git a/Engine/Source/Runtime/NavigationSystem/Public/NavModifierComponent.h b/Engine/Source/Runtime/NavigationSystem/Public/NavModifierComponent.h
--- a/Engine/Source/Runtime/NavigationSystem/Public/NavModifierComponent.h
+++ b/Engine/Source/Runtime/NavigationSystem/Public/NavModifierComponent.h
@@ -2,7 +2,9 @@
 
 #pragma once
 
+#if UE_ENABLE_INCLUDE_ORDER_DEPRECATED_IN_5_4
 #include "CoreMinimal.h"
+#endif
 #include "UObject/ObjectMacros.h"
 #include "Templates/SubclassOf.h"
 #include "NavAreas/NavArea.h"
diff --git a/Engine/Source/Runtime/NavigationSystem/Public/NavModifierVolume.h b/Engine/Source/Runtime/NavigationSystem/Public/NavModifierVolume.h
--- a/Engine/Source/Runtime/NavigationSystem/Public/NavModifierVolume.h
+++ b/Engine/Source/Runtime/NavigationSystem/Public/NavModifierVolume.h
@@ -2,7 +2,9 @@
 
 #pragma once 
 
+#if UE_ENABLE_INCLUDE_ORDER_DEPRECATED_IN_5_4
 #include "CoreMinimal.h"
+#endif
 #include "UObject/ObjectMacros.h"
 #include "Templates/SubclassOf.h"
 #include "AI/Navigation/NavRelevantInterface.h"
diff --git a/Engine/Source/Runtime/NavigationSystem/Public/NavNodeInterface.h b/Engine/Source/Runtime/NavigationSystem/Public/NavNodeInterface.h
--- a/Engine/Source/Runtime/NavigationSystem/Public/NavNodeInterface.h
+++ b/Engine/Source/Runtime/NavigationSystem/Public/NavNodeInterface.h
@@ -2,10 +2,14 @@
 
 #pragma once
 
+#if UE_ENABLE_INCLUDE_ORDER_DEPRECATED_IN_5_4
 #include "CoreMinimal.h"
+#endif
 #include "UObject/ObjectMacros.h"
 #include "UObject/Interface.h"
+#if UE_ENABLE_INCLUDE_ORDER_DEPRECATED_IN_5_4
 #include "AI/Navigation/NavigationTypes.h"
+#endif
 #include "NavNodeInterface.generated.h"
 
 UINTERFACE(MinimalAPI, meta=(CannotImplementInterfaceInBlueprint))
diff --git a/Engine/Source/Runtime/NavigationSystem/Public/NavRelevantComponent.h b/Engine/Source/Runtime/NavigationSystem/Public/NavRelevantComponent.h
--- a/Engine/Source/Runtime/NavigationSystem/Public/NavRelevantComponent.h
+++ b/Engine/Source/Runtime/NavigationSystem/Public/NavRelevantComponent.h
@@ -2,7 +2,9 @@
 
 #pragma once
 
+#if UE_ENABLE_INCLUDE_ORDER_DEPRECATED_IN_5_4
 #include "CoreMinimal.h"
+#endif
 #include "UObject/ObjectMacros.h"
 #include "Components/ActorComponent.h"
 #include "AI/Navigation/NavRelevantInterface.h"
diff --git a/Engine/Source/Runtime/NavigationSystem/Public/NavSystemConfigOverride.h b/Engine/Source/Runtime/NavigationSystem/Public/NavSystemConfigOverride.h
--- a/Engine/Source/Runtime/NavigationSystem/Public/NavSystemConfigOverride.h
+++ b/Engine/Source/Runtime/NavigationSystem/Public/NavSystemConfigOverride.h
@@ -2,7 +2,9 @@
 
 #pragma once
 
+#if UE_ENABLE_INCLUDE_ORDER_DEPRECATED_IN_5_4
 #include "CoreMinimal.h"
+#endif
 #include "UObject/ObjectMacros.h"
 #include "GameFramework/Actor.h"
 #include "AI/NavigationSystemBase.h"
diff --git a/Engine/Source/Runtime/NavigationSystem/Public/NavigationData.h b/Engine/Source/Runtime/NavigationSystem/Public/NavigationData.h
--- a/Engine/Source/Runtime/NavigationSystem/Public/NavigationData.h
+++ b/Engine/Source/Runtime/NavigationSystem/Public/NavigationData.h
@@ -2,7 +2,9 @@
 
 #pragma once
 
+#if UE_ENABLE_INCLUDE_ORDER_DEPRECATED_IN_5_4
 #include "CoreMinimal.h"
+#endif
 #include "Stats/Stats.h"
 #include "UObject/ObjectMacros.h"
 #include "Templates/SubclassOf.h"
diff --git a/Engine/Source/Runtime/NavigationSystem/Public/NavigationInvokerComponent.h b/Engine/Source/Runtime/NavigationSystem/Public/NavigationInvokerComponent.h
--- a/Engine/Source/Runtime/NavigationSystem/Public/NavigationInvokerComponent.h
+++ b/Engine/Source/Runtime/NavigationSystem/Public/NavigationInvokerComponent.h
@@ -1,7 +1,9 @@
 // Copyright Epic Games, Inc. All Rights Reserved.
 #pragma once
 
+#if UE_ENABLE_INCLUDE_ORDER_DEPRECATED_IN_5_4
 #include "CoreMinimal.h"
+#endif
 #include "UObject/ObjectMacros.h"
 #include "UObject/UObjectGlobals.h"
 #include "AI/Navigation/NavAgentSelector.h"
diff --git a/Engine/Source/Runtime/NavigationSystem/Public/NavigationOctree.h b/Engine/Source/Runtime/NavigationSystem/Public/NavigationOctree.h
--- a/Engine/Source/Runtime/NavigationSystem/Public/NavigationOctree.h
+++ b/Engine/Source/Runtime/NavigationSystem/Public/NavigationOctree.h
@@ -2,7 +2,9 @@
 
 #pragma once
 
+#if UE_ENABLE_INCLUDE_ORDER_DEPRECATED_IN_5_4
 #include "CoreMinimal.h"
+#endif
 #include "Stats/Stats.h"
 #include "AI/Navigation/NavigationTypes.h"
 #include "Math/GenericOctreePublic.h"
diff --git a/Engine/Source/Runtime/NavigationSystem/Public/NavigationPath.h b/Engine/Source/Runtime/NavigationSystem/Public/NavigationPath.h
--- a/Engine/Source/Runtime/NavigationSystem/Public/NavigationPath.h
+++ b/Engine/Source/Runtime/NavigationSystem/Public/NavigationPath.h
@@ -2,7 +2,9 @@
 
 #pragma once
 
+#if UE_ENABLE_INCLUDE_ORDER_DEPRECATED_IN_5_4
 #include "CoreMinimal.h"
+#endif
 #include "UObject/ObjectMacros.h"
 #include "UObject/Object.h"
 #include "NavigationData.h"
diff --git a/Engine/Source/Runtime/NavigationSystem/Public/NavigationPathGenerator.h b/Engine/Source/Runtime/NavigationSystem/Public/NavigationPathGenerator.h
--- a/Engine/Source/Runtime/NavigationSystem/Public/NavigationPathGenerator.h
+++ b/Engine/Source/Runtime/NavigationSystem/Public/NavigationPathGenerator.h
@@ -2,7 +2,9 @@
 
 #pragma once
 
+#if UE_ENABLE_INCLUDE_ORDER_DEPRECATED_IN_5_4
 #include "CoreMinimal.h"
+#endif
 #include "UObject/ObjectMacros.h"
 #include "UObject/Interface.h"
 #include "AI/Navigation/NavigationTypes.h"
diff --git a/Engine/Source/Runtime/NavigationSystem/Public/NavigationSystem.h b/Engine/Source/Runtime/NavigationSystem/Public/NavigationSystem.h
--- a/Engine/Source/Runtime/NavigationSystem/Public/NavigationSystem.h
+++ b/Engine/Source/Runtime/NavigationSystem/Public/NavigationSystem.h
@@ -2,7 +2,9 @@
 
 #pragma once
 
+#if UE_ENABLE_INCLUDE_ORDER_DEPRECATED_IN_5_4
 #include "CoreMinimal.h"
+#endif
 #include "UObject/ObjectMacros.h"
 #include "Templates/SubclassOf.h"
 #include "UObject/WeakObjectPtr.h"
diff --git a/Engine/Source/Runtime/NavigationSystem/Public/NavigationSystemModule.h b/Engine/Source/Runtime/NavigationSystem/Public/NavigationSystemModule.h
--- a/Engine/Source/Runtime/NavigationSystem/Public/NavigationSystemModule.h
+++ b/Engine/Source/Runtime/NavigationSystem/Public/NavigationSystemModule.h
@@ -2,7 +2,9 @@
 
 #pragma once
 
+#if UE_ENABLE_INCLUDE_ORDER_DEPRECATED_IN_5_4
 #include "CoreMinimal.h"
+#endif
 #include "Modules/ModuleManager.h"
 //#include "AI/NavigationSystemBase.h"
 #if WITH_EDITOR
diff --git a/Engine/Source/Runtime/NavigationSystem/Public/NavigationSystemTypes.h b/Engine/Source/Runtime/NavigationSystem/Public/NavigationSystemTypes.h
--- a/Engine/Source/Runtime/NavigationSystem/Public/NavigationSystemTypes.h
+++ b/Engine/Source/Runtime/NavigationSystem/Public/NavigationSystemTypes.h
@@ -2,13 +2,15 @@
 
 #pragma once
 
-#include "CoreMinimal.h"
 #if UE_ENABLE_INCLUDE_ORDER_DEPRECATED_IN_5_4
+#include "CoreMinimal.h"
 #include "NavFilters/NavigationQueryFilter.h"
-#endif
+#endif // UE_ENABLE_INCLUDE_ORDER_DEPRECATED_IN_5_4
 #include "AI/NavigationSystemConfig.h"
 #include "AI/Navigation/NavLinkDefinition.h"
+#if UE_ENABLE_INCLUDE_ORDER_DEPRECATED_IN_5_4
 #include "Math/GenericOctreePublic.h"
+#endif
 #include "AI/NavigationModifier.h"
 #include "UObject/WeakInterfacePtr.h"
 
diff --git a/Engine/Source/Runtime/NavigationSystem/Public/NavigationTestingActor.h b/Engine/Source/Runtime/NavigationSystem/Public/NavigationTestingActor.h
--- a/Engine/Source/Runtime/NavigationSystem/Public/NavigationTestingActor.h
+++ b/Engine/Source/Runtime/NavigationSystem/Public/NavigationTestingActor.h
@@ -2,8 +2,10 @@
 
 #pragma once
 
+#if UE_ENABLE_INCLUDE_ORDER_DEPRECATED_IN_5_4
 #include "CoreMinimal.h"
 #include "Stats/Stats.h"
+#endif
 #include "UObject/ObjectMacros.h"
 #include "Templates/SubclassOf.h"
 #include "GameFramework/Actor.h"
