Removing NavigationTypes.h from NavRelevantInterface.h, StaticMeshActor.h and AITypes.h
Removing NavigationSystemConfig.h and NavLinkDefinition.h from NavigationSystemTypes.h

diff --git a/Engine/Plugins/Runtime/NavCorridor/Source/NavCorridor/Public/NavCorridorTestingComponent.h b/Engine/Plugins/Runtime/NavCorridor/Source/NavCorridor/Public/NavCorridorTestingComponent.h
--- a/Engine/Plugins/Runtime/NavCorridor/Source/NavCorridor/Public/NavCorridorTestingComponent.h
+++ b/Engine/Plugins/Runtime/NavCorridor/Source/NavCorridor/Public/NavCorridorTestingComponent.h
@@ -4,6 +4,7 @@
 
 #include "Debug/DebugDrawComponent.h"
 #include "NavCorridor.h"
+#include "AI/Navigation/NavigationTypes.h" // FNavAgentProperties
 #include "NavCorridorTestingComponent.generated.h"
 
 class ANavigationData;
diff --git a/Engine/Source/Developer/LogVisualizer/Private/VisualLoggerRenderingActorBase.cpp b/Engine/Source/Developer/LogVisualizer/Private/VisualLoggerRenderingActorBase.cpp
--- a/Engine/Source/Developer/LogVisualizer/Private/VisualLoggerRenderingActorBase.cpp
+++ b/Engine/Source/Developer/LogVisualizer/Private/VisualLoggerRenderingActorBase.cpp
@@ -7,6 +7,7 @@
 #endif // WITH_EDITOR
 #include "VisualLoggerRenderingComponent.h"
 #include "VisualLogger/VisualLogger.h"
+#include "AI/Navigation/NavigationTypes.h"
 
 #include UE_INLINE_GENERATED_CPP_BY_NAME(VisualLoggerRenderingActorBase)
 
diff --git a/Engine/Source/Runtime/AIModule/Classes/AITypes.h b/Engine/Source/Runtime/AIModule/Classes/AITypes.h
--- a/Engine/Source/Runtime/AIModule/Classes/AITypes.h
+++ b/Engine/Source/Runtime/AIModule/Classes/AITypes.h
@@ -8,8 +8,8 @@
 #include "Templates/SubclassOf.h"
 #if UE_ENABLE_INCLUDE_ORDER_DEPRECATED_IN_5_4
 #include "NavFilters/NavigationQueryFilter.h"
-#endif
 #include "AI/Navigation/NavigationTypes.h"
+#endif
 #include "NavigationSystemTypes.h"
 #if UE_ENABLE_INCLUDE_ORDER_DEPRECATED_IN_5_4
 #include "GameFramework/Actor.h"
@@ -17,6 +17,7 @@
 #include "AITypes.generated.h"
 
 class AActor;
+typedef TSharedPtr<FMoveRequestCustomData, ESPMode::ThreadSafe> FCustomMoveSharedPtr;
 
 DECLARE_CYCLE_STAT_EXTERN(TEXT("Overall AI Time"), STAT_AI_Overall, STATGROUP_AI, AIMODULE_API);
 
diff --git a/Engine/Source/Runtime/Engine/Classes/AI/Navigation/NavDataGatheringMode.h b/Engine/Source/Runtime/Engine/Classes/AI/Navigation/NavDataGatheringMode.h
new file mode 100644
--- /dev/null
+++ b/Engine/Source/Runtime/Engine/Classes/AI/Navigation/NavDataGatheringMode.h
@@ -0,0 +1,13 @@
+﻿// Copyright Epic Games, Inc. All Rights Reserved.
+
+#pragma once
+
+#include "NavDataGatheringMode.generated.h"
+
+UENUM()
+enum class ENavDataGatheringMode : uint8
+{
+	Default,
+	Instant,
+	Lazy
+};
diff --git a/Engine/Source/Runtime/Engine/Classes/AI/Navigation/NavRelevantInterface.h b/Engine/Source/Runtime/Engine/Classes/AI/Navigation/NavRelevantInterface.h
--- a/Engine/Source/Runtime/Engine/Classes/AI/Navigation/NavRelevantInterface.h
+++ b/Engine/Source/Runtime/Engine/Classes/AI/Navigation/NavRelevantInterface.h
@@ -7,8 +7,9 @@
 #endif
 #include "UObject/ObjectMacros.h"
 #include "UObject/Interface.h"
-#include "AI/Navigation/NavigationTypes.h"	//ENavDataGatheringMode
+#include "AI/Navigation/NavDataGatheringMode.h"
 #if UE_ENABLE_INCLUDE_ORDER_DEPRECATED_IN_5_4
+#include "AI/Navigation/NavigationTypes.h"
 #include "AI/NavigationModifier.h"
 #include "AI/Navigation/NavigationRelevantData.h"
 #endif
diff --git a/Engine/Source/Runtime/Engine/Classes/AI/Navigation/NavigationTypes.h b/Engine/Source/Runtime/Engine/Classes/AI/Navigation/NavigationTypes.h
--- a/Engine/Source/Runtime/Engine/Classes/AI/Navigation/NavigationTypes.h
+++ b/Engine/Source/Runtime/Engine/Classes/AI/Navigation/NavigationTypes.h
@@ -10,6 +10,7 @@
 #include "UObject/SoftObjectPath.h"
 #if UE_ENABLE_INCLUDE_ORDER_DEPRECATED_IN_5_4
 #include "GameFramework/Actor.h"
+#include "AI/Navigation/NavDataGatheringMode.h"
 #endif
 #if UE_ENABLE_INCLUDE_ORDER_DEPRECATED_IN_5_2
 #include "AI/Navigation/NavAgentSelector.h"
@@ -114,14 +115,6 @@ struct FNavigationDirtyArea
 	}
 };
 
-UENUM()
-enum class ENavDataGatheringMode : uint8
-{
-	Default,
-	Instant,
-	Lazy
-};
-
 UENUM()
 enum class ENavDataGatheringModeConfig : uint8
 {
diff --git a/Engine/Source/Runtime/Engine/Classes/Engine/StaticMeshActor.h b/Engine/Source/Runtime/Engine/Classes/Engine/StaticMeshActor.h
--- a/Engine/Source/Runtime/Engine/Classes/Engine/StaticMeshActor.h
+++ b/Engine/Source/Runtime/Engine/Classes/Engine/StaticMeshActor.h
@@ -5,7 +5,10 @@
 #include "CoreMinimal.h"
 #include "UObject/ObjectMacros.h"
 #include "GameFramework/Actor.h"
+#if UE_ENABLE_INCLUDE_ORDER_DEPRECATED_IN_5_4
 #include "AI/Navigation/NavigationTypes.h"
+#endif //UE_ENABLE_INCLUDE_ORDER_DEPRECATED_IN_5_4
+#include "AI/Navigation/NavDataGatheringMode.h"
 #include "StaticMeshActor.generated.h"
 
 /**
diff --git a/Engine/Source/Runtime/NavigationSystem/Public/NavMesh/NavTestRenderingComponent.h b/Engine/Source/Runtime/NavigationSystem/Public/NavMesh/NavTestRenderingComponent.h
--- a/Engine/Source/Runtime/NavigationSystem/Public/NavMesh/NavTestRenderingComponent.h
+++ b/Engine/Source/Runtime/NavigationSystem/Public/NavMesh/NavTestRenderingComponent.h
@@ -3,6 +3,7 @@
 #pragma once
 
 #include "Debug/DebugDrawComponent.h"
+#include "AI/Navigation/NavigationTypes.h" // NavNodeRef
 
 #if UE_ENABLE_INCLUDE_ORDER_DEPRECATED_IN_5_4
 #include "PrimitiveViewRelevance.h"
diff --git a/Engine/Source/Runtime/NavigationSystem/Public/NavMesh/RecastNavMesh.h b/Engine/Source/Runtime/NavigationSystem/Public/NavMesh/RecastNavMesh.h
--- a/Engine/Source/Runtime/NavigationSystem/Public/NavMesh/RecastNavMesh.h
+++ b/Engine/Source/Runtime/NavigationSystem/Public/NavMesh/RecastNavMesh.h
@@ -9,6 +9,7 @@
 #include "Templates/SubclassOf.h"
 #include "EngineDefines.h"
 #include "AI/Navigation/NavigationTypes.h"
+#include "AI/Navigation/NavigationDataResolution.h"
 #include "NavigationSystemTypes.h"
 #include "NavigationData.h"
 #include "NavMesh/NavMeshPath.h"
@@ -59,6 +60,7 @@ class dtQueryFilter;
 class FRecastNavMeshGenerator;
 struct dtMeshTile;
 class UNavigationSystemV1;
+class UNavigationSystemBase;
 
 UENUM()
 namespace ERecastPartitioning
diff --git a/Engine/Source/Runtime/NavigationSystem/Public/NavigationSystemTypes.h b/Engine/Source/Runtime/NavigationSystem/Public/NavigationSystemTypes.h
--- a/Engine/Source/Runtime/NavigationSystem/Public/NavigationSystemTypes.h
+++ b/Engine/Source/Runtime/NavigationSystem/Public/NavigationSystemTypes.h
@@ -5,13 +5,13 @@
 #if UE_ENABLE_INCLUDE_ORDER_DEPRECATED_IN_5_4
 #include "CoreMinimal.h"
 #include "NavFilters/NavigationQueryFilter.h"
-#endif // UE_ENABLE_INCLUDE_ORDER_DEPRECATED_IN_5_4
 #include "AI/NavigationSystemConfig.h"
 #include "AI/Navigation/NavLinkDefinition.h"
-#if UE_ENABLE_INCLUDE_ORDER_DEPRECATED_IN_5_4
 #include "Math/GenericOctreePublic.h"
-#endif
 #include "AI/NavigationModifier.h"
+#endif
+#include "AI/Navigation/NavigationTypes.h"
+#include "AI/Navigation/NavAgentSelector.h"
 #include "UObject/WeakInterfacePtr.h"
 
 #define NAVSYS_DEBUG (0 && UE_BUILD_DEBUG)
@@ -26,9 +26,12 @@ class FNavigationOctree;
 class UNavigationPath;
 class ANavigationData;
 class INavigationInvokerInterface;
-
+struct FCompositeNavModifier;
+struct FNavigationLink;
+struct FNavigationSegmentLink;
 struct FNavigationQueryFilter;
 typedef TSharedPtr<const FNavigationQueryFilter, ESPMode::ThreadSafe> FSharedConstNavQueryFilter;
+typedef TSharedPtr<struct FNavigationPath, ESPMode::ThreadSafe> FNavPathSharedPtr;
 
 struct FPathFindingQueryData
 {
