OnNavigationBoundsUpdated, make a removal request if the bounds are invalid.
This avoids having invalid RegisteredNavBounds when undoing the addtion of a NavMeshBoundVolume.

diff --git a/Engine/Source/Runtime/NavigationSystem/Private/NavigationSystem.cpp b/Engine/Source/Runtime/NavigationSystem/Private/NavigationSystem.cpp
--- a/Engine/Source/Runtime/NavigationSystem/Private/NavigationSystem.cpp
+++ b/Engine/Source/Runtime/NavigationSystem/Private/NavigationSystem.cpp
@@ -3700,7 +3700,16 @@ void UNavigationSystemV1::OnNavigationBoundsUpdated(ANavMeshBoundsVolume* NavVol
 	UpdateRequest.NavBounds.Level = NavVolume->GetLevel();
 	UpdateRequest.NavBounds.SupportedAgents = NavVolume->SupportedAgents;
 	
-	UpdateRequest.UpdateRequest = FNavigationBoundsUpdateRequest::Updated;
+	if (UpdateRequest.NavBounds.AreaBox.IsValid)
+	{
+		UpdateRequest.UpdateRequest = FNavigationBoundsUpdateRequest::Updated;
+	}
+	else
+	{
+		// Make a removal request if the bounds are invalid.
+		UpdateRequest.UpdateRequest = FNavigationBoundsUpdateRequest::Removed;
+	}
+
 	AddNavigationBoundsUpdateRequest(UpdateRequest);
 }
 
