[Invokers] Avoid reserving local containers every frame.
Add a profiler tag for more precise tracking.

diff --git a/Engine/Source/Runtime/NavigationSystem/Private/NavMesh/RecastNavMesh.cpp b/Engine/Source/Runtime/NavigationSystem/Private/NavMesh/RecastNavMesh.cpp
--- a/Engine/Source/Runtime/NavigationSystem/Private/NavMesh/RecastNavMesh.cpp
+++ b/Engine/Source/Runtime/NavigationSystem/Private/NavMesh/RecastNavMesh.cpp
@@ -3572,14 +3572,40 @@ void ARecastNavMesh::UpdateActiveTiles(const TArray<FNavigationInvokerRaw>& Invo
 	const FVector NavmeshOrigin = Recast2UnrealPoint(NavParams->orig);
 	const FVector::FReal TileDim = Config.GetTileSizeUU();
 
+	TSet<FIntPoint>& OldActiveSet = UpdateActiveTilesWorkingMem.OldActiveSet;
+	TArray<FNavMeshDirtyTileElement>& TilesInMinDistance = UpdateActiveTilesWorkingMem.TilesInMinDistance;
+	TSet<FIntPoint>& TilesInMaxDistance = UpdateActiveTilesWorkingMem.TilesInMaxDistance;
+	TArray<FIntPoint>& TileToAppend = UpdateActiveTilesWorkingMem.TileToAppend;
+	
 	TSet<FIntPoint>& ActiveTiles = GetActiveTileSet();
-	TSet<FIntPoint> OldActiveSet = ActiveTiles;
-	TArray<FNavMeshDirtyTileElement> TilesInMinDistance;
-	TSet<FIntPoint> TilesInMaxDistance;
-	TArray<FIntPoint> TileToAppend;
-	TilesInMinDistance.Reserve(ActiveTiles.Num());
-	TilesInMaxDistance.Reserve(ActiveTiles.Num());
-	TileToAppend.Reserve(ActiveTiles.Num());
+	const int32 ActiveTilesCount = ActiveTiles.Num();
+	OldActiveSet = ActiveTiles;
+
+	{
+		TRACE_CPUPROFILER_EVENT_SCOPE(ARecastNavMesh::Reserving);
+		TilesInMinDistance.Reset();
+		TilesInMaxDistance.Reset();
+		TileToAppend.Reset();
+
+		const int32 ShrinkThreshold = 1.2*ActiveTilesCount;
+		if (TilesInMinDistance.Max() > ShrinkThreshold)
+		{
+			TilesInMinDistance.Shrink();
+		}
+		if (TilesInMaxDistance.GetMaxIndex() > ShrinkThreshold)
+		{
+			TilesInMaxDistance.Shrink();
+		}
+		if (TileToAppend.Max() > ShrinkThreshold)
+		{
+			TileToAppend.Shrink();
+		}
+		
+		TilesInMinDistance.Reserve(ActiveTilesCount);
+		TilesInMaxDistance.Reserve(ActiveTilesCount);
+		TileToAppend.Reserve(ActiveTilesCount);
+	}
+
 	ActiveTiles.Reset();
 
 	{
diff --git a/Engine/Source/Runtime/NavigationSystem/Public/NavMesh/RecastNavMesh.h b/Engine/Source/Runtime/NavigationSystem/Public/NavMesh/RecastNavMesh.h
--- a/Engine/Source/Runtime/NavigationSystem/Public/NavMesh/RecastNavMesh.h
+++ b/Engine/Source/Runtime/NavigationSystem/Public/NavMesh/RecastNavMesh.h
@@ -1529,6 +1529,15 @@ protected:
 	FPImplRecastNavMesh* GetRecastNavMeshImpl() { return RecastNavMeshImpl; }
 	const FPImplRecastNavMesh* GetRecastNavMeshImpl() const { return RecastNavMeshImpl; }
 
+	struct FUpdateActiveTilesWorkingMem
+	{
+		TSet<FIntPoint> OldActiveSet;
+		TArray<FNavMeshDirtyTileElement> TilesInMinDistance;
+		TSet<FIntPoint> TilesInMaxDistance;
+		TArray<FIntPoint> TileToAppend;
+	};
+	FUpdateActiveTilesWorkingMem UpdateActiveTilesWorkingMem;
+	
 private:
 	/** @return Navmesh data chunk that belongs to this actor */
 	NAVIGATIONSYSTEM_API URecastNavMeshDataChunk* GetNavigationDataChunk(const TArray<UNavigationDataChunk*>& InChunks) const;
