[Navigation] Addition of a query to ANavigationTestingActor to display if navigation data is ready (has been generated) within a given radius.

diff --git a/Engine/Source/Runtime/NavigationSystem/Private/NavMesh/NavTestRenderingComponent.cpp b/Engine/Source/Runtime/NavigationSystem/Private/NavMesh/NavTestRenderingComponent.cpp
--- a/Engine/Source/Runtime/NavigationSystem/Private/NavMesh/NavTestRenderingComponent.cpp
+++ b/Engine/Source/Runtime/NavigationSystem/Private/NavMesh/NavTestRenderingComponent.cpp
@@ -97,6 +97,16 @@ void FNavTestSceneProxy::GetDynamicMeshElements(const TArray<const FSceneView*>&
 				{
 					PDI->DrawLine(ClosestWallLocation, ActorLocation, ClosestWallColor, SDPG_World, 2.5);
 				}
+
+				if (NavTestActor->bDrawIfNavDataIsReadyInRadius)
+				{
+					constexpr double HalfHeight = 1000;
+					constexpr int32 NumSides = 32;
+					DrawWireCylinder(PDI, ActorLocation, FVector(1, 0, 0), FVector(0, 1, 0), FVector(0, 0, 1),
+						NavTestActor->bNavDataIsReadyInRadius ? FColor::Green : FColor::Red, 
+						NavTestActor->RadiusUsedToValidateNavData, HalfHeight,  
+						NumSides, SDPG_World);
+				}
 			}
 
 			// draw path
diff --git a/Engine/Source/Runtime/NavigationSystem/Private/NavigationTestingActor.cpp b/Engine/Source/Runtime/NavigationSystem/Private/NavigationTestingActor.cpp
--- a/Engine/Source/Runtime/NavigationSystem/Private/NavigationTestingActor.cpp
+++ b/Engine/Source/Runtime/NavigationSystem/Private/NavigationTestingActor.cpp
@@ -55,6 +55,7 @@ ANavigationTestingActor::ANavigationTestingActor(const FObjectInitializer& Objec
 	bGatherDetailedInfo = true;
 	bDrawDistanceToWall = false;
 	ClosestWallLocation = FNavigationSystem::InvalidLocation;
+	bNavDataIsReadyInRadius = false;
 	OffsetFromCornersDistance = 0.f;
 
 	QueryingExtent = FVector(DEFAULT_NAV_QUERY_EXTENT_HORIZONTAL, DEFAULT_NAV_QUERY_EXTENT_HORIZONTAL, DEFAULT_NAV_QUERY_EXTENT_VERTICAL);
@@ -162,6 +163,10 @@ void ANavigationTestingActor::PostEditChangeProperty(FPropertyChangedEvent& Prop
 			{
 				ClosestWallLocation = FindClosestWallLocation();
 			}
+			else if (bDrawIfNavDataIsReadyInRadius)
+			{
+				bNavDataIsReadyInRadius = CheckIfNavDataIsReadyInRadius();
+			}
 #if WITH_EDITORONLY_DATA
 			else
 			{
@@ -250,6 +255,11 @@ void ANavigationTestingActor::PostEditMove(bool bFinished)
 		{
 			ClosestWallLocation = FindClosestWallLocation();
 		}
+
+		if (bDrawIfNavDataIsReadyInRadius)
+		{
+			bNavDataIsReadyInRadius = CheckIfNavDataIsReadyInRadius();
+		}
 	}
 }
 
@@ -394,6 +404,27 @@ FVector ANavigationTestingActor::FindClosestWallLocation() const
 	return FNavigationSystem::InvalidLocation;
 }
 
+bool ANavigationTestingActor::CheckIfNavDataIsReadyInRadius()
+{
+#if WITH_EDITORONLY_DATA
+	if (EdRenderComp)
+	{
+		EdRenderComp->MarkRenderStateDirty();
+	}
+#endif // WITH_EDITORONLY_DATA
+	
+#if WITH_RECAST
+	UpdateNavData();
+	const ARecastNavMesh* RecastNavMesh = Cast<ARecastNavMesh>(MyNavData);
+	if (RecastNavMesh)
+	{
+		return RecastNavMesh->HasCompleteDataInRadius(GetActorLocation(), RadiusUsedToValidateNavData);
+	}
+#endif // WITH_RECAST
+	
+	return false;
+}
+
 void ANavigationTestingActor::SearchPathTo(ANavigationTestingActor* Goal)
 {
 #if WITH_EDITORONLY_DATA
diff --git a/Engine/Source/Runtime/NavigationSystem/Public/NavigationTestingActor.h b/Engine/Source/Runtime/NavigationSystem/Public/NavigationTestingActor.h
--- a/Engine/Source/Runtime/NavigationSystem/Public/NavigationTestingActor.h
+++ b/Engine/Source/Runtime/NavigationSystem/Public/NavigationTestingActor.h
@@ -116,6 +116,10 @@ public:
 	UPROPERTY(EditAnywhere, Category=Query)
 	uint32 bDrawDistanceToWall : 1;
 
+	/** If set, a cylinder is drawn to indicate if the navigation data is ready (has been generated) for the given radius (green when ready, red otherwise). */
+	UPROPERTY(EditAnywhere, Category=Query)
+	uint32 bDrawIfNavDataIsReadyInRadius : 1;
+	
 	/** show polys from open (orange) and closed (yellow) sets */
 	UPROPERTY(EditAnywhere, Category=Debug)
 	uint32 bShowNodePool : 1;
@@ -131,6 +135,10 @@ public:
 	UPROPERTY(EditAnywhere, Category=Debug)
 	uint32 bShouldBeVisibleInGame : 1;
 
+	/** NavData must be ready for all tiles within radius. When using 0, NavData must be ready at the actor location. */
+	UPROPERTY(EditAnywhere, Category=Query)
+	float RadiusUsedToValidateNavData = 0;
+	
 	/** determines which cost will be shown*/
 	UPROPERTY(EditAnywhere, Category=Debug)
 	TEnumAsByte<ENavCostDisplay::Type> CostDisplayMode;
@@ -172,6 +180,8 @@ public:
 	float OffsetFromCornersDistance;
 
 	FVector ClosestWallLocation;
+	
+	bool bNavDataIsReadyInRadius;
 
 #if WITH_RECAST && WITH_EDITORONLY_DATA
 	/** detail data gathered from each step of regular A* algorithm */
@@ -211,7 +221,6 @@ public:
 
 	NAVIGATIONSYSTEM_API void UpdateNavData();
 	NAVIGATIONSYSTEM_API void UpdatePathfinding();
-	NAVIGATIONSYSTEM_API void GatherDetailedData(ANavigationTestingActor* Goal);
 	NAVIGATIONSYSTEM_API virtual void SearchPathTo(ANavigationTestingActor* Goal);
 
 	/*	Called when given path becomes invalid (via @see PathObserverDelegate)
@@ -233,4 +242,5 @@ public:
 
 protected:
 	NAVIGATIONSYSTEM_API FVector FindClosestWallLocation() const;
+	bool CheckIfNavDataIsReadyInRadius();
 };
