Add custom link owner's name to link ID warning to allow identifying the problematic actor.

diff --git a/Engine/Source/Runtime/NavigationSystem/Private/NavigationSystem.cpp b/Engine/Source/Runtime/NavigationSystem/Private/NavigationSystem.cpp
--- a/Engine/Source/Runtime/NavigationSystem/Private/NavigationSystem.cpp
+++ b/Engine/Source/Runtime/NavigationSystem/Private/NavigationSystem.cpp
@@ -2801,7 +2801,10 @@ void UNavigationSystemV1::RegisterCustomLink(INavLinkCustomInterface& CustomLink
 			}
 			
 			// This should be very unlikely to occur, if its causing issues we should add code to handle this being careful to account for the editor world being run as a commandlet to cook and build paths on seperate runs.
-			UE_CLOG(!bGenerateNewId, LogNavLink, Warning, TEXT("%hs navlink ID %llu is clashing with existing ID. This will not be regenerated automatically in editor although for dynamic navmesh this will be handled at run time in game. For static mesh in the editor world the INavLinkCustomInterface implementor should regenerate the ID, deleting the owning actor and or component and placing again should fix this."), __FUNCTION__, CustomLink.GetId().GetId());
+			UE_CLOG(!bGenerateNewId, LogNavLink, Warning, TEXT("%hs navlink ID %llu is clashing with existing ID (Owner: %s). "
+				"This will not be regenerated automatically in editor although for dynamic navmesh this will be handled at run time in game. "
+				"For static mesh in the editor world the INavLinkCustomInterface implementor should regenerate the ID, "
+				"deleting the owning actor and or component and placing again should fix this."), __FUNCTION__, CustomLink.GetId().GetId(), *GetFullNameSafe(CustomLink.GetLinkOwner()));
 		}
 		else
 		{
@@ -2825,7 +2828,8 @@ void UNavigationSystemV1::RegisterCustomLink(INavLinkCustomInterface& CustomLink
 		}
 	}
 
-	UE_CLOG(bGenerateNewId && CustomNavLinksMap.Contains(CustomLink.GetId()), LogNavLink, Warning, TEXT("%hs New navlink ID %llu is clashing with existing ID."), __FUNCTION__, CustomLink.GetId().GetId());
+	UE_CLOG(bGenerateNewId && CustomNavLinksMap.Contains(CustomLink.GetId()), LogNavLink, Warning, TEXT("%hs New navlink ID %llu is clashing with existing ID (Owner: %s)."),
+		__FUNCTION__, CustomLink.GetId().GetId(), *GetFullNameSafe(CustomLink.GetLinkOwner()));
 	CustomNavLinksMap.Add(CustomLink.GetId(), FNavigationSystem::FCustomLinkOwnerInfo(&CustomLink));
 }
 
