Use an existing adaptive tessellation algorithm in BezierUtilities to subdivide the spline in SplineNavModifierComponent. This algorithm adapts to the curvature of the spline so that curvier spots get broken down into more linear segments for higher fidelity

#changelist validated
#virtualized

diff --git a/Engine/Plugins/Runtime/ZoneGraph/Source/ZoneGraph/Private/BezierUtilities.cpp b/Engine/Plugins/Runtime/ZoneGraph/Source/ZoneGraph/Private/BezierUtilities.cpp
deleted file mode 100644
diff --git a/Engine/Plugins/Runtime/ZoneGraph/Source/ZoneGraph/Public/BezierUtilities.h b/Engine/Plugins/Runtime/ZoneGraph/Source/ZoneGraph/Public/BezierUtilities.h
deleted file mode 100644
diff --git a/Engine/Plugins/Runtime/ZoneGraph/Source/ZoneGraph/ZoneGraph.Build.cs b/Engine/Plugins/Runtime/ZoneGraph/Source/ZoneGraph/ZoneGraph.Build.cs
--- a/Engine/Plugins/Runtime/ZoneGraph/Source/ZoneGraph/ZoneGraph.Build.cs
+++ b/Engine/Plugins/Runtime/ZoneGraph/Source/ZoneGraph/ZoneGraph.Build.cs
@@ -21,6 +21,8 @@ namespace UnrealBuildTool.Rules
 				}
 			);
 
+			PrivateDependencyModuleNames.Add("AIModule");
+
 			if (Target.bBuildEditor == true)
 			{
 				PrivateDependencyModuleNames.Add("EditorFramework");
diff --git a/Engine/Plugins/Runtime/ZoneGraph/Source/ZoneGraphEditor/ZoneGraphEditor.Build.cs b/Engine/Plugins/Runtime/ZoneGraph/Source/ZoneGraphEditor/ZoneGraphEditor.Build.cs
--- a/Engine/Plugins/Runtime/ZoneGraph/Source/ZoneGraphEditor/ZoneGraphEditor.Build.cs
+++ b/Engine/Plugins/Runtime/ZoneGraph/Source/ZoneGraphEditor/ZoneGraphEditor.Build.cs
@@ -33,6 +33,7 @@ namespace UnrealBuildTool.Rules
 
 			PrivateDependencyModuleNames.AddRange(
 			new string[] {
+				"AIModule",
 				"RenderCore",
 				"GraphEditor",
 				"KismetWidgets",
diff --git a/Engine/Source/Runtime/AIModule/Private/BezierUtilities.cpp b/Engine/Source/Runtime/AIModule/Private/BezierUtilities.cpp
new file mode 100644
--- /dev/null
+++ b/Engine/Source/Runtime/AIModule/Private/BezierUtilities.cpp
@@ -0,0 +1,210 @@
+// Copyright Epic Games, Inc. All Rights Reserved.
+
+#include "BezierUtilities.h"
+
+namespace UE
+{
+namespace CubicBezier
+{
+
+void SplitAt(const FVector& P0, const FVector& P1, const FVector& P2, const FVector& P3, const float t, FVector OutResult[5])
+{
+	// Split using De Casteljau algorithm.
+	const FVector P01 = FMath::Lerp(P0, P1, t);
+	const FVector P12 = FMath::Lerp(P1, P2, t);
+	const FVector P23 = FMath::Lerp(P2, P3, t);
+
+	OutResult[0] = P01;
+	OutResult[1] = FMath::Lerp(P01, P12, t);
+
+	OutResult[3] = FMath::Lerp(P12, P23, t);
+	OutResult[4] = P23;
+
+	OutResult[2] = FMath::Lerp(OutResult[1], OutResult[3], t);
+}
+
+FBox CalcBounds(const FVector& P0, const FVector& P1, const FVector& P2, const FVector& P3)
+{
+	FBox Result(EForceInit::ForceInit);
+
+	// Start the bounding box by end points
+	Result += P0;
+	Result += P3;
+
+	// Bezier curve fits inside the convex hull of it's control points.
+	// If control points are inside the bounds, we're done.
+	if (Result.IsInsideOrOn(P1) && Result.IsInsideOrOn(P2))
+		return Result;
+
+	// Add bezier curve inflection points in X, Y, and Z.
+	const FVector A = -3.0f * P0 + 9.0f * P1 - 9.0f * P2 + 3.0f * P3;
+	const FVector B = 6.0f * P0 - 12.0f * P1 + 6.0f * P2;
+	const FVector C = 3.0f * P1 - 3.0f * P0;
+
+	for (int32 Axis = 0; Axis < 3; Axis++)
+	{
+		float Roots[2];
+		int32 RootNum = 0;
+		if (FMath::IsNearlyZero(A[Axis]))
+		{
+			// It's linear equation.
+			if (!FMath::IsNearlyZero(B[Axis]))
+			{
+				Roots[0] = -C[Axis] / B[Axis];
+				RootNum = 1;
+			}
+		}
+		else
+		{
+			// Solve quadratic roots.
+			float Disc = B[Axis] * B[Axis] - 4.0f * C[Axis] * A[Axis];
+			if (Disc > KINDA_SMALL_NUMBER)
+			{
+				Disc = FMath::Sqrt(Disc);
+				const float Denom = 1.0f / (2.0f * A[Axis]);
+				Roots[0] = (-B[Axis] + Disc) * Denom;
+				Roots[1] = (-B[Axis] - Disc) * Denom;
+				RootNum = 2;
+			}
+		}
+		for (int32 i = 0; i < RootNum; i++)
+		{
+			// Accept only roots which are inside the bounds of the curve (also excluding enpoints which are already added above).
+			const float t = Roots[i];
+			if (t > KINDA_SMALL_NUMBER && t < (1.0f - KINDA_SMALL_NUMBER))
+			{
+				Result += Eval(P0, P1, P2, P3, t);
+			}
+		}
+	}
+
+	return Result;
+}
+
+void ClosestPointApproximate(const FVector& FromPoint, const FVector& P0, const FVector& P1, const FVector& P2, const FVector& P3, FVector& OutClosestPoint, float& OutClosestT, const int Steps)
+{
+	float ClosestDistanceSqr = TNumericLimits<float>::Max();
+	FVector ClosestPoint = P0;
+	float ClosestT = 0.0f;
+
+	const int NumSteps = FMath::Max(1, Steps);
+	const float Step = 1.0f / (float)NumSteps;
+
+	FVector PrevPoint = P0;
+
+	for (int j = 0; j < Steps; j++)
+	{
+		const float t = (float)(j + 1) * Step;
+		const FVector CurrPoint = Eval(P0, P1, P2, P3, t);
+
+		const FVector Closest = FMath::ClosestPointOnSegment(FromPoint, PrevPoint, CurrPoint);
+
+		const float DistanceSqr = FVector::DistSquared(Closest, FromPoint);
+		if (DistanceSqr < ClosestDistanceSqr)
+		{
+			ClosestDistanceSqr = DistanceSqr;
+			ClosestPoint = Closest;
+			// The SegT is calculate by the closest point calculation above, but it does not return in!
+			const float SegLengthSqr = FVector::DistSquared(PrevPoint, CurrPoint);
+			const float SegT = SegLengthSqr > KINDA_SMALL_NUMBER ? FVector::DistSquared(PrevPoint, Closest) / SegLengthSqr : 0.0f;
+			ClosestT = (t - Step) + SegT * Step;
+		}
+
+		PrevPoint = CurrPoint;
+	}
+
+	OutClosestPoint = ClosestPoint;
+	OutClosestT = ClosestT;
+}
+
+void SegmentClosestPointApproximate(const FVector& SegStart, const FVector& SegEnd, const FVector& P0, const FVector& P1, const FVector& P2, const FVector& P3, FVector& OutClosestPoint, float& OutClosestT, const int Steps)
+{
+	float ClosestDistanceSqr = TNumericLimits<float>::Max();
+	FVector ClosestPoint = P0;
+	float ClosestT = 0.0f;
+
+	const int NumSteps = FMath::Max(1, Steps);
+	const float Step = 1.0f / (float)NumSteps;
+
+	FVector PrevPoint = P0;
+
+	for (int j = 0; j < NumSteps; j++)
+	{
+		const float t = (float)(j + 1) * Step;
+		const FVector CurrPoint = Eval(P0, P1, P2, P3, t);
+
+		FVector Closest;
+		FVector SegClosest;
+		FMath::SegmentDistToSegmentSafe(PrevPoint, CurrPoint, SegStart, SegEnd, Closest, SegClosest);
+
+		const float DistanceSqr = FVector::DistSquared(Closest, SegClosest);
+		if (DistanceSqr < ClosestDistanceSqr)
+		{
+			ClosestDistanceSqr = DistanceSqr;
+			ClosestPoint = Closest;
+			// The SegT is calculate by the closest point calculation above, but it does not return in!
+			const float SegLengthSqr = FVector::DistSquared(PrevPoint, CurrPoint);
+			const float SegT = SegLengthSqr > KINDA_SMALL_NUMBER ? FVector::DistSquared(PrevPoint, Closest) / SegLengthSqr : 0.0f;
+			ClosestT = (t - Step) + SegT * Step;
+		}
+
+		PrevPoint = CurrPoint;
+	}
+
+	OutClosestPoint = ClosestPoint;
+	OutClosestT = ClosestT;
+}
+
+float ArcLengthApproximate(const FVector& P0, const FVector& P1, const FVector& P2, const FVector& P3)
+{
+	FVector Result[5];
+	SplitAt(P0, P1, P2, P3, 0.5f, Result);
+	return FVector::Dist(P0, Result[0]) + FVector::Dist(Result[1], Result[2]) + FVector::Dist(Result[2], Result[3]) + FVector::Dist(Result[3], Result[4]) + FVector::Dist(Result[4], P3);
+}
+
+void TessellateRecursive(TArray<FVector>& Output, const FVector& P0, const FVector& P1, const FVector& P2, const FVector& P3, const float ToleranceSqr, int Level, const int MaxLevel)
+{
+	// Handle degenerate segment.
+	FVector Dir = P3 - P0;
+	if (Dir.IsNearlyZero())
+	{
+		Output.Add(P3);
+		return;
+	}
+
+	// If the control points are close enough to approximate a line within tolerance, stop recursing.
+	Dir = Dir.GetUnsafeNormal();
+	const FVector RelP1 = P1 - P0;
+	const FVector RelP2 = P2 - P0;
+	const FVector ProjP1 = Dir * FVector::DotProduct(Dir, RelP1);
+	const FVector ProjP2 = Dir * FVector::DotProduct(Dir, RelP2);
+	const float DistP1Sqr = FVector::DistSquared(RelP1, ProjP1);
+	const float DistP2Sqr = FVector::DistSquared(RelP2, ProjP2);
+	if (DistP1Sqr < ToleranceSqr && DistP2Sqr < ToleranceSqr)
+	{
+		Output.Add(P3);
+		return;
+	}
+
+	if (Level < MaxLevel)
+	{
+		// Split the curve in half and recurse.
+		const FVector P01 = FMath::Lerp(P0, P1, 0.5f);
+		const FVector P12 = FMath::Lerp(P1, P2, 0.5f);
+		const FVector P23 = FMath::Lerp(P2, P3, 0.5f);
+		const FVector P012 = FMath::Lerp(P01, P12, 0.5f);
+		const FVector P123 = FMath::Lerp(P12, P23, 0.5f);
+		const FVector P0123 = FMath::Lerp(P012, P123, 0.5f);
+
+		TessellateRecursive(Output, P0, P01, P012, P0123, ToleranceSqr, Level + 1, MaxLevel);
+		TessellateRecursive(Output, P0123, P123, P23, P3, ToleranceSqr, Level + 1, MaxLevel);
+	}
+}
+
+void Tessellate(TArray<FVector>& Output, const FVector& P0, const FVector& P1, const FVector& P2, const FVector& P3, const float Tolerance, const int MaxLevel)
+{
+	TessellateRecursive(Output, P0, P1, P2, P3, Tolerance * Tolerance, 0, MaxLevel);
+}
+
+} // namespace CubicBezier
+} // namespace UE
diff --git a/Engine/Source/Runtime/AIModule/Public/BezierUtilities.h b/Engine/Source/Runtime/AIModule/Public/BezierUtilities.h
new file mode 100644
--- /dev/null
+++ b/Engine/Source/Runtime/AIModule/Public/BezierUtilities.h
@@ -0,0 +1,57 @@
+// Copyright Epic Games, Inc. All Rights Reserved.
+
+#pragma once
+
+#include "Math/Vector.h"
+#include "Math/Box.h"
+#include "Math/UnrealMathUtility.h"
+
+namespace UE
+{
+namespace CubicBezier
+{
+
+/** Evaluates cubic bezier value at specified time */
+inline FVector Eval(const FVector& P0, const FVector& P1, const FVector& P2, const FVector& P3, const float t)
+{
+	// Interpolate using De Casteljau algorithm.
+	const FVector P01 = FMath::Lerp(P0, P1, t);
+	const FVector P12 = FMath::Lerp(P1, P2, t);
+	const FVector P23 = FMath::Lerp(P2, P3, t);
+	const FVector P012 = FMath::Lerp(P01, P12, t);
+	const FVector P123 = FMath::Lerp(P12, P23, t);
+	return FMath::Lerp(P012, P123, t);
+}
+
+/** Evaluates cubic bezier derivative at specified time */
+inline FVector EvalDerivate(const FVector& P0, const FVector& P1, const FVector& P2, const FVector& P3, const float t)
+{
+	const FVector P01 = FMath::Lerp(P0, P1, t);
+	const FVector P12 = FMath::Lerp(P1, P2, t);
+	const FVector P23 = FMath::Lerp(P2, P3, t);
+	return 3.0f * (FMath::Lerp(P12, P23, t) - FMath::Lerp(P01, P12, t));
+}
+
+/** Splits cubic bezier curve at specified time, returns new points.
+  *	The new curves are (P0, Result[0], Result[1], Result[2]) and (P0, Result[2], Result[3], Result[4], P3) */
+AIMODULE_API void SplitAt(const FVector& P0, const FVector& P1, const FVector& P2, const FVector& P3, const float t, FVector OutResult[5]);
+
+/** Returns bounding box of cubic bezier curve. */
+AIMODULE_API FBox CalcBounds(const FVector& P0, const FVector& P1, const FVector& P2, const FVector& P3);
+
+/** Returns closest point to cubic bezier segment. The curve is split to up to Steps linear segments for approximate calculation. */
+AIMODULE_API void ClosestPointApproximate(const FVector& FromPoint, const FVector& P0, const FVector& P1, const FVector& P2, const FVector& P3, FVector& OutClosestPoint, float& OutClosestT, const int Steps = 16);
+
+/** Returns closest point between a linear segment and cubic bezier segment. The curve is split to up to Steps linear segments for approximate calculation. */
+AIMODULE_API void SegmentClosestPointApproximate(const FVector& SegStart, const FVector& SegEnd, const FVector& P0, const FVector& P1, const FVector& P2, const FVector& P3, FVector& OutClosestPoint, float& OutClosestT, const int Steps = 16);
+
+/** Returns approximate arc length of cubic bezier. The length is calculated based on control point hull with one subdivision. */
+AIMODULE_API float ArcLengthApproximate(const FVector& P0, const FVector& P1, const FVector& P2, const FVector& P3);
+
+/** Recursively tessellate a bezier curve until the segments are within Tolerance from the curve, or until max recursion level is reached.
+  * Output will contain points describing linear segments of the flattened curve.
+  * Note that the points are appended, and that the first point is omitted so that the function can be used to flatten a bezier shape. */
+AIMODULE_API void Tessellate(TArray<FVector>& Output, const FVector& P0, const FVector& P1, const FVector& P2, const FVector& P3, const float Tolerance, const int MaxLevel = 6);
+
+} // namespace CubicBezier
+} // namespace UE
diff --git a/Engine/Source/Runtime/Engine/Private/AI/Navigation/NavigationModifier.cpp b/Engine/Source/Runtime/Engine/Private/AI/Navigation/NavigationModifier.cpp
--- a/Engine/Source/Runtime/Engine/Private/AI/Navigation/NavigationModifier.cpp
+++ b/Engine/Source/Runtime/Engine/Private/AI/Navigation/NavigationModifier.cpp
@@ -339,6 +339,12 @@ FAreaNavModifier::FAreaNavModifier(const TArray<FVector>& InPoints, ENavigationC
 	SetConvex(InPoints.GetData(), 0, InPoints.Num(), CoordType, LocalToWorld);
 }
 
+FAreaNavModifier::FAreaNavModifier(const TConstArrayView<FVector> InPoints, ENavigationCoordSystem::Type CoordType, const FTransform& LocalToWorld, const TSubclassOf<UNavAreaBase> InAreaClass)
+{
+	Init(InAreaClass);
+	SetConvex(InPoints.GetData(), 0, InPoints.Num(), CoordType, LocalToWorld);
+}
+
 FAreaNavModifier::FAreaNavModifier(const TArray<FVector>& InPoints, const int32 FirstIndex, const int32 LastIndex, ENavigationCoordSystem::Type CoordType, const FTransform& LocalToWorld, const TSubclassOf<UNavAreaBase> InAreaClass)
 {
 	check(InPoints.IsValidIndex(FirstIndex) && InPoints.IsValidIndex(LastIndex-1));
diff --git a/Engine/Source/Runtime/Engine/Public/AI/NavigationModifier.h b/Engine/Source/Runtime/Engine/Public/AI/NavigationModifier.h
--- a/Engine/Source/Runtime/Engine/Public/AI/NavigationModifier.h
+++ b/Engine/Source/Runtime/Engine/Public/AI/NavigationModifier.h
@@ -104,6 +104,7 @@ struct FAreaNavModifier : public FNavigationModifier
 	ENGINE_API FAreaNavModifier(const FVector& Extent, const FTransform& LocalToWorld, const TSubclassOf<UNavAreaBase> AreaClass);
 	ENGINE_API FAreaNavModifier(const FBox& Box, const FTransform& LocalToWorld, const TSubclassOf<UNavAreaBase> AreaClass);
 	ENGINE_API FAreaNavModifier(const TArray<FVector>& Points, ENavigationCoordSystem::Type CoordType, const FTransform& LocalToWorld, const TSubclassOf<UNavAreaBase> AreaClass);
+	ENGINE_API FAreaNavModifier(const TConstArrayView<FVector> Points, ENavigationCoordSystem::Type CoordType, const FTransform& LocalToWorld, const TSubclassOf<UNavAreaBase> AreaClass);
 	ENGINE_API FAreaNavModifier(const TArray<FVector>& Points, const int32 FirstIndex, const int32 LastIndex, ENavigationCoordSystem::Type CoordType, const FTransform& LocalToWorld, const TSubclassOf<UNavAreaBase> AreaClass);
 	ENGINE_API FAreaNavModifier(const TNavStatArray<FVector>& Points, const int32 FirstIndex, const int32 LastIndex, ENavigationCoordSystem::Type CoordType, const FTransform& LocalToWorld, const TSubclassOf<UNavAreaBase> AreaClass);
 	UE_DEPRECATED(5.0, "FAreaNavModifier constructor with a UBrushComponent* parameter has been deprecated since it wasn't able to handle concave shapes. Use FCompositeNavModifier::CreateAreaModifiers instead")
diff --git a/Engine/Source/Runtime/NavigationSystem/NavigationSystem.Build.cs b/Engine/Source/Runtime/NavigationSystem/NavigationSystem.Build.cs
--- a/Engine/Source/Runtime/NavigationSystem/NavigationSystem.Build.cs
+++ b/Engine/Source/Runtime/NavigationSystem/NavigationSystem.Build.cs
@@ -28,6 +28,7 @@ namespace UnrealBuildTool.Rules
 
 			PrivateDependencyModuleNames.AddRange(
 				new string[] {
+					"AIModule",
 					"RHI",
 					"RenderCore",
 				}
diff --git a/Engine/Source/Runtime/NavigationSystem/Private/SplineNavModifierComponent.cpp b/Engine/Source/Runtime/NavigationSystem/Private/SplineNavModifierComponent.cpp
--- a/Engine/Source/Runtime/NavigationSystem/Private/SplineNavModifierComponent.cpp
+++ b/Engine/Source/Runtime/NavigationSystem/Private/SplineNavModifierComponent.cpp
@@ -4,6 +4,7 @@
 
 #include "AI/NavigationSystemBase.h"
 #include "AI/Navigation/NavigationRelevantData.h"
+#include "BezierUtilities.h"
 #include "Components/SplineComponent.h"
 #include "VisualLogger/VisualLogger.h"
 
@@ -11,6 +12,7 @@
 
 namespace
 {
+	// Fetch the spline component from the actor
 	const USplineComponent* GetSpline(const AActor* Owner)
 	{
 		if (!Owner)
@@ -24,14 +26,50 @@ namespace
 
 		return Spline;
 	}
+
+	// Subdivide the spline into linear segments, adapting to its curvature (more curvy means more linear segments)
+	void SubdivideSpline(TArray<FVector>& OutSubdivisions, const USplineComponent& Spline, const float SubdivisionThreshold)
+	{
+		// Sample at least 2 points
+		const int32 NumSplinePoints = FMath::Max(Spline.GetNumberOfSplinePoints(), 2);
+
+		// The USplineComponent's Hermite spline tangents are 3 times larger than Bezier tangents and we need to convert before tessellation
+		constexpr double HermiteToBezierFactor = 3.0;
+
+		FSplinePoint PrevSplinePoint;
+		for (int32 PointIndex = 0; PointIndex < NumSplinePoints; PointIndex++)
+		{
+			FSplinePoint CurrSplinePoint = Spline.GetSplinePointAt(PointIndex, ESplineCoordinateSpace::World);
+
+			if (PointIndex > 0)
+			{
+				// The first point of the segment is appended before tessellation since UE::CubicBezier::Tessellate does not add it
+				OutSubdivisions.Add(PrevSplinePoint.Position);
+
+				// Convert this segment of the spline from Hermite to Bezier and subdivide it 
+				UE::CubicBezier::Tessellate(OutSubdivisions,
+					PrevSplinePoint.Position,
+					PrevSplinePoint.Position + PrevSplinePoint.LeaveTangent / HermiteToBezierFactor,
+					CurrSplinePoint.Position - CurrSplinePoint.ArriveTangent / HermiteToBezierFactor,
+					CurrSplinePoint.Position,
+					SubdivisionThreshold);
+			}
+
+			PrevSplinePoint = CurrSplinePoint;
+		}
+	}
 }
 
 void USplineNavModifierComponent::CalculateBounds() const
 {
-	if (const USplineComponent* Spline = GetSpline(GetOwner()))
+	const USplineComponent* Spline = GetSpline(GetOwner());
+	if (!Spline)
 	{
-		Bounds = Spline->CalcBounds(GetOwner()->GetTransform()).GetBox();
+		return;
 	}
+
+	const double Buffer = FMath::Max(StrokeWidth / 2.0, StrokeHeight / 2.0);
+	Bounds = Spline->CalcBounds(Spline->GetComponentTransform()).GetBox().ExpandBy(Buffer);
 }
 
 void USplineNavModifierComponent::GetNavigationData(FNavigationRelevantData& Data) const
@@ -42,45 +80,64 @@ void USplineNavModifierComponent::GetNavigationData(FNavigationRelevantData& Dat
 		return;
 	}
 
-	// Square in the YZ plane used to sample the spline at each cross section
-	constexpr int32 SampleVertices = 4;
-	TStaticArray<FVector, SampleVertices> SampleSquare;
-	SampleSquare[0] = FVector(0, -SplineExtent, -SplineExtent);
-	SampleSquare[1] = FVector(0,  SplineExtent, -SplineExtent);
-	SampleSquare[2] = FVector(0,  SplineExtent,  SplineExtent);
-	SampleSquare[3] = FVector(0, -SplineExtent,  SplineExtent);
-
-	// Vertices (in no particular order) of a prism which will enclose each segment of the spline
-	// @Note The only reason Tube isn't a TStaticArray is because FAreaNavModifier expects a TArray
-	TArray<FVector> Tube;
-	Tube.SetNum(SampleSquare.Num() * 2);
-
-	// Always sample at least the start and end points
-	const int32 NumPoints = FMath::Max(NumSplineSamples, 2);
-	
-	const float SplineLength = Spline->GetSplineLength();
-	FTransform PreviousTransform;
-
-	for (int32 SampleIndex = 0; SampleIndex < NumPoints; SampleIndex++)
-	{
-		// Sample a point on the spline at the current distance
-		const double Distance = SplineLength * SampleIndex / (NumPoints - 1);
-		const FTransform CurrentTransform = Spline->GetTransformAtDistanceAlongSpline(Distance, ESplineCoordinateSpace::World);
+	// Build a rectangle in the YZ plane used to sample the spline at each cross section
+	constexpr int32 NumCrossSectionVertices = 4;
+	const double StrokeHalfWidth = StrokeWidth / 2.0;
+	const double StrokeHalfHeight = StrokeHeight / 2.0;
+	TStaticArray<FVector, NumCrossSectionVertices> CrossSectionRect;
+	CrossSectionRect[0] = FVector(0.0, -StrokeHalfWidth, -StrokeHalfHeight);
+	CrossSectionRect[1] = FVector(0.0,  StrokeHalfWidth, -StrokeHalfHeight);
+	CrossSectionRect[2] = FVector(0.0,  StrokeHalfWidth,  StrokeHalfHeight);
+	CrossSectionRect[3] = FVector(0.0, -StrokeHalfWidth,  StrokeHalfHeight);
+
+	// Vertices (in an arbitrary order) of a prism which will enclose each segment of the spline
+	TStaticArray<FVector, NumCrossSectionVertices * 2> Tube;
 
-		if (SampleIndex > 0)
+	// Subdivide the spline so that high curvature sections get smaller and more linear segments than straighter sections
+	TArray<FVector> Subdivisions;
+	SubdivideSpline(Subdivisions, *Spline, GetSudivisionThreshold());
+	const int32 NumSubdivisions = Subdivisions.Num();
+
+	// Create volumes from the spline subdivisions and use them to mark the nav mesh with the given are
+	const FTransform ComponentTransform = Spline->GetComponentTransform();
+	int32 PrevIndex = Spline->IsClosedLoop() ? (NumSubdivisions - 1) : INDEX_NONE;
+	for (int32 SubdivisionIndex = 0; SubdivisionIndex < NumSubdivisions; SubdivisionIndex++)
+	{
+		if (SubdivisionIndex > 0)
 		{
-			// Compute the vertices of the current tube segment
-			for (int i = 0; i < SampleVertices; i++)
+			// Compute the rotation of this tube segment
+			const double TubeAngle = (Subdivisions[SubdivisionIndex] - Subdivisions[PrevIndex]).HeadingAngle();
+			const FQuat TubeRotation(FVector::UnitZ(), TubeAngle);
+			
+			// Compute the vertices of this tube segment
+			for (int i = 0; i < NumCrossSectionVertices; i++)
 			{
-				Tube[i] = PreviousTransform.TransformPosition(SampleSquare[i]);
-				Tube[i + SampleVertices] = CurrentTransform.TransformPosition(SampleSquare[i]);
+				// For each vertex of the tube segment, first rotate about the positive Z axis, then translate to the subdivision point
+				Tube[i] = (TubeRotation * CrossSectionRect[i]) + Subdivisions[PrevIndex];
+				Tube[i + NumCrossSectionVertices] = (TubeRotation * CrossSectionRect[i]) + Subdivisions[SubdivisionIndex];
 			}
 
 			// From the tube construct a convex hull whose volume will be used to mark the nav mesh with the selected AreaClass
-			const FAreaNavModifier NavModifier(Tube, ENavigationCoordSystem::Type::Unreal, FTransform::Identity, AreaClass);
+			const FAreaNavModifier NavModifier(Tube, ENavigationCoordSystem::Type::Unreal, ComponentTransform, AreaClass);
 			Data.Modifiers.Add(NavModifier);
+
+			PrevIndex = SubdivisionIndex;
 		}
-		
-		PreviousTransform = CurrentTransform;
 	}
-}
\ No newline at end of file
+}
+
+float USplineNavModifierComponent::GetSudivisionThreshold() const
+{
+	switch (SubdivisionLOD)
+	{
+	case ESubdivisionLOD::Ultra:
+		return 10.0f;
+	case ESubdivisionLOD::High:
+		return 100.0f;
+	case ESubdivisionLOD::Medium:
+		return 250.0f;
+	case ESubdivisionLOD::Low:
+	default: // Fallthrough
+		return 500.0f;
+	}
+}
diff --git a/Engine/Source/Runtime/NavigationSystem/Public/SplineNavModifierComponent.h b/Engine/Source/Runtime/NavigationSystem/Public/SplineNavModifierComponent.h
--- a/Engine/Source/Runtime/NavigationSystem/Public/SplineNavModifierComponent.h
+++ b/Engine/Source/Runtime/NavigationSystem/Public/SplineNavModifierComponent.h
@@ -7,6 +7,15 @@
 
 struct FNavigationRelevantData;
 
+UENUM()
+enum class ESubdivisionLOD
+{
+	Low,
+	Medium,
+	High,
+	Ultra,
+};
+
 /**
  *	Assign a chosen NavArea in the vicinity of a spline
  *	This component must only be attached to an actor with a USplineComponent
@@ -16,14 +25,21 @@ class USplineNavModifierComponent : public UNavModifierComponent
 {
 	GENERATED_BODY()
 
-	/** Radius of the tube surrounding the spline which will be used to generate the Nav Modifiers */
-	UPROPERTY(EditAnywhere, Category = Navigation, Meta=(UIMin="1", ClampMin="1"))
-	double SplineExtent = 1000.0f;
-
-	/** How many sections the spline will be divided into for generating overlap volumes. More samples result in finer detail */
-	UPROPERTY(EditAnywhere, Category = Navigation, Meta=(UIMin="2", ClampMin="2"))
-	int32 NumSplineSamples = 100;
+	/** Width of the generated tube-shaped NavModifierVolume enclosing the spline */
+	UPROPERTY(EditAnywhere, Category = Navigation, Meta=(UIMin="10", ClampMin="10"))
+	double StrokeWidth = 500.0f;
+	
+	/** Height of the generated tube-shaped nav modifier volume enclosing the spline, must be tall enough to pass through the desired surface beneath the spline */
+	UPROPERTY(EditAnywhere, Category = Navigation, Meta=(UIMin="10", ClampMin="10"))
+	double StrokeHeight = 500.0f;
+		
+	/** Higher LOD will capture finer details in the spline */
+	UPROPERTY(EditAnywhere, Category = Navigation)
+	ESubdivisionLOD SubdivisionLOD = ESubdivisionLOD::Medium;
 	
 	virtual void CalculateBounds() const override;
 	virtual void GetNavigationData(FNavigationRelevantData& Data) const override;
+
+protected:
+	float GetSudivisionThreshold() const;
 };
