Temporary fix to circular dependency caused by 31581755

diff --git a/Engine/Source/Runtime/NavigationSystem/NavigationSystem.Build.cs b/Engine/Source/Runtime/NavigationSystem/NavigationSystem.Build.cs
--- a/Engine/Source/Runtime/NavigationSystem/NavigationSystem.Build.cs
+++ b/Engine/Source/Runtime/NavigationSystem/NavigationSystem.Build.cs
@@ -28,7 +28,6 @@ namespace UnrealBuildTool.Rules
 
 			PrivateDependencyModuleNames.AddRange(
 				new string[] {
-					"AIModule",
 					"RHI",
 					"RenderCore",
 				}
diff --git a/Engine/Source/Runtime/NavigationSystem/Private/SplineNavModifierComponent.cpp b/Engine/Source/Runtime/NavigationSystem/Private/SplineNavModifierComponent.cpp
--- a/Engine/Source/Runtime/NavigationSystem/Private/SplineNavModifierComponent.cpp
+++ b/Engine/Source/Runtime/NavigationSystem/Private/SplineNavModifierComponent.cpp
@@ -4,7 +4,7 @@
 
 #include "AI/NavigationSystemBase.h"
 #include "AI/Navigation/NavigationRelevantData.h"
-#include "BezierUtilities.h"
+// #include "BezierUtilities.h"
 #include "Components/SplineComponent.h"
 #include "VisualLogger/VisualLogger.h"
 
@@ -12,6 +12,51 @@
 
 namespace
 {
+	// Todo: temporary solution to circular dependency with BezierUtilities's new location, remove after fixing circular dependency
+	void TessellateRecursive(TArray<FVector>& Output, const FVector& P0, const FVector& P1, const FVector& P2, const FVector& P3, const float ToleranceSqr, int Level, const int MaxLevel)
+	{
+		// Handle degenerate segment.
+		FVector Dir = P3 - P0;
+		if (Dir.IsNearlyZero())
+		{
+			Output.Add(P3);
+			return;
+		}
+
+		// If the control points are close enough to approximate a line within tolerance, stop recursing.
+		Dir = Dir.GetUnsafeNormal();
+		const FVector RelP1 = P1 - P0;
+		const FVector RelP2 = P2 - P0;
+		const FVector ProjP1 = Dir * FVector::DotProduct(Dir, RelP1);
+		const FVector ProjP2 = Dir * FVector::DotProduct(Dir, RelP2);
+		const float DistP1Sqr = FVector::DistSquared(RelP1, ProjP1);
+		const float DistP2Sqr = FVector::DistSquared(RelP2, ProjP2);
+		if (DistP1Sqr < ToleranceSqr && DistP2Sqr < ToleranceSqr)
+		{
+			Output.Add(P3);
+			return;
+		}
+
+		if (Level < MaxLevel)
+		{
+			// Split the curve in half and recurse.
+			const FVector P01 = FMath::Lerp(P0, P1, 0.5f);
+			const FVector P12 = FMath::Lerp(P1, P2, 0.5f);
+			const FVector P23 = FMath::Lerp(P2, P3, 0.5f);
+			const FVector P012 = FMath::Lerp(P01, P12, 0.5f);
+			const FVector P123 = FMath::Lerp(P12, P23, 0.5f);
+			const FVector P0123 = FMath::Lerp(P012, P123, 0.5f);
+
+			TessellateRecursive(Output, P0, P01, P012, P0123, ToleranceSqr, Level + 1, MaxLevel);
+			TessellateRecursive(Output, P0123, P123, P23, P3, ToleranceSqr, Level + 1, MaxLevel);
+		}
+	}
+
+	void Tessellate(TArray<FVector>& Output, const FVector& P0, const FVector& P1, const FVector& P2, const FVector& P3, const float Tolerance, const int MaxLevel = 6)
+	{
+		TessellateRecursive(Output, P0, P1, P2, P3, Tolerance * Tolerance, 0, MaxLevel);
+	}
+
 	// Fetch the spline component from the actor
 	const USplineComponent* GetSpline(const AActor* Owner)
 	{
@@ -47,7 +92,7 @@ namespace
 				OutSubdivisions.Add(PrevSplinePoint.Position);
 
 				// Convert this segment of the spline from Hermite to Bezier and subdivide it 
-				UE::CubicBezier::Tessellate(OutSubdivisions,
+				Tessellate(OutSubdivisions,
 					PrevSplinePoint.Position,
 					PrevSplinePoint.Position + PrevSplinePoint.LeaveTangent / HermiteToBezierFactor,
 					CurrSplinePoint.Position - CurrSplinePoint.ArriveTangent / HermiteToBezierFactor,
