Fix case where the final segment of the spline wasn't tessellated if it's a closed loop. Also thoroughly tested edge cases such as loop/non-loops with 0, 1, and 2 points

diff --git a/Engine/Source/Runtime/NavigationSystem/Private/SplineNavModifierComponent.cpp b/Engine/Source/Runtime/NavigationSystem/Private/SplineNavModifierComponent.cpp
--- a/Engine/Source/Runtime/NavigationSystem/Private/SplineNavModifierComponent.cpp
+++ b/Engine/Source/Runtime/NavigationSystem/Private/SplineNavModifierComponent.cpp
@@ -36,13 +36,15 @@ namespace
 		// The USplineComponent's Hermite spline tangents are 3 times larger than Bezier tangents and we need to convert before tessellation
 		constexpr double HermiteToBezierFactor = 3.0;
 
-		FSplinePoint PrevSplinePoint;
-		for (int32 PointIndex = 0; PointIndex < NumSplinePoints; PointIndex++)
+		// Tessellate the spline segments
+		int32 PrevIndex = Spline.IsClosedLoop() ? (NumSplinePoints - 1) : INDEX_NONE;
+		for (int32 SplinePointIndex = 0; SplinePointIndex < NumSplinePoints; SplinePointIndex++)
 		{
-			FSplinePoint CurrSplinePoint = Spline.GetSplinePointAt(PointIndex, ESplineCoordinateSpace::World);
-
-			if (PointIndex > 0)
+			if (PrevIndex >= 0)
 			{
+				const FSplinePoint PrevSplinePoint = Spline.GetSplinePointAt(PrevIndex, ESplineCoordinateSpace::World);
+				const FSplinePoint CurrSplinePoint = Spline.GetSplinePointAt(SplinePointIndex, ESplineCoordinateSpace::World);
+
 				// The first point of the segment is appended before tessellation since UE::CubicBezier::Tessellate does not add it
 				OutSubdivisions.Add(PrevSplinePoint.Position);
 
@@ -55,7 +57,7 @@ namespace
 					SubdivisionThreshold);
 			}
 
-			PrevSplinePoint = CurrSplinePoint;
+			PrevIndex = SplinePointIndex;
 		}
 	}
 }
@@ -100,29 +102,26 @@ void USplineNavModifierComponent::GetNavigationData(FNavigationRelevantData& Dat
 
 	// Create volumes from the spline subdivisions and use them to mark the nav mesh with the given are
 	const FTransform ComponentTransform = Spline->GetComponentTransform();
-	int32 PrevIndex = Spline->IsClosedLoop() ? (NumSubdivisions - 1) : INDEX_NONE;
-	for (int32 SubdivisionIndex = 0; SubdivisionIndex < NumSubdivisions; SubdivisionIndex++)
+	int32 PrevIndex = 0;
+	for (int32 SubdivisionIndex = 1; SubdivisionIndex < NumSubdivisions; SubdivisionIndex++)
 	{
-		if (SubdivisionIndex > 0)
+		// Compute the rotation of this tube segment
+		const double TubeAngle = (Subdivisions[SubdivisionIndex] - Subdivisions[PrevIndex]).HeadingAngle();
+		const FQuat TubeRotation(FVector::UnitZ(), TubeAngle);
+		
+		// Compute the vertices of this tube segment
+		for (int i = 0; i < NumCrossSectionVertices; i++)
 		{
-			// Compute the rotation of this tube segment
-			const double TubeAngle = (Subdivisions[SubdivisionIndex] - Subdivisions[PrevIndex]).HeadingAngle();
-			const FQuat TubeRotation(FVector::UnitZ(), TubeAngle);
-			
-			// Compute the vertices of this tube segment
-			for (int i = 0; i < NumCrossSectionVertices; i++)
-			{
-				// For each vertex of the tube segment, first rotate about the positive Z axis, then translate to the subdivision point
-				Tube[i] = (TubeRotation * CrossSectionRect[i]) + Subdivisions[PrevIndex];
-				Tube[i + NumCrossSectionVertices] = (TubeRotation * CrossSectionRect[i]) + Subdivisions[SubdivisionIndex];
-			}
+			// For each vertex of the tube segment, first rotate about the positive Z axis, then translate to the subdivision point
+			Tube[i] = (TubeRotation * CrossSectionRect[i]) + Subdivisions[PrevIndex];
+			Tube[i + NumCrossSectionVertices] = (TubeRotation * CrossSectionRect[i]) + Subdivisions[SubdivisionIndex];
+		}
 
-			// From the tube construct a convex hull whose volume will be used to mark the nav mesh with the selected AreaClass
-			const FAreaNavModifier NavModifier(Tube, ENavigationCoordSystem::Type::Unreal, ComponentTransform, AreaClass);
-			Data.Modifiers.Add(NavModifier);
+		// From the tube construct a convex hull whose volume will be used to mark the nav mesh with the selected AreaClass
+		const FAreaNavModifier NavModifier(Tube, ENavigationCoordSystem::Type::Unreal, ComponentTransform, AreaClass);
+		Data.Modifiers.Add(NavModifier);
 
-			PrevIndex = SubdivisionIndex;
-		}
+		PrevIndex = SubdivisionIndex;
 	}
 }
 
