NavigationTestingActor: Don't reset ShowStepIndex for every request and allow to save it's value (not transient anymore).

diff --git a/Engine/Source/Runtime/NavigationSystem/Private/NavigationTestingActor.cpp b/Engine/Source/Runtime/NavigationSystem/Private/NavigationTestingActor.cpp
--- a/Engine/Source/Runtime/NavigationSystem/Private/NavigationTestingActor.cpp
+++ b/Engine/Source/Runtime/NavigationSystem/Private/NavigationTestingActor.cpp
@@ -321,7 +321,6 @@ void ANavigationTestingActor::UpdatePathfinding()
 	bPathIsPartial = false;
 	bPathExist = false;
 	LastPath.Reset();
-	ShowStepIndex = -1;
 	PathfindingSteps = 0;
 #if WITH_RECAST && WITH_EDITORONLY_DATA
 	DebugSteps.Reset();
diff --git a/Engine/Source/Runtime/NavigationSystem/Public/NavigationTestingActor.h b/Engine/Source/Runtime/NavigationSystem/Public/NavigationTestingActor.h
--- a/Engine/Source/Runtime/NavigationSystem/Public/NavigationTestingActor.h
+++ b/Engine/Source/Runtime/NavigationSystem/Public/NavigationTestingActor.h
@@ -173,7 +173,8 @@ public:
 	UPROPERTY(EditAnywhere, Category=Pathfinding)
 	TSubclassOf<class UNavigationQueryFilter> FilterClass;
 
-	UPROPERTY(transient, EditInstanceOnly, Category=Debug, meta=(ClampMin="-1", UIMin="-1"))
+	/** Show debug steps up to this index. Use -1 to disable. */
+	UPROPERTY(EditInstanceOnly, Category=Debug, meta=(ClampMin="-1", UIMin="-1"))
 	int32 ShowStepIndex;
 
 	UPROPERTY(EditAnywhere, Category=Pathfinding)
