[Navigation] Replaced custom NavLink static registration queue by a new world subsystem that will also store NavRelevantObject that are not AActor/UActorComponent based.
This change allow new NavigationSystem created after initial map load (e.g. ANavSystemConfigOverride) to properly recreate its custom nav link mappings and octree.

diff --git a/Engine/Source/Runtime/NavigationSystem/Private/NavMesh/NavMeshRenderingComponent.cpp b/Engine/Source/Runtime/NavigationSystem/Private/NavMesh/NavMeshRenderingComponent.cpp
--- a/Engine/Source/Runtime/NavigationSystem/Private/NavMesh/NavMeshRenderingComponent.cpp
+++ b/Engine/Source/Runtime/NavigationSystem/Private/NavMesh/NavMeshRenderingComponent.cpp
@@ -600,6 +600,7 @@ void FNavMeshSceneProxyData::GatherData(const ARecastNavMesh* NavMesh, int32 InN
 			{
 				DebugLabels.Add(FDebugText(FString::Printf(TEXT("NavData count: %i"), NavSys->NavDataSet.Num())));
 				DebugLabels.Add(FDebugText(FString::Printf(TEXT("MainNavData: %s"), NavSys->MainNavData ? *NavSys->MainNavData->GetName() : TEXT("none"))));
+				DebugLabels.Add(FDebugText(FString::Printf(TEXT("Custom NavLinks count: %i"), NavSys->GetNumCustomLinks())));
 
 #if WITH_NAVMESH_CLUSTER_LINKS
 				DebugLabels.Add(FDebugText(FString::Printf(TEXT("Using cluster links"))));
diff --git a/Engine/Source/Runtime/NavigationSystem/Private/NavigationObjectRepository.cpp b/Engine/Source/Runtime/NavigationSystem/Private/NavigationObjectRepository.cpp
new file mode 100644
--- /dev/null
+++ b/Engine/Source/Runtime/NavigationSystem/Private/NavigationObjectRepository.cpp
@@ -0,0 +1,47 @@
+﻿// Copyright Epic Games, Inc. All Rights Reserved.
+
+#include "NavigationObjectRepository.h"
+#include "AI/Navigation/NavRelevantInterface.h"
+#include "NavLinkCustomInterface.h"
+
+void UNavigationObjectRepository::RegisterNavRelevantObject(INavRelevantInterface& NavRelevantObject)
+{
+#if DO_ENSURE // We don't want to execute the Find at all for targets where ensures are disabled
+	if (!ensureMsgf(NavRelevantObjects.Find(&NavRelevantObject) == INDEX_NONE, TEXT("Same interface pointer can't be registered twice.")))
+	{
+		return;
+	}
+#endif
+
+	NavRelevantObjects.Emplace(&NavRelevantObject);
+
+	OnNavRelevantObjectRegistered.ExecuteIfBound(NavRelevantObject);
+}
+
+void UNavigationObjectRepository::UnregisterNavRelevantObject(INavRelevantInterface& NavRelevantObject)
+{
+	ensureMsgf(NavRelevantObjects.Remove(&NavRelevantObject) > 0, TEXT("Interface can't be removed since it was not registered or already unregistered)"));
+
+	OnNavRelevantObjectUnregistered.ExecuteIfBound(NavRelevantObject);
+}
+
+void UNavigationObjectRepository::RegisterCustomNavLinkObject(INavLinkCustomInterface& CustomNavLinkObject)
+{
+#if DO_ENSURE // We don't want to execute the Find at all for targets where ensures are disabled
+	if (!ensureMsgf(CustomLinkObjects.Find(&CustomNavLinkObject) == INDEX_NONE, TEXT("Same interface pointer can't be registered twice.")))
+	{
+		return;
+	}
+#endif
+
+	CustomLinkObjects.Emplace(&CustomNavLinkObject);
+
+	OnCustomNavLinkObjectRegistered.ExecuteIfBound(CustomNavLinkObject);
+}
+
+void UNavigationObjectRepository::UnregisterCustomNavLinkObject(INavLinkCustomInterface& CustomNavLinkObject)
+{
+	ensureMsgf(CustomLinkObjects.Remove(&CustomNavLinkObject) > 0, TEXT("Interface can't be removed since it was not registered or already unregistered)"));
+
+	OnCustomNavLinkObjectUnregistered.ExecuteIfBound(CustomNavLinkObject);
+}
diff --git a/Engine/Source/Runtime/NavigationSystem/Private/NavigationObjectRepository.h b/Engine/Source/Runtime/NavigationSystem/Private/NavigationObjectRepository.h
new file mode 100644
--- /dev/null
+++ b/Engine/Source/Runtime/NavigationSystem/Private/NavigationObjectRepository.h
@@ -0,0 +1,91 @@
+﻿// Copyright Epic Games, Inc. All Rights Reserved.
+
+#pragma once
+
+#include "Subsystems/WorldSubsystem.h"
+#include "UObject/WeakInterfacePtr.h"
+#include "NavigationObjectRepository.generated.h"
+
+class INavLinkCustomInterface;
+class INavRelevantInterface;
+
+DECLARE_DELEGATE_OneParam(FOnNavRelevantObjectRegistrationEvent, INavRelevantInterface&);
+DECLARE_DELEGATE_OneParam(FOnCustomNavLinkObjectRegistrationEvent, INavLinkCustomInterface&);
+
+/**
+ * World subsystem dedicated to store different types of navigation related objects that the
+ * NavigationSystem needs to access.
+ */
+UCLASS()
+class UNavigationObjectRepository : public UWorldSubsystem
+{
+	GENERATED_BODY()
+
+public:
+
+	/**
+	 * Adds the provided object implementing INavRelevantInterface to the list of registered navigation relevant objects.
+	 * @param NavRelevantObject INavRelevantInterface of the object to register.
+	 * @note Method will also assert if same interface pointers is registered twice.
+	 */
+	void RegisterNavRelevantObject(INavRelevantInterface& NavRelevantObject);
+
+	/**
+	 * Removes the provided interface from the list of registered navigation relevant objects.
+	 * @param NavRelevantObject INavRelevantInterface of the object to unregister.
+	 * @note Method will also assert if interface can't be removed (i.e. not registered or already unregistered).
+	 */
+	void UnregisterNavRelevantObject(INavRelevantInterface& NavRelevantObject);
+
+	/**
+	 * Returns the list of registered navigation relevant objects.
+	 * @return Const view on the list of all registered navigation relevant objects.
+	 */
+	TConstArrayView<TWeakInterfacePtr<INavRelevantInterface>> GetNavRelevantObjects() const
+	{
+		return NavRelevantObjects;
+	}
+
+	/**
+	 * Adds the provided interface to the list of registered custom navigation links.
+	 * @param CustomNavLinkObject INavLinkCustomInterface of the object to register.
+	 * @note Method will also assert if same interface pointers is registered twice.
+	 */
+	void RegisterCustomNavLinkObject(INavLinkCustomInterface& CustomNavLinkObject);
+
+	/**
+	 * Removes the provided interface from the list of registered custom navigation links.
+	 * @param CustomNavLinkObject INavLinkCustomInterface of the object to unregister.
+	 * @note Method will also assert if interface can't be removed (i.e. not registered or already unregistered).
+	 */
+	void UnregisterCustomNavLinkObject(INavLinkCustomInterface& CustomNavLinkObject);
+
+	/**
+	 * Returns the list of registered custom navigation links.
+	 * @return Const view on the list of all registered custom navigation links.
+	 */
+	TConstArrayView<TWeakInterfacePtr<INavLinkCustomInterface>> GetCustomLinks() const
+	{
+		return CustomLinkObjects;
+	}
+
+	/** Delegate executed when a navigation relevant object is registered to the repository. */
+	FOnNavRelevantObjectRegistrationEvent OnNavRelevantObjectRegistered;
+
+	/** Delegate executed when a navigation relevant object is unregistered from the repository. */
+	FOnNavRelevantObjectRegistrationEvent OnNavRelevantObjectUnregistered;
+
+	/** Delegate executed when a custom navigation link is registered to the repository. */
+	FOnCustomNavLinkObjectRegistrationEvent OnCustomNavLinkObjectRegistered;
+
+	/** Delegate executed when a custom navigation link is unregistered from the repository. */
+	FOnCustomNavLinkObjectRegistrationEvent OnCustomNavLinkObjectUnregistered;
+
+private:
+
+	/** List of registered navigation relevant objects. */
+	TArray<TWeakInterfacePtr<INavRelevantInterface>> NavRelevantObjects;
+
+	/** List of registered custom navigation link objects. */
+	TArray<TWeakInterfacePtr<INavLinkCustomInterface>> CustomLinkObjects;
+};
diff --git a/Engine/Source/Runtime/NavigationSystem/Private/NavigationSystem.cpp b/Engine/Source/Runtime/NavigationSystem/Private/NavigationSystem.cpp
--- a/Engine/Source/Runtime/NavigationSystem/Private/NavigationSystem.cpp
+++ b/Engine/Source/Runtime/NavigationSystem/Private/NavigationSystem.cpp
@@ -23,6 +23,7 @@
 #include "VisualLogger/VisualLogger.h"
 #include "NavMesh/NavMeshBoundsVolume.h"
 #include "NavigationInvokerComponent.h"
+#include "NavigationObjectRepository.h"
 #include "AI/Navigation/NavigationDataChunk.h"
 #include "Engine/Engine.h"
 #include "UObject/Package.h"
@@ -618,8 +619,12 @@ void FNavRegenTimeSliceManager::LogTileStatistics(const TArray<TObjectPtr<ANavig
 // UNavigationSystemV1                                                                
 //----------------------------------------------------------------------//
 bool UNavigationSystemV1::bNavigationAutoUpdateEnabled = true;
+
+PRAGMA_DISABLE_DEPRECATION_WARNINGS
 TMap<INavLinkCustomInterface*, FWeakObjectPtr> UNavigationSystemV1::PendingCustomLinkRegistration;
 FCriticalSection UNavigationSystemV1::CustomLinkRegistrationSection;
+PRAGMA_ENABLE_DEPRECATION_WARNINGS
+
 #if !(UE_BUILD_SHIPPING || UE_BUILD_TEST)
 FNavigationSystemExec UNavigationSystemV1::ExecHandler;
 #endif // !UE_BUILD_SHIPPING
@@ -867,6 +872,8 @@ void UNavigationSystemV1::DoInitialSetup()
 	UpdateAbstractNavData();
 	CreateCrowdManager();
 
+	RegisterToRepositoryDelegates();
+
 	bInitialSetupHasBeenPerformed = true;
 }
 
@@ -1044,6 +1051,19 @@ bool UNavigationSystemV1::ConditionalPopulateNavOctree()
 						AddLevelToOctree(*Level);
 					}
 				}
+
+				// Register nav relevant objects currently registered in the repository world subsystem.
+				// This covers objects that are not AActor/UActorComponent based.
+				if (const UNavigationObjectRepository* Repository = UWorld::GetSubsystem<UNavigationObjectRepository>(GetWorld()))
+				{
+					for (TWeakInterfacePtr<INavRelevantInterface> It : Repository->GetNavRelevantObjects())
+					{
+						if (INavRelevantInterface* Interface = It.Get())
+						{
+							RegisterNavOctreeElement(Cast<UObject>(Interface), Interface, FNavigationOctreeController::OctreeUpdate_Default);
+						}
+					}
+				}
 			}
 		}
 	}
@@ -1142,17 +1162,20 @@ void UNavigationSystemV1::OnBeginTearingDown(UWorld* World)
 void UNavigationSystemV1::OnWorldInitDone(FNavigationSystemRunMode Mode)
 {
 	UNavigationSystemBase::OnNavigationInitStartStaticDelegate().Broadcast(*this);
-	
+
 	OperationMode = Mode;
 	DoInitialSetup();
 	
 	UWorld* World = GetWorld();
 	check(World);
 
-	// process all queued custom link registration requests
+	// process all registered link from the repository subsystem
 	// (since it's possible navigation system was not ready by the time
 	// those links were serialized-in or spawned)
-	ProcessCustomLinkPendingRegistration();
+	if (!bWorldInitDone)
+	{
+		ProcessCustomLinkPendingRegistration();
+	}
 
 	if (IsThereAnywhereToBuildNavigation() == false)
 	{
@@ -2587,29 +2610,14 @@ void UNavigationSystemV1::ProcessRegistrationCandidates()
 
 void UNavigationSystemV1::ProcessCustomLinkPendingRegistration()
 {
-	FScopeLock AccessLock(&CustomLinkRegistrationSection);
-
-	TMap<INavLinkCustomInterface*, FWeakObjectPtr> TempPending = PendingCustomLinkRegistration;
-	PendingCustomLinkRegistration.Empty();
-
-	for (TMap<INavLinkCustomInterface*, FWeakObjectPtr>::TIterator It(TempPending); It; ++It)
+	if (const UNavigationObjectRepository* Repository = UWorld::GetSubsystem<UNavigationObjectRepository>(GetWorld()))
 	{
-		INavLinkCustomInterface* ILink = It.Key();
-		FWeakObjectPtr LinkOb = It.Value();
-		
-		if (LinkOb.IsValid() && ILink)
+		for (TWeakInterfacePtr<INavLinkCustomInterface> It : Repository->GetCustomLinks())
 		{
-#if WITH_EDITOR
-			// In Editor multiple NavigationSystems may exist at the same time (i.e. Editor, Client Game, Server Game worlds)
-			// so we want to make sure that any given NavigationSystem instance performs a single flush of the global pending queue
-			// to register the links associated to their outer World.
-			// Following registration requests will be forwarded directly to the NavigationSystem and won't use the queue.
-			// We call RequestCustomLinkRegistering instead of RegisterCustomLink so each link
-			// will register to the navigation system associated to their outer world (if created) or put back in the queue.
-			RequestCustomLinkRegistering(*ILink, LinkOb.Get());
-#else
-			RegisterCustomLink(*ILink);
-#endif // WITH_EDITOR
+			if (INavLinkCustomInterface* Interface = It.Get())
+			{
+				RegisterCustomLink(*Interface);
+			}
 		}
 	}
 }
@@ -2856,31 +2864,25 @@ void UNavigationSystemV1::UpdateCustomLink(const INavLinkCustomInterface* Custom
 	}
 }
 
-void UNavigationSystemV1::RequestCustomLinkRegistering(INavLinkCustomInterface& CustomLink, UObject* OwnerOb)
+void UNavigationSystemV1::RequestCustomLinkRegistering(INavLinkCustomInterface& CustomLink, UObject* Owner)
 {
-	UNavigationSystemV1* NavSys = FNavigationSystem::GetCurrent<UNavigationSystemV1>(OwnerOb);
-	if (NavSys)
+	if (Owner != nullptr)
 	{
-		NavSys->RegisterCustomLink(CustomLink);
-	}
-	else
-	{
-		FScopeLock AccessLock(&CustomLinkRegistrationSection);
-		PendingCustomLinkRegistration.Add(&CustomLink, OwnerOb);
+		if (UNavigationObjectRepository* Repository = UWorld::GetSubsystem<UNavigationObjectRepository>(Owner->GetWorld()))
+		{
+			Repository->RegisterCustomNavLinkObject(CustomLink);
+		}
 	}
 }
 
-void UNavigationSystemV1::RequestCustomLinkUnregistering(INavLinkCustomInterface& CustomLink, UObject* OwnerOb)
+void UNavigationSystemV1::RequestCustomLinkUnregistering(INavLinkCustomInterface& CustomLink, UObject* Owner)
 {
-	UNavigationSystemV1* NavSys = FNavigationSystem::GetCurrent<UNavigationSystemV1>(OwnerOb);
-	if (NavSys)
-	{
-		NavSys->UnregisterCustomLink(CustomLink);
-	}
-	else
+	if (Owner != nullptr)
 	{
-		FScopeLock AccessLock(&CustomLinkRegistrationSection);
-		PendingCustomLinkRegistration.Remove(&CustomLink);
+		if (UNavigationObjectRepository* Repository = UWorld::GetSubsystem<UNavigationObjectRepository>(Owner->GetWorld()))
+		{
+			Repository->UnregisterCustomNavLinkObject(CustomLink);
+		}
 	}
 }
 
@@ -3194,13 +3196,11 @@ void UNavigationSystemV1::OnNavRelevantObjectRegistered(UObject& Object)
 		return;
 	}
 
-	SCOPE_CYCLE_COUNTER(STAT_DebugNavOctree);
 	if (INavRelevantInterface* NavInterface = Cast<INavRelevantInterface>(&Object))
 	{
-		UWorld* World = Object.GetWorld();
-		if (UNavigationSystemV1* NavSys = FNavigationSystem::GetCurrent<UNavigationSystemV1>(World))
+		if (UNavigationObjectRepository* Repository = UWorld::GetSubsystem<UNavigationObjectRepository>(Object.GetWorld()))
 		{
-			NavSys->RegisterNavOctreeElement(&Object, NavInterface, FNavigationOctreeController::OctreeUpdate_Default);
+			Repository->RegisterNavRelevantObject(*NavInterface);
 		}
 	}
 }
@@ -3235,13 +3235,11 @@ void UNavigationSystemV1::OnNavRelevantObjectUnregistered(UObject& Object)
 		return;
 	}
 
-	SCOPE_CYCLE_COUNTER(STAT_DebugNavOctree);
 	if (INavRelevantInterface* NavInterface = Cast<INavRelevantInterface>(&Object))
 	{
-		UWorld* World = Object.GetWorld();
-		if (UNavigationSystemV1* NavSys = FNavigationSystem::GetCurrent<UNavigationSystemV1>(World))
+		if (UNavigationObjectRepository* Repository = UWorld::GetSubsystem<UNavigationObjectRepository>(Object.GetWorld()))
 		{
-			NavSys->UnregisterNavOctreeElement(&Object, NavInterface, FNavigationOctreeController::OctreeUpdate_Default);
+			Repository->UnregisterNavRelevantObject(*NavInterface);
 		}
 	}
 }
@@ -4639,11 +4637,12 @@ void UNavigationSystemV1::CleanUp(FNavigationSystem::ECleanupMode Mode)
 	
 	MainNavData = nullptr;
 
-	// reset unique link Id for new map
-	UWorld* MyWorld = (Mode == FNavigationSystem::ECleanupMode::CleanupWithWorld) ? GetWorld() : NULL;
+	const UWorld* MyWorld = (Mode == FNavigationSystem::ECleanupMode::CleanupWithWorld) ? GetWorld() : nullptr;
 	if (MyWorld)
 	{
+		UnregisterFromRepositoryDelegates();
 
+		// reset unique link Id for new map
 		if (MyWorld->WorldType == EWorldType::Game || MyWorld->WorldType == EWorldType::Editor)
 		{
 			UE_LOG(LogNavLink, VeryVerbose, TEXT("Reset navlink id on cleanup."));
@@ -5247,6 +5246,45 @@ void UNavigationSystemV1::UnregisterInvoker_Internal(const UObject& Invoker)
 	Invokers.Remove(&Invoker);
 }
 
+void UNavigationSystemV1::RegisterToRepositoryDelegates()
+{
+	if (UNavigationObjectRepository* Repository = UWorld::GetSubsystem<UNavigationObjectRepository>(GetWorld()))
+	{
+		Repository->OnCustomNavLinkObjectRegistered.BindWeakLambda(this, [this](INavLinkCustomInterface& CustomLink)
+			{
+				RegisterCustomLink(CustomLink);
+			});
+
+		Repository->OnCustomNavLinkObjectUnregistered.BindWeakLambda(this, [this](INavLinkCustomInterface& CustomLink)
+			{
+				UnregisterCustomLink(CustomLink);
+			});
+
+		Repository->OnNavRelevantObjectRegistered.BindWeakLambda(this, [this](INavRelevantInterface& NavRelevantObject)
+			{
+				SCOPE_CYCLE_COUNTER(STAT_DebugNavOctree);
+				RegisterNavOctreeElement(Cast<UObject>(&NavRelevantObject), &NavRelevantObject, FNavigationOctreeController::OctreeUpdate_Default);
+			});
+
+		Repository->OnNavRelevantObjectUnregistered.BindWeakLambda(this, [this](INavRelevantInterface& NavRelevantObject)
+			{
+				SCOPE_CYCLE_COUNTER(STAT_DebugNavOctree);
+				UnregisterNavOctreeElement(Cast<UObject>(&NavRelevantObject), &NavRelevantObject, FNavigationOctreeController::OctreeUpdate_Default);
+			});
+	}
+}
+
+void UNavigationSystemV1::UnregisterFromRepositoryDelegates() const
+{
+	if (UNavigationObjectRepository* Repository = UWorld::GetSubsystem<UNavigationObjectRepository>(GetWorld()))
+	{
+		Repository->OnCustomNavLinkObjectRegistered = nullptr;
+		Repository->OnCustomNavLinkObjectUnregistered = nullptr;
+		Repository->OnNavRelevantObjectRegistered = nullptr;
+		Repository->OnNavRelevantObjectUnregistered = nullptr;
+	}
+}
+
 void UNavigationSystemV1::UpdateInvokers()
 {
 	QUICK_SCOPE_CYCLE_COUNTER(STAT_Navigation_UpdateInvokers);
diff --git a/Engine/Source/Runtime/NavigationSystem/Public/NavigationSystem.h b/Engine/Source/Runtime/NavigationSystem/Public/NavigationSystem.h
--- a/Engine/Source/Runtime/NavigationSystem/Public/NavigationSystem.h
+++ b/Engine/Source/Runtime/NavigationSystem/Public/NavigationSystem.h
@@ -471,6 +471,12 @@ private:
 
 	void UnregisterInvoker_Internal(const UObject& Invoker);
 
+	/** Registers to the navigation objects repository subsystem delegates to get notified on object registration */
+	void RegisterToRepositoryDelegates();
+
+	/** Unregisters from the navigation objects repository subsystem delegates */
+	void UnregisterFromRepositoryDelegates() const;
+
 public:
 	//----------------------------------------------------------------------//
 	// Blueprint functions
@@ -868,9 +874,15 @@ protected:
 	/** Processes registration of candidates queues via RequestRegistration and stored in NavDataRegistrationQueue */
 	NAVIGATIONSYSTEM_API virtual void ProcessRegistrationCandidates();
 
-	/** registers CustomLinks awaiting registration in PendingCustomLinkRegistration */
+	/** Registers custom navigation links awaiting registration in the navigation object repository */
 	NAVIGATIONSYSTEM_API void ProcessCustomLinkPendingRegistration();
 
+	UE_DEPRECATED(5.4, "Obsolete since navigation related objects register to the NavigationObjectRepository world subsytem.")
+	static NAVIGATIONSYSTEM_API FCriticalSection CustomLinkRegistrationSection;
+
+	UE_DEPRECATED(5.4, "Obsolete since navigation related objects register to the NavigationObjectRepository world subsytem.")
+	static NAVIGATIONSYSTEM_API TMap<INavLinkCustomInterface*, FWeakObjectPtr> PendingCustomLinkRegistration;
+
 	/** used to apply updates of nav volumes in navigation system's tick */
 	NAVIGATIONSYSTEM_API virtual void PerformNavigationBoundsUpdate(const TArray<FNavigationBoundsUpdateRequest>& UpdateRequests);
 	
@@ -983,12 +995,13 @@ public:
 	//----------------------------------------------------------------------//
 	NAVIGATIONSYSTEM_API virtual void RegisterCustomLink(INavLinkCustomInterface& CustomLink);
 	NAVIGATIONSYSTEM_API void UnregisterCustomLink(INavLinkCustomInterface& CustomLink);
-	
+	int32 GetNumCustomLinks() const { return CustomNavLinksMap.Num(); }
+
 	static NAVIGATIONSYSTEM_API void RequestCustomLinkRegistering(INavLinkCustomInterface& CustomLink, UObject* OwnerOb);
 	static NAVIGATIONSYSTEM_API void RequestCustomLinkUnregistering(INavLinkCustomInterface& CustomLink, UObject* ObjectOb);
 
 	UE_DEPRECATED(5.3, "LinkIds are now based on FNavLinkId call the version of this function that takes FNavLinkId. This function only returns nullptr.")
-	INavLinkCustomInterface* GetCustomLink(uint32 UniqueLinkId) const { return nullptr; };
+	INavLinkCustomInterface* GetCustomLink(uint32 UniqueLinkId) const { return nullptr; }
 
 	/** find custom link by unique ID */
 	NAVIGATIONSYSTEM_API INavLinkCustomInterface* GetCustomLink(FNavLinkId UniqueLinkId) const;
@@ -1221,8 +1234,7 @@ protected:
 
 	// async queries
 	FCriticalSection NavDataRegistrationSection;
-	static NAVIGATIONSYSTEM_API FCriticalSection CustomLinkRegistrationSection;
-	
+
 #if WITH_EDITOR
 	uint8 NavUpdateLockFlags;
 #endif
@@ -1256,7 +1268,6 @@ protected:
 	static NAVIGATIONSYSTEM_API bool bStaticRuntimeNavigation;
 	static NAVIGATIONSYSTEM_API bool bIsPIEActive;
 
-	static NAVIGATIONSYSTEM_API TMap<INavLinkCustomInterface*, FWeakObjectPtr> PendingCustomLinkRegistration;
 	TSet<TObjectPtr<const UClass>> NavAreaClasses;
 
 	FNavRegenTimeSliceManager NavRegenTimeSliceManager;
