Fix for InstanceStaticMeshComp with odd numbers of negated scale axis won't generate navmesh correctly.

Added the missing flipping indices step where the vertices of ISMComp are converted to Recast Space.

diff --git a/Engine/Source/Runtime/NavigationSystem/Private/NavMesh/RecastNavMeshGenerator.cpp b/Engine/Source/Runtime/NavigationSystem/Private/NavMesh/RecastNavMeshGenerator.cpp
--- a/Engine/Source/Runtime/NavigationSystem/Private/NavMesh/RecastNavMeshGenerator.cpp
+++ b/Engine/Source/Runtime/NavigationSystem/Private/NavMesh/RecastNavMeshGenerator.cpp
@@ -1649,7 +1649,7 @@ FRecastTileGenerator::FRecastTileGenerator(FRecastNavMeshGenerator& ParentGenera
 	ParentGeneratorWeakPtr = ((FNavDataGenerator&)ParentGenerator).AsShared();
 
 	RasterizeGeomRecastState = ERasterizeGeomRecastTimeSlicedState::MarkWalkableTriangles;
-	RasterizeGeomState = ERasterizeGeomTimeSlicedState::RasterizeGeometryTransformCoords;
+	RasterizeGeomState = ERasterizeGeomTimeSlicedState::RasterizeGeometryTransformCoordsAndFlipIndices;
 	GenerateRecastFilterState = EGenerateRecastFilterTimeSlicedState::FilterLowHangingWalkableObstacles;
 	GenRecastFilterLedgeSpansYStart = 0;
 	DoWorkTimeSlicedState = EDoWorkTimeSlicedState::GatherGeometryFromSources;
@@ -2716,15 +2716,6 @@ ETimeSliceWorkResult FRecastTileGenerator::RasterizeGeometryRecastTimeSliced(FNa
 	return ETimeSliceWorkResult::Succeeded;
 }
 
-ETimeSliceWorkResult FRecastTileGenerator::RasterizeGeometryRecastTimeSliced(FNavMeshBuildContext& BuildContext, const TArray<float>& Coords, const TArray<int32>& Indices, const rcRasterizationFlags RasterizationFlags, FTileRasterizationContext& RasterContext)
-{
-	TArray<FVector::FReal> RealCoords;
-
-	RealCoords = UE::LWC::ConvertArrayType<FVector::FReal>(Coords);
-
-	return RasterizeGeometryRecastTimeSliced(BuildContext, RealCoords, Indices, RasterizationFlags, RasterContext);
-}
-
 void FRecastTileGenerator::RasterizeGeometryRecast(FNavMeshBuildContext& BuildContext, const TArray<FVector::FReal>& Coords, const TArray<int32>& Indices, const rcRasterizationFlags RasterizationFlags, FTileRasterizationContext& RasterContext)
 {
 	QUICK_SCOPE_CYCLE_COUNTER(STAT_Navigation_RasterizeGeometryRecast);
@@ -2772,13 +2763,39 @@ void FRecastTileGenerator::RasterizeGeometryRecast(FNavMeshBuildContext& BuildCo
 	RasterizeGeomRecastTriAreas.Reset();
 }
 
-void FRecastTileGenerator::RasterizeGeometryRecast(FNavMeshBuildContext& BuildContext, const TArray<float>& Coords, const TArray<int32>& Indices, const rcRasterizationFlags RasterizationFlags, FTileRasterizationContext& RasterContext)
+void FRecastTileGenerator::RasterizeGeometryTransformCoordsAndFlipIndices(const TArray<FVector::FReal>& Coords, const TArray<int32>& Indices, const FTransform& LocalToWorld)
 {
-	TArray<FVector::FReal> RealCoords;
+	QUICK_SCOPE_CYCLE_COUNTER(STAT_Navigation_RasterizeGeometryTransformCoordsAndFlipIndices);
+
+	RasterizeGeometryWorldRecastCoords.SetNumUninitialized(Coords.Num(), EAllowShrinking::No);
+
+	FMatrix LocalToRecastWorld = LocalToWorld.ToMatrixWithScale() * Unreal2RecastMatrix();
+
+	// Convert geometry to recast world space
+	for (int32 i = 0; i < Coords.Num(); i += 3)
+	{
+		// collision cache stores coordinates in recast space, convert them to unreal and transform to recast world space
+		FVector WorldRecastCoord = LocalToRecastWorld.TransformPosition(Recast2UnrealPoint(&Coords[i]));
+
+		RasterizeGeometryWorldRecastCoords[i + 0] = WorldRecastCoord.X;
+		RasterizeGeometryWorldRecastCoords[i + 1] = WorldRecastCoord.Y;
+		RasterizeGeometryWorldRecastCoords[i + 2] = WorldRecastCoord.Z;
+	}
 
-	RealCoords = UE::LWC::ConvertArrayType<FVector::FReal>(Coords);
+	//Flip the order of indices for each triangle if the source object is mirrored
+	bRasterizeGeometryUseFlippedIndices = false;
+	if (LocalToWorld.GetDeterminant() < 0.f)
+	{
+		bRasterizeGeometryUseFlippedIndices = true;
+		RasterizeGeometryFlippedIndices.SetNumUninitialized(Indices.Num(), EAllowShrinking::No);
 
-	RasterizeGeometryRecast(BuildContext, RealCoords, Indices, RasterizationFlags, RasterContext);
+		for (int32 i = 0; i < Indices.Num(); i += 3)
+		{
+			RasterizeGeometryFlippedIndices[i] = Indices[i + 2];
+			RasterizeGeometryFlippedIndices[i + 1] = Indices[i + 1];
+			RasterizeGeometryFlippedIndices[i + 2] = Indices[i];
+		}
+	}
 }
 
 void FRecastTileGenerator::RasterizeGeometryTransformCoords(const TArray<FVector::FReal>& Coords, const FTransform& LocalToWorld)
@@ -2801,15 +2818,6 @@ void FRecastTileGenerator::RasterizeGeometryTransformCoords(const TArray<FVector
 	}
 }
 
-void FRecastTileGenerator::RasterizeGeometryTransformCoords(const TArray<float>& Coords, const FTransform& LocalToWorld)
-{
-	TArray<FVector::FReal> RealCoords;
-
-	RealCoords = UE::LWC::ConvertArrayType<FVector::FReal>(Coords);
-
-	RasterizeGeometryTransformCoords(RealCoords, LocalToWorld);
-}
-
 ETimeSliceWorkResult FRecastTileGenerator::RasterizeGeometryTimeSliced(FNavMeshBuildContext& BuildContext, const TArray<FVector::FReal>& Coords, const TArray<int32>& Indices, const FTransform& LocalToWorld, const rcRasterizationFlags RasterizationFlags, FTileRasterizationContext& RasterContext)
 {
 	QUICK_SCOPE_CYCLE_COUNTER(STAT_Navigation_RasterizeGeometry);
@@ -2820,13 +2828,13 @@ ETimeSliceWorkResult FRecastTileGenerator::RasterizeGeometryTimeSliced(FNavMeshB
 
 	switch (RasterizeGeomState)
 	{
-	case ERasterizeGeomTimeSlicedState::RasterizeGeometryTransformCoords:
+	case ERasterizeGeomTimeSlicedState::RasterizeGeometryTransformCoordsAndFlipIndices:
 	{
-		RasterizeGeometryTransformCoords(Coords, LocalToWorld);
+		RasterizeGeometryTransformCoordsAndFlipIndices(Coords, Indices, LocalToWorld);
 
 		RasterizeGeomState = ERasterizeGeomTimeSlicedState::RasterizeGeometryRecast;
 
-		MARK_TIMESLICE_SECTION_DEBUG(TimeSliceManager->GetTimeSlicer(), RasterizeGeometryTransformCoords);
+		MARK_TIMESLICE_SECTION_DEBUG(TimeSliceManager->GetTimeSlicer(), RasterizeGeometryTransformCoordsAndFlipIndices);
 
 		if (TimeSliceManager->GetTimeSlicer().TestTimeSliceFinished())
 		{
@@ -2835,12 +2843,13 @@ ETimeSliceWorkResult FRecastTileGenerator::RasterizeGeometryTimeSliced(FNavMeshB
 	}// fall through to next state
 	case ERasterizeGeomTimeSlicedState::RasterizeGeometryRecast:
 	{
-		WorkResult = RasterizeGeometryRecastTimeSliced(BuildContext, RasterizeGeometryWorldRecastCoords, Indices, RasterizationFlags, RasterContext);
+		const TArray<int32>& RasterizeGeometryIndices = bRasterizeGeometryUseFlippedIndices ? RasterizeGeometryFlippedIndices : Indices;
+		WorkResult = RasterizeGeometryRecastTimeSliced(BuildContext, RasterizeGeometryWorldRecastCoords, RasterizeGeometryIndices, RasterizationFlags, RasterContext);
 
 		if (WorkResult != ETimeSliceWorkResult::CallAgainNextTimeSlice)
 		{
-			//if we have finished rasterizing this geometry then reset RasterizeGeomTimeSlicedState so next time this function is called we go back to RasterizeGeometryTransformCoords first
-			RasterizeGeomState = ERasterizeGeomTimeSlicedState::RasterizeGeometryTransformCoords;
+			//if we have finished rasterizing this geometry then reset RasterizeGeomTimeSlicedState so next time this function is called we go back to RasterizeGeometryTransformCoordsAndFlipIndices first
+			RasterizeGeomState = ERasterizeGeomTimeSlicedState::RasterizeGeometryTransformCoordsAndFlipIndices;
 		}
 	}
 	break;
@@ -2853,30 +2862,13 @@ ETimeSliceWorkResult FRecastTileGenerator::RasterizeGeometryTimeSliced(FNavMeshB
 	return WorkResult;
 }
 
-ETimeSliceWorkResult FRecastTileGenerator::RasterizeGeometryTimeSliced(FNavMeshBuildContext& BuildContext, const TArray<float>& Coords, const TArray<int32>& Indices, const FTransform& LocalToWorld, const rcRasterizationFlags RasterizationFlags, FTileRasterizationContext& RasterContext)
-{
-	TArray<FVector::FReal> RealCoords;
-
-	RealCoords = UE::LWC::ConvertArrayType<FVector::FReal>(Coords);
-
-	return RasterizeGeometryTimeSliced(BuildContext, RealCoords, Indices, LocalToWorld, RasterizationFlags, RasterContext);
-}
-
 void FRecastTileGenerator::RasterizeGeometry(FNavMeshBuildContext& BuildContext, const TArray<FVector::FReal>& Coords, const TArray<int32>& Indices, const FTransform& LocalToWorld, const rcRasterizationFlags RasterizationFlags, FTileRasterizationContext& RasterContext)
 {
 	QUICK_SCOPE_CYCLE_COUNTER(STAT_Navigation_RasterizeGeometry);
 
-	RasterizeGeometryTransformCoords(Coords, LocalToWorld);
-	RasterizeGeometryRecast(BuildContext, RasterizeGeometryWorldRecastCoords, Indices, RasterizationFlags, RasterContext);
-}
-
-void FRecastTileGenerator::RasterizeGeometry(FNavMeshBuildContext& BuildContext, const TArray<float>& Coords, const TArray<int32>& Indices, const FTransform& LocalToWorld, const rcRasterizationFlags RasterizationFlags, FTileRasterizationContext& RasterContext)
-{
-	TArray<FVector::FReal> RealCoords;
-
-	RealCoords = UE::LWC::ConvertArrayType<FVector::FReal>(Coords);
-
-	return RasterizeGeometry(BuildContext, RealCoords, Indices, LocalToWorld, RasterizationFlags, RasterContext);
+	RasterizeGeometryTransformCoordsAndFlipIndices(Coords, Indices, LocalToWorld);
+	const TArray<int32>& RasterizeGeometryIndices = bRasterizeGeometryUseFlippedIndices ? RasterizeGeometryFlippedIndices : Indices;
+	RasterizeGeometryRecast(BuildContext, RasterizeGeometryWorldRecastCoords, RasterizeGeometryIndices, RasterizationFlags, RasterContext);
 }
 
 ETimeSliceWorkResult FRecastTileGenerator::RasterizeTrianglesTimeSliced(FNavMeshBuildContext& BuildContext, FTileRasterizationContext& RasterContext)
@@ -5466,15 +5458,6 @@ namespace RecastTileVersionHelper
 	}
 }
 
-// Deprecated
-TArray<uint32> FRecastNavMeshGenerator::RemoveTileLayers(const int32 TileX, const int32 TileY, TMap<int32, dtPolyRef>* OldLayerTileIdMap)
-{
-	const TArray<FNavTileRef>& TileRefs = RemoveTileLayersAndGetUpdatedTiles(TileX, TileY, OldLayerTileIdMap);
-	TArray<uint32> TileIds;
-	FNavTileRef::DeprecatedGetTileIdsFromNavTileRefs(DestNavMesh->GetRecastNavMeshImpl(), TileRefs, TileIds);
-	return TileIds;
-}
-
 TArray<FNavTileRef> FRecastNavMeshGenerator::RemoveTileLayersAndGetUpdatedTiles(const int32 TileX, const int32 TileY, TMap<int32, dtPolyRef>* OldLayerTileIdMap)
 {
 	dtNavMesh* DetourMesh = DestNavMesh->GetRecastNavMeshImpl()->GetRecastMesh();
@@ -5728,17 +5711,6 @@ void FRecastNavMeshGenerator::LogDirtyAreas(const UObject& OwnerNav,
 }
 #endif
 
-// Deprecated 
-ETimeSliceWorkResult FRecastNavMeshGenerator::AddGeneratedTilesTimeSliced(FRecastTileGenerator& TileGenerator, TArray<uint32>& OutResultTileIndices)
-{
-	TArray<FNavTileRef> OutTileRefs;
-	const ETimeSliceWorkResult Result = AddGeneratedTilesTimeSliced(TileGenerator, OutTileRefs);
-
-	TArray<uint32> TileIds;
-	FNavTileRef::DeprecatedGetTileIdsFromNavTileRefs(DestNavMesh->GetRecastNavMeshImpl(), OutTileRefs, TileIds);
-	return Result;
-}
-
 ETimeSliceWorkResult FRecastNavMeshGenerator::AddGeneratedTilesTimeSliced(FRecastTileGenerator& TileGenerator, TArray<FNavTileRef>& OutResultTileRefs)
 {
 	SCOPE_CYCLE_COUNTER(STAT_Navigation_RecastAddGeneratedTiles);
@@ -5819,15 +5791,6 @@ ETimeSliceWorkResult FRecastNavMeshGenerator::AddGeneratedTilesTimeSliced(FRecas
 	return WorkResult;
 }
 
-// Deprecated
-TArray<uint32> FRecastNavMeshGenerator::AddGeneratedTiles(FRecastTileGenerator& TileGenerator)
-{
-	const TArray<FNavTileRef>& TileRefs = AddGeneratedTilesAndGetUpdatedTiles(TileGenerator);
-	TArray<uint32> TileIds;
-	FNavTileRef::DeprecatedGetTileIdsFromNavTileRefs(DestNavMesh->GetRecastNavMeshImpl(), TileRefs, TileIds);
-	return TileIds;
-}
-
 TArray<FNavTileRef> FRecastNavMeshGenerator::AddGeneratedTilesAndGetUpdatedTiles(FRecastTileGenerator& TileGenerator)
 {
 	SCOPE_CYCLE_COUNTER(STAT_Navigation_RecastAddGeneratedTiles);
@@ -6350,14 +6313,6 @@ TSharedRef<FRecastTileGenerator> FRecastNavMeshGenerator::CreateTileGenerator(co
 	return ConstructTileGeneratorImpl<FRecastTileGenerator>(Coord, DirtyAreas, PendingTileCreationTime);
 }
 
-// Deprecated
-void FRecastNavMeshGenerator::RemoveLayers(const FIntPoint& Tile, TArray<uint32>& UpdatedTiles)
-{
-	TArray<FNavTileRef> TileRefs;
-	RemoveLayers(Tile, TileRefs);
-	FNavTileRef::DeprecatedGetTileIdsFromNavTileRefs(DestNavMesh->GetRecastNavMeshImpl(), TileRefs, UpdatedTiles);
-}
-	
 void FRecastNavMeshGenerator::RemoveLayers(const FIntPoint& Tile, TArray<FNavTileRef>& UpdatedTiles)
 {
 	SCOPE_CYCLE_COUNTER(STAT_Navigation_RemoveLayers);
@@ -6392,15 +6347,6 @@ void FRecastNavMeshGenerator::StoreDebugData(const FRecastTileGenerator& TileGen
 #endif
 
 #if RECAST_ASYNC_REBUILDING
-// Deprecated	
-TArray<uint32> FRecastNavMeshGenerator::ProcessTileTasksAsync(const int32 NumTasksToProcess)
-{
-	const TArray<FNavTileRef>& TileRefs = ProcessTileTasksAsyncAndGetUpdatedTiles(NumTasksToProcess);
-	TArray<uint32> TileIds;
-	FNavTileRef::DeprecatedGetTileIdsFromNavTileRefs(DestNavMesh->GetRecastNavMeshImpl(), TileRefs, TileIds);
-	return TileIds;
-}
-
 TArray<FNavTileRef> FRecastNavMeshGenerator::ProcessTileTasksAsyncAndGetUpdatedTiles(const int32 NumTasksToProcess)
 {
 	QUICK_SCOPE_CYCLE_COUNTER(STAT_RecastNavMeshGenerator_ProcessTileTasksAsync);
@@ -6523,15 +6469,6 @@ TSharedRef<FRecastTileGenerator> FRecastNavMeshGenerator::CreateTileGeneratorFro
 	return TileGenerator;
 }
 
-// Deprecated
-TArray<uint32> FRecastNavMeshGenerator::ProcessTileTasksSyncTimeSliced()
-{
-	const TArray<FNavTileRef>& TileRefs = ProcessTileTasksSyncTimeSlicedAndGetUpdatedTiles();
-	TArray<uint32> TileIds;
-	FNavTileRef::DeprecatedGetTileIdsFromNavTileRefs(DestNavMesh->GetRecastNavMeshImpl(), TileRefs, TileIds);
-	return TileIds;
-}
-
 TArray<FNavTileRef> FRecastNavMeshGenerator::ProcessTileTasksSyncTimeSlicedAndGetUpdatedTiles()
 {
 	QUICK_SCOPE_CYCLE_COUNTER(STAT_RecastNavMeshGenerator_ProcessTileTasksSyncTimeSliced);
@@ -6857,14 +6794,6 @@ TArray<FNavTileRef> FRecastNavMeshGenerator::ProcessTileTasksSyncAndGetUpdatedTi
 }
 #endif
 
-TArray<uint32> FRecastNavMeshGenerator::ProcessTileTasks(const int32 NumTasksToProcess)
-{
-	const TArray<FNavTileRef>& TileRefs = ProcessTileTasksAndGetUpdatedTiles(NumTasksToProcess);
-	TArray<uint32> TileIds;
-	FNavTileRef::DeprecatedGetTileIdsFromNavTileRefs(DestNavMesh->GetRecastNavMeshImpl(), TileRefs, TileIds);
-	return TileIds;
-}
-
 TArray<FNavTileRef> FRecastNavMeshGenerator::ProcessTileTasksAndGetUpdatedTiles(const int32 NumTasksToProcess)
 {
 	QUICK_SCOPE_CYCLE_COUNTER(STAT_RecastNavMeshGenerator_ProcessTileTasks);
diff --git a/Engine/Source/Runtime/NavigationSystem/Public/NavMesh/RecastNavMeshGenerator.h b/Engine/Source/Runtime/NavigationSystem/Public/NavMesh/RecastNavMeshGenerator.h
--- a/Engine/Source/Runtime/NavigationSystem/Public/NavMesh/RecastNavMeshGenerator.h
+++ b/Engine/Source/Runtime/NavigationSystem/Public/NavMesh/RecastNavMeshGenerator.h
@@ -282,7 +282,8 @@ enum class ERasterizeGeomRecastTimeSlicedState : uint8
 
 enum class ERasterizeGeomTimeSlicedState : uint8
 {
-	RasterizeGeometryTransformCoords,
+	RasterizeGeometryTransformCoordsAndFlipIndices,
+	RasterizeGeometryTransformCoords UE_DEPRECATED(5.5, "The state doesn't handle indices flipping. Use RasterizeGeometryTransformCoordsAndIndices instead.") = RasterizeGeometryTransformCoordsAndFlipIndices,
 	RasterizeGeometryRecast,
 };
 
@@ -412,20 +413,12 @@ protected:
 	NAVIGATIONSYSTEM_API ETimeSliceWorkResult RasterizeTrianglesTimeSliced(FNavMeshBuildContext& BuildContext, FTileRasterizationContext& RasterContext);
 	NAVIGATIONSYSTEM_API void RasterizeTriangles(FNavMeshBuildContext& BuildContext, FTileRasterizationContext& RasterContext);
 	NAVIGATIONSYSTEM_API ETimeSliceWorkResult RasterizeGeometryRecastTimeSliced(FNavMeshBuildContext& BuildContext, const TArray<FVector::FReal>& Coords, const TArray<int32>& Indices, const rcRasterizationFlags RasterizationFlags, FTileRasterizationContext& RasterContext);
-	UE_DEPRECATED(5.0, "Call the version of this function where Coords are now a TArray of FReals!")
-	NAVIGATIONSYSTEM_API ETimeSliceWorkResult RasterizeGeometryRecastTimeSliced(FNavMeshBuildContext& BuildContext, const TArray<float>& Coords, const TArray<int32>& Indices, const rcRasterizationFlags RasterizationFlags, FTileRasterizationContext& RasterContext);
 	NAVIGATIONSYSTEM_API void RasterizeGeometryRecast(FNavMeshBuildContext& BuildContext, const TArray<FVector::FReal>& Coords, const TArray<int32>& Indices, const rcRasterizationFlags RasterizationFlags, FTileRasterizationContext& RasterContext);
-	UE_DEPRECATED(5.0, "Call the version of this function where Coords are now a TArray of FReals!")
-	NAVIGATIONSYSTEM_API void RasterizeGeometryRecast(FNavMeshBuildContext& BuildContext, const TArray<float>& Coords, const TArray<int32>& Indices, const rcRasterizationFlags RasterizationFlags, FTileRasterizationContext& RasterContext);
+	NAVIGATIONSYSTEM_API void RasterizeGeometryTransformCoordsAndFlipIndices(const TArray<FVector::FReal>& Coords, const TArray<int32>& Indices, const FTransform& LocalToWorld);
+	UE_DEPRECATED(5.5, "This function was not handling the indices order correctly and is replaced by RasterizeGeometryTransformCoordsAndFlipIndices.")
 	NAVIGATIONSYSTEM_API void RasterizeGeometryTransformCoords(const TArray<FVector::FReal>& Coords, const FTransform& LocalToWorld);
-	UE_DEPRECATED(5.0, "Call the version of this function where Coords are now a TArray of FReals!")
-	NAVIGATIONSYSTEM_API void RasterizeGeometryTransformCoords(const TArray<float>& Coords, const FTransform& LocalToWorld);
 	NAVIGATIONSYSTEM_API ETimeSliceWorkResult RasterizeGeometryTimeSliced(FNavMeshBuildContext& BuildContext, const TArray<FVector::FReal>& Coords, const TArray<int32>& Indices, const FTransform& LocalToWorld, const rcRasterizationFlags RasterizationFlags, FTileRasterizationContext& RasterContext);
-	UE_DEPRECATED(5.0, "Call the version of this function where Coords are now a TArray of FReals!")
-	NAVIGATIONSYSTEM_API ETimeSliceWorkResult RasterizeGeometryTimeSliced(FNavMeshBuildContext& BuildContext, const TArray<float>& Coords, const TArray<int32>& Indices, const FTransform& LocalToWorld, const rcRasterizationFlags RasterizationFlags, FTileRasterizationContext& RasterContext);
 	NAVIGATIONSYSTEM_API void RasterizeGeometry(FNavMeshBuildContext& BuildContext, const TArray<FVector::FReal>& Coords, const TArray<int32>& Indices, const FTransform& LocalToWorld, const rcRasterizationFlags RasterizationFlags, FTileRasterizationContext& RasterContext);
-	UE_DEPRECATED(5.0, "Call the version of this function where Coords are now a TArray of FReals!")
-	NAVIGATIONSYSTEM_API void RasterizeGeometry(FNavMeshBuildContext& BuildContext, const TArray<float>& Coords, const TArray<int32>& Indices, const FTransform& LocalToWorld, const rcRasterizationFlags RasterizationFlags, FTileRasterizationContext& RasterContext);
 	NAVIGATIONSYSTEM_API void GenerateRecastFilter(FNavMeshBuildContext& BuildContext, FTileRasterizationContext& RasterContext);
 	NAVIGATIONSYSTEM_API ETimeSliceWorkResult GenerateRecastFilterTimeSliced(FNavMeshBuildContext& BuildContext, FTileRasterizationContext& RasterContext);
 	NAVIGATIONSYSTEM_API bool BuildCompactHeightField(FNavMeshBuildContext& BuildContext, FTileRasterizationContext& RasterContext);
@@ -538,9 +531,11 @@ protected:
 	TArray<FNavMeshTileData> CompressedLayers;
 	TArray<FNavMeshTileData> NavigationData;
 
-	/** Result of calling RasterizeGeometryInitVars() */
+	/** Result of calling RasterizeGeometry() */
 	TArray<FVector::FReal> RasterizeGeometryWorldRecastCoords;
-	
+	TArray<int32> RasterizeGeometryFlippedIndices;
+	uint8 bRasterizeGeometryUseFlippedIndices : 1;
+
 	// tile's geometry: without voxel cache
 	TArray<FRecastRawGeometryElement> RawGeometry;
 	// areas used for creating navigation data: obstacles
@@ -736,10 +731,6 @@ public:
 	NAVIGATIONSYSTEM_API virtual int32 GetNumRemaningBuildTasks() const override;
 	NAVIGATIONSYSTEM_API virtual int32 GetNumRunningBuildTasks() const override;
 
-	/** Checks if a given tile is being build or has just finished building */
-	UE_DEPRECATED(5.1, "Use new version with FNavTileRef")
-	NAVIGATIONSYSTEM_API bool IsTileChanged(int32 TileIdx) const;
-
 	/** Checks if a given tile is being build or has just finished building */
 	NAVIGATIONSYSTEM_API bool IsTileChanged(const FNavTileRef InTileRef) const;
 		
@@ -887,9 +878,6 @@ protected:
 	/** Marks all tiles overlapping with InclusionBounds dirty (via MarkDirtyTiles). */
 	NAVIGATIONSYSTEM_API bool MarkNavBoundsDirty();
 
-	UE_DEPRECATED(5.1, "Use new version with FNavTileRef")
-	NAVIGATIONSYSTEM_API void RemoveLayers(const FIntPoint& Tile, TArray<uint32>& UpdatedTiles);
-
 	NAVIGATIONSYSTEM_API void RemoveLayers(const FIntPoint& Tile, TArray<FNavTileRef>& UpdatedTiles);
 	
 	NAVIGATIONSYSTEM_API void StoreCompressedTileCacheLayers(const FRecastTileGenerator& TileGenerator, int32 TileX, int32 TileY);
@@ -899,56 +887,30 @@ protected:
 #endif
 
 #if RECAST_ASYNC_REBUILDING
-	/** Processes pending tile generation tasks Async*/
-	UE_DEPRECATED(5.1, "Use ProcessTileTasksAsyncAndGetUpdatedTiles instead")
-	NAVIGATIONSYSTEM_API TArray<uint32> ProcessTileTasksAsync(const int32 NumTasksToProcess);
-
 	/** Processes pending tile generation tasks Async*/
 	NAVIGATIONSYSTEM_API TArray<FNavTileRef> ProcessTileTasksAsyncAndGetUpdatedTiles(const int32 NumTasksToProcess);
 #else
 	NAVIGATIONSYSTEM_API TSharedRef<FRecastTileGenerator> CreateTileGeneratorFromPendingElement(FIntPoint &OutTileLocation, const int32 ForcedPendingTileIdx = INDEX_NONE);
 
-	/** Processes pending tile generation tasks Sync with option for time slicing currently an experimental feature. */
-	UE_DEPRECATED(5.1, "Use ProcessTileTasksSyncTimeSlicedAndGetUpdatedTiles instead")
-	NAVIGATIONSYSTEM_API virtual TArray<uint32> ProcessTileTasksSyncTimeSliced();
-	UE_DEPRECATED(5.1, "Use ProcessTileTasksSyncAndGetUpdatedTiles instead")
-	NAVIGATIONSYSTEM_API TArray<uint32> ProcessTileTasksSync(const int32 NumTasksToProcess);
-
 	/** Processes pending tile generation tasks Sync with option for time slicing currently an experimental feature. */
 	NAVIGATIONSYSTEM_API virtual TArray<FNavTileRef> ProcessTileTasksSyncTimeSlicedAndGetUpdatedTiles();
 	NAVIGATIONSYSTEM_API TArray<FNavTileRef> ProcessTileTasksSyncAndGetUpdatedTiles(const int32 NumTasksToProcess);
 
 	NAVIGATIONSYSTEM_API virtual int32 GetNextPendingDirtyTileToBuild() const;
 #endif
-	/** Processes pending tile generation tasks */
-	UE_DEPRECATED(5.1, "Use ProcessTileTasksAndGetUpdatedTiles instead")
-	NAVIGATIONSYSTEM_API TArray<uint32> ProcessTileTasks(const int32 NumTasksToProcess);
-
 	/** Processes pending tile generation tasks */
 	NAVIGATIONSYSTEM_API TArray<FNavTileRef> ProcessTileTasksAndGetUpdatedTiles(const int32 NumTasksToProcess);
 
 	NAVIGATIONSYSTEM_API void ResetTimeSlicedTileGeneratorSync();
 
 public:
-	/** Adds generated tiles to NavMesh, replacing old ones, uses time slicing returns Failed if any layer failed */
-	UE_DEPRECATED(5.1, "Use new version with FNavTileRef")
-	NAVIGATIONSYSTEM_API ETimeSliceWorkResult AddGeneratedTilesTimeSliced(FRecastTileGenerator& TileGenerator, TArray<uint32>& OutResultTileIndices);
-
 	/** Adds generated tiles to NavMesh, replacing old ones, uses time slicing returns Failed if any layer failed */
 	NAVIGATIONSYSTEM_API ETimeSliceWorkResult AddGeneratedTilesTimeSliced(FRecastTileGenerator& TileGenerator, TArray<FNavTileRef>& OutResultTileRefs);
 
-	/** Adds generated tiles to NavMesh, replacing old ones */
-	UE_DEPRECATED(5.1, "Use AddGeneratedTilesAndGetUpdatedTiles instead")
-	NAVIGATIONSYSTEM_API TArray<uint32> AddGeneratedTiles(FRecastTileGenerator& TileGenerator);
-
 	/** Adds generated tiles to NavMesh, replacing old ones */
 	NAVIGATIONSYSTEM_API TArray<FNavTileRef> AddGeneratedTilesAndGetUpdatedTiles(FRecastTileGenerator& TileGenerator);
 
 public:
-	/** Removes all tiles at specified grid location */
-	UE_DEPRECATED(5.1, "Use RemoveTileLayersAndGetUpdatedTiles instead")
-	NAVIGATIONSYSTEM_API TArray<uint32> RemoveTileLayers(const int32 TileX, const int32 TileY, TMap<int32, dtPolyRef>* OldLayerTileIdMap = nullptr);
-
 	/** Removes all tiles at specified grid location and returns the updated FNavTileRef */
 	NAVIGATIONSYSTEM_API TArray<FNavTileRef> RemoveTileLayersAndGetUpdatedTiles(const int32 TileX, const int32 TileY, TMap<int32, dtPolyRef>* OldLayerTileIdMap = nullptr);
 
@@ -1004,9 +966,6 @@ protected:
 	//----------------------------------------------------------------------//
 	NAVIGATIONSYSTEM_API virtual uint32 LogMemUsed() const override;
 
-	UE_DEPRECATED(5.1, "Use new version with FNavTileRef")
-	NAVIGATIONSYSTEM_API void AddGeneratedTileLayer(int32 LayerIndex, FRecastTileGenerator& TileGenerator, const TMap<int32, dtPolyRef>& OldLayerTileIdMap, TArray<uint32>& OutResultTileIndices);
-
 	NAVIGATIONSYSTEM_API bool IsAllowedToAddTileLayers(const FIntPoint Tile) const;
 	NAVIGATIONSYSTEM_API void AddGeneratedTileLayer(int32 LayerIndex, FRecastTileGenerator& TileGenerator, const TMap<int32, dtPolyRef>& OldLayerTileIdMap, TArray<FNavTileRef>& OutResultTileRefs);
 
