Fix for InstanceStaticMeshComp with odd numbers of negated scale axis won't generate navmesh correctly.

Added the missing flipping indices step where the vertices of ISMComp are converted to Recast Space.

diff --git a/Engine/Source/Runtime/NavigationSystem/Private/NavMesh/RecastNavMeshGenerator.cpp b/Engine/Source/Runtime/NavigationSystem/Private/NavMesh/RecastNavMeshGenerator.cpp
--- a/Engine/Source/Runtime/NavigationSystem/Private/NavMesh/RecastNavMeshGenerator.cpp
+++ b/Engine/Source/Runtime/NavigationSystem/Private/NavMesh/RecastNavMeshGenerator.cpp
@@ -6738,15 +6738,6 @@ int32 FRecastNavMeshGenerator::GetNextPendingDirtyTileToBuild() const
 {
 	return PendingDirtyTiles.IsEmpty() ? INDEX_NONE : PendingDirtyTiles.Num() - 1;
 }
-
-// Deprecated
-TArray<uint32> FRecastNavMeshGenerator::ProcessTileTasksSync(const int32 NumTasksToProcess)
-{
-	const TArray<FNavTileRef>& TileRefs = ProcessTileTasksSyncAndGetUpdatedTiles(NumTasksToProcess);
-	TArray<uint32> TileIds;
-	FNavTileRef::DeprecatedGetTileIdsFromNavTileRefs(DestNavMesh->GetRecastNavMeshImpl(), TileRefs, TileIds);
-	return TileIds;
-}
 	
 //this code path is approx 10% faster than ProcessTileTasksSyncTimeSliced, however it spikes far worse for most use cases.
 TArray<FNavTileRef> FRecastNavMeshGenerator::ProcessTileTasksSyncAndGetUpdatedTiles(const int32 NumTasksToProcess)
