[Navmesh] Added more information to the error message and warning messages related to the tile limits.
Added more information when reporting the error message related to the TileNumberHardLimit.
Added more information when reporting the warning message related to the tile pool size.

diff --git a/Engine/Source/Runtime/NavigationSystem/Private/NavMesh/PImplRecastNavMesh.cpp b/Engine/Source/Runtime/NavigationSystem/Private/NavMesh/PImplRecastNavMesh.cpp
--- a/Engine/Source/Runtime/NavigationSystem/Private/NavMesh/PImplRecastNavMesh.cpp
+++ b/Engine/Source/Runtime/NavigationSystem/Private/NavMesh/PImplRecastNavMesh.cpp
@@ -669,7 +669,7 @@ void FPImplRecastNavMesh::Serialize( FArchive& Ar, int32 NavMeshVersion )
 					Status = DetourNavMesh->addTile(TileData, TileDataSize, DT_TILE_FREE_DATA, TileRef, NULL);
 					if (dtStatusDetail(Status, DT_OUT_OF_MEMORY))
 					{
-						UE_LOG(LogNavigation, Warning, TEXT("%s Failed to add tile (%d,%d:%d), %d tile limit reached in %s."),
+						UE_LOG(LogNavigation, Warning, TEXT("%s Failed to add tile (%d,%d:%d), %d tile limit reached in %s. If using FixedTilePoolSize, try increasing the TilePoolSize or using bigger tiles."),
 							ANSI_TO_TCHAR(__FUNCTION__), TileHeader->x, TileHeader->y, TileHeader->layer, DetourNavMesh->getMaxTiles(), *NavMeshOwner->GetFullName());
 					}
 
diff --git a/Engine/Source/Runtime/NavigationSystem/Private/NavMesh/RecastNavMeshDataChunk.cpp b/Engine/Source/Runtime/NavigationSystem/Private/NavMesh/RecastNavMeshDataChunk.cpp
--- a/Engine/Source/Runtime/NavigationSystem/Private/NavMesh/RecastNavMeshDataChunk.cpp
+++ b/Engine/Source/Runtime/NavigationSystem/Private/NavMesh/RecastNavMeshDataChunk.cpp
@@ -315,7 +315,7 @@ TArray<FNavTileRef> URecastNavMeshDataChunk::AttachTiles(ARecastNavMesh& NavMesh
 				{
 					if (dtStatusDetail(status, DT_OUT_OF_MEMORY))
 					{
-						UE_LOG(LogNavigation, Warning, TEXT("%s> Failed to add tile (%d,%d:%d), %d tile limit reached! (from: %s)"),
+						UE_LOG(LogNavigation, Warning, TEXT("%s> Failed to add tile (%d,%d:%d), %d tile limit reached! (from: %s). If using FixedTilePoolSize, try increasing the TilePoolSize or using bigger tiles."),
 							*NavMesh.GetName(), Header->x, Header->y, Header->layer, DetourNavMesh->getMaxTiles(), ANSI_TO_TCHAR(__FUNCTION__));
 					}
 					
diff --git a/Engine/Source/Runtime/NavigationSystem/Private/NavMesh/RecastNavMeshGenerator.cpp b/Engine/Source/Runtime/NavigationSystem/Private/NavMesh/RecastNavMeshGenerator.cpp
--- a/Engine/Source/Runtime/NavigationSystem/Private/NavMesh/RecastNavMeshGenerator.cpp
+++ b/Engine/Source/Runtime/NavigationSystem/Private/NavMesh/RecastNavMeshGenerator.cpp
@@ -5030,7 +5030,8 @@ void FRecastNavMeshGenerator::CalcNavMeshProperties(int32& MaxTiles, int32& MaxP
 
 	if (MaxRequestedTiles < 0 || MaxRequestedTiles > MaxTilesFromMask)
 	{
-		UE_LOG(LogNavigation, Error, TEXT("Navmesh bounds are too large! Limiting requested tiles count (%d) to: (%d) for %s"), MaxRequestedTiles, MaxTilesFromMask, *GetFullNameSafe(DestNavMesh));
+		UE_LOG(LogNavigation, Error, TEXT("Navmesh bounds are too large! Limiting requested tiles count (%d) to: (%d) for %s. To resolve this, try using bigger tiles or increasing the TileNumberHardLimit in the NavMesh properties."),
+			MaxRequestedTiles, MaxTilesFromMask, *GetFullNameSafe(DestNavMesh));
 		MaxRequestedTiles = MaxTilesFromMask;
 	}
 
@@ -5619,7 +5620,7 @@ void FRecastNavMeshGenerator::AddGeneratedTileLayer(int32 LayerIndex, FRecastTil
 			{
 				if (dtStatusDetail(status, DT_OUT_OF_MEMORY))
 				{
-					UE_LOG(LogNavigation, Error, TEXT("%s> Failed to add tile (%d,%d:%d), %d tile limit reached! (from %s)"),
+					UE_LOG(LogNavigation, Error, TEXT("%s> Failed to add tile (%d,%d:%d), %d tile limit reached! (from %s). If using FixedTilePoolSize, try increasing the TilePoolSize or using bigger tiles."),
 						*DestNavMesh->GetName(), TileX, TileY, LayerIndex, DetourMesh->getMaxTiles(), ANSI_TO_TCHAR(__FUNCTION__));
 				}
 			}
