Upgrades to SplineNavModifierComponent:
- It's now an editor only component
- Now updates any time any sibling USplineComponent changes in a meaningful way such as changes to: bIsLooping, moving/adding/removing spline points, modifying component transform, adding/removing USplineComponents from the actor
- Now lets the user assign a spline instead of arbitrarily finding the first spline referenced on a the object

#changelist validated
#virtualized

diff --git a/Engine/Source/Runtime/NavigationSystem/Private/SplineNavModifierComponent.cpp b/Engine/Source/Runtime/NavigationSystem/Private/SplineNavModifierComponent.cpp
--- a/Engine/Source/Runtime/NavigationSystem/Private/SplineNavModifierComponent.cpp
+++ b/Engine/Source/Runtime/NavigationSystem/Private/SplineNavModifierComponent.cpp
@@ -4,29 +4,13 @@
 
 #include "AI/NavigationSystemBase.h"
 #include "AI/Navigation/NavigationRelevantData.h"
-#include "Curves/BezierUtilities.h"
 #include "Components/SplineComponent.h"
-#include "VisualLogger/VisualLogger.h"
+#include "Curves/BezierUtilities.h"
 
 #include UE_INLINE_GENERATED_CPP_BY_NAME(SplineNavModifierComponent)
 
 namespace
 {
-	// Fetch the spline component from the actor
-	const USplineComponent* GetSpline(const AActor* Owner)
-	{
-		if (!Owner)
-		{
-			UE_LOG(LogNavigation, Warning, TEXT("USplineNavModifierComponent has no owner, cannot proceed"));
-			return nullptr;
-		}
-
-		const USplineComponent* Spline = Owner->GetComponentByClass<USplineComponent>();
-		UE_CVLOG_UELOG(!Spline, Owner, LogNavigation, Warning, TEXT("USplineNavModifierComponent attached to \"%s\" could not find a spline component, cannot proceed"), *Owner->GetName());
-
-		return Spline;
-	}
-
 	// Subdivide the spline into linear segments, adapting to its curvature (more curvy means more linear segments)
 	void SubdivideSpline(TArray<FVector>& OutSubdivisions, const USplineComponent& Spline, const float SubdivisionThreshold)
 	{
@@ -64,19 +48,19 @@ namespace
 
 void USplineNavModifierComponent::CalculateBounds() const
 {
-	const USplineComponent* Spline = GetSpline(GetOwner());
-	if (!Spline)
+	Bounds = FBox(ForceInit);
+
+	if (const USplineComponent* Spline = Cast<USplineComponent>(AttachedSpline.GetComponent(GetOwner())))
 	{
-		return;
+		// The largest stroke length is used to expand the bounds
+		const double Buffer = FMath::Max(StrokeWidth / 2.0, StrokeHeight / 2.0);
+		Bounds = Spline->CalcBounds(SplineTransform).GetBox().ExpandBy(Buffer);
 	}
-
-	const double Buffer = FMath::Max(StrokeWidth / 2.0, StrokeHeight / 2.0);
-	Bounds = Spline->CalcBounds(Spline->GetComponentTransform()).GetBox().ExpandBy(Buffer);
 }
 
 void USplineNavModifierComponent::GetNavigationData(FNavigationRelevantData& Data) const
 {
-	const USplineComponent* Spline = GetSpline(GetOwner());
+	const USplineComponent* Spline = Cast<USplineComponent>(AttachedSpline.GetComponent(GetOwner()));
 	if (!Spline)
 	{
 		return;
@@ -97,7 +81,7 @@ void USplineNavModifierComponent::GetNavigationData(FNavigationRelevantData& Dat
 
 	// Subdivide the spline so that high curvature sections get smaller and more linear segments than straighter sections
 	TArray<FVector> Subdivisions;
-	SubdivideSpline(Subdivisions, *Spline, GetSudivisionThreshold());
+	SubdivideSpline(Subdivisions, *Spline, GetSubdivisionThreshold());
 	const int32 NumSubdivisions = Subdivisions.Num();
 
 	// Create volumes from the spline subdivisions and use them to mark the nav mesh with the given are
@@ -108,7 +92,7 @@ void USplineNavModifierComponent::GetNavigationData(FNavigationRelevantData& Dat
 		// Compute the rotation of this tube segment
 		const double TubeAngle = (Subdivisions[SubdivisionIndex] - Subdivisions[PrevIndex]).HeadingAngle();
 		const FQuat TubeRotation(FVector::UnitZ(), TubeAngle);
-		
+
 		// Compute the vertices of this tube segment
 		for (int i = 0; i < NumCrossSectionVertices; i++)
 		{
@@ -125,7 +109,84 @@ void USplineNavModifierComponent::GetNavigationData(FNavigationRelevantData& Dat
 	}
 }
 
-float USplineNavModifierComponent::GetSudivisionThreshold() const
+#if WITH_EDITORONLY_DATA
+USplineNavModifierComponent::USplineNavModifierComponent(const FObjectInitializer& ObjectInitializer) : Super(ObjectInitializer)
+{
+	// Should tick in the editor in order to track whether the spline has updated
+	bTickInEditor = true;
+	PrimaryComponentTick.bCanEverTick = true;
+
+	// If a spline is already attached, store its update-checking data
+	if (const USplineComponent* Spline = Cast<USplineComponent>(AttachedSpline.GetComponent(GetOwner())))
+	{
+		SplineVersion = Spline->SplineCurves.Version;
+		SplineTransform = Spline->GetComponentTransform();
+	}
+}
+
+bool USplineNavModifierComponent::IsComponentTickEnabled() const
+{
+	const UWorld* World = GetWorld();
+	return World && !World->IsGameWorld();
+}
+
+void USplineNavModifierComponent::UpdateNavigationData()
+{
+	CalculateBounds();
+	FNavigationSystem::UpdateComponentData(*this);
+}
+
+void USplineNavModifierComponent::TickComponent(const float DeltaTime, const ELevelTick TickType, FActorComponentTickFunction* ThisTickFunction)
+{
+	Super::TickComponent(DeltaTime, TickType, ThisTickFunction);
+
+	if (const USplineComponent* Spline = Cast<USplineComponent>(AttachedSpline.GetComponent(GetOwner())))
+	{
+		// Update spline data, and if anything changed then update nav data 
+		if (SplineVersion != INVALID_SPLINE_VERSION)
+		{
+			bool bRequiresNavigationUpdate = false;
+	
+			const uint32 NextVersion = Spline->SplineCurves.Version;
+			if (SplineVersion != NextVersion)
+			{
+				SplineVersion = NextVersion;
+				bRequiresNavigationUpdate = true;
+			}
+
+			const FTransform& NextTransform = Spline->GetComponentTransform();
+			if (!SplineTransform.Equals(NextTransform))
+			{
+				SplineTransform = NextTransform;
+				bRequiresNavigationUpdate = true;
+			}
+
+			// This can be expensive (i.e. updating every tick as the user drags a spline point), so only update nav data if the editor flag is set
+			if (bRequiresNavigationUpdate && bUpdateNavDataOnSplineChange)
+			{
+				UpdateNavigationData();
+			}
+		}
+		else
+		{
+			// The spline just became valid; store its data and use it to update nav data
+			SplineVersion = Spline->SplineCurves.Version;
+			SplineTransform = Spline->GetComponentTransform();
+
+			UpdateNavigationData();
+		}
+	}
+	else if (SplineVersion != INVALID_SPLINE_VERSION)
+	{
+		// The spline just became invalid; reset the version and recompute nav data without the spline
+		SplineVersion = INVALID_SPLINE_VERSION;
+		UpdateNavigationData();
+	}
+}
+
+#endif // WITH_EDITORONLY_DATA
+
+float USplineNavModifierComponent::GetSubdivisionThreshold() const
 {
 	switch (SubdivisionLOD)
 	{
diff --git a/Engine/Source/Runtime/NavigationSystem/Public/SplineNavModifierComponent.h b/Engine/Source/Runtime/NavigationSystem/Public/SplineNavModifierComponent.h
--- a/Engine/Source/Runtime/NavigationSystem/Public/SplineNavModifierComponent.h
+++ b/Engine/Source/Runtime/NavigationSystem/Public/SplineNavModifierComponent.h
@@ -5,6 +5,8 @@
 
 #include "SplineNavModifierComponent.generated.h"
 
+#define INVALID_SPLINE_VERSION MIN_int32
+
 struct FNavigationRelevantData;
 
 UENUM()
@@ -17,29 +19,63 @@ enum class ESubdivisionLOD
 };
 
 /**
- *	Assign a chosen NavArea in the vicinity of a spline
- *	This component must only be attached to an actor with a USplineComponent
+ *	Used to assign a chosen NavArea to the nav mesh in the vicinity of a chosen spline.
+ *	A tube is constructed around the spline and intersected with the nav mesh. Set its dimensions with StrokeWidth and StrokeHeight.
  */
 UCLASS(Blueprintable, Meta = (BlueprintSpawnableComponent), hidecategories = (Variable, Tags, Cooking, Collision), MInimalAPI)
 class USplineNavModifierComponent : public UNavModifierComponent
 {
 	GENERATED_BODY()
 
-	/** Width of the generated tube-shaped NavModifierVolume enclosing the spline */
+	/**
+	 * If true, any changes to Spline Components on this actor will cause this component to update the nav mesh.
+	 * This will be slow if the spline has many points, or the nav mesh is sufficiently large.
+	 */ 
+	UPROPERTY(EditAnywhere, Category = Navigation)
+	bool bUpdateNavDataOnSplineChange = true;
+
+	/** The SplineComponent which will modify the nav mesh; it must also be attached to this component's owner actor */ 
+	UPROPERTY(EditAnywhere, Category = Navigation, meta = (UseComponentPicker, AllowedClasses = "/Script/Engine.SplineComponent", DisplayName = "Nav Modifier Spline"))
+	FComponentReference AttachedSpline;
+
+	/** Cross-sectional width of the tube enclosing the spline */
 	UPROPERTY(EditAnywhere, Category = Navigation, Meta=(UIMin="10", ClampMin="10"))
 	double StrokeWidth = 500.0f;
-	
-	/** Height of the generated tube-shaped nav modifier volume enclosing the spline, must be tall enough to pass through the desired surface beneath the spline */
+
+	/** Cross-sectional height of the tube enclosing the spline */
 	UPROPERTY(EditAnywhere, Category = Navigation, Meta=(UIMin="10", ClampMin="10"))
 	double StrokeHeight = 500.0f;
-		
+
 	/** Higher LOD will capture finer details in the spline */
 	UPROPERTY(EditAnywhere, Category = Navigation)
 	ESubdivisionLOD SubdivisionLOD = ESubdivisionLOD::Medium;
-	
+
+#if WITH_EDITORONLY_DATA
+	/**
+	 * Recalculates bounds, then re-computes the NavModifierVolumes and re-marks the nav mesh.
+	 * Disable UpdateNavDataOnSplineChange and use this to manually update when either the spline or nav mesh is too large to handle rapid updates.
+	 */
+	UFUNCTION(CallInEditor, Category = Navigation)
+	void UpdateNavigationData();
+
+	USplineNavModifierComponent(const FObjectInitializer& ObjectInitializer);
+#endif // WITH_EDITORONLY_DATA
+
+protected:
+#if WITH_EDITORONLY_DATA
+	virtual void TickComponent(float DeltaTime, ELevelTick TickType, FActorComponentTickFunction* ThisTickFunction) override;
+	virtual bool IsComponentTickEnabled() const override;
+#endif // WITH_EDITORONLY_DATA
+
 	virtual void CalculateBounds() const override;
 	virtual void GetNavigationData(FNavigationRelevantData& Data) const override;
 
-protected:
-	float GetSudivisionThreshold() const;
+	float GetSubdivisionThreshold() const;
+
+private:
+#if WITH_EDITORONLY_DATA
+	// Data used to check if the SplineComponent has changed
+	uint32 SplineVersion = INVALID_SPLINE_VERSION;
+	FTransform SplineTransform = FTransform();
+#endif // WITH_EDITORONLY_DATA
 };
