Added a function in RecastNavMesh that returns the surface area of a poly.

diff --git a/Engine/Source/Runtime/NavigationSystem/Private/NavMesh/PImplRecastNavMesh.cpp b/Engine/Source/Runtime/NavigationSystem/Private/NavMesh/PImplRecastNavMesh.cpp
--- a/Engine/Source/Runtime/NavigationSystem/Private/NavMesh/PImplRecastNavMesh.cpp
+++ b/Engine/Source/Runtime/NavigationSystem/Private/NavMesh/PImplRecastNavMesh.cpp
@@ -2187,6 +2187,32 @@ bool FPImplRecastNavMesh::GetRandomPointInPoly(NavNodeRef PolyID, FVector& OutPo
 	return false;
 }
 
+FVector::FReal FPImplRecastNavMesh::GetPolySurfaceArea(NavNodeRef PolyID) const
+{
+	if (DetourNavMesh)
+	{
+		dtPoly const* Poly = 0;
+		dtMeshTile const* Tile = 0;
+		dtStatus Status = DetourNavMesh->getTileAndPolyByRef((dtPolyRef)PolyID, &Tile, &Poly);
+		if (dtStatusSucceed(Status))
+		{
+			// Calc area of the polygon.
+			dtReal PolyArea = 0;
+			for (int j = 2; j < Poly->vertCount; ++j)
+			{
+				const dtReal* VA = &Tile->verts[Poly->verts[0] * 3];
+				const dtReal* VB = &Tile->verts[Poly->verts[j - 1] * 3];
+				const dtReal* VC = &Tile->verts[Poly->verts[j] * 3];
+				PolyArea += dtTriArea2D(VA, VB, VC);
+			}
+
+			return (FVector::FReal)PolyArea;
+		}
+	}
+
+	return 0;
+}
+
 uint32 FPImplRecastNavMesh::GetPolyAreaID(NavNodeRef PolyID) const
 {
 	uint32 AreaID = RECAST_NULL_AREA;
diff --git a/Engine/Source/Runtime/NavigationSystem/Private/NavMesh/RecastNavMesh.cpp b/Engine/Source/Runtime/NavigationSystem/Private/NavMesh/RecastNavMesh.cpp
--- a/Engine/Source/Runtime/NavigationSystem/Private/NavMesh/RecastNavMesh.cpp
+++ b/Engine/Source/Runtime/NavigationSystem/Private/NavMesh/RecastNavMesh.cpp
@@ -2200,6 +2200,11 @@ bool ARecastNavMesh::GetRandomPointInPoly(NavNodeRef PolyID, FVector& OutPoint)
 	return RecastNavMeshImpl && RecastNavMeshImpl->GetRandomPointInPoly(PolyID, OutPoint);
 }
 
+FVector::FReal ARecastNavMesh::GetPolySurfaceArea(NavNodeRef PolyID) const
+{
+	return RecastNavMeshImpl ? RecastNavMeshImpl->GetPolySurfaceArea(PolyID) : 0;
+}
+
 uint32 ARecastNavMesh::GetPolyAreaID(NavNodeRef PolyID) const
 {
 	uint32 AreaID = RECAST_DEFAULT_AREA;
diff --git a/Engine/Source/Runtime/NavigationSystem/Public/NavMesh/PImplRecastNavMesh.h b/Engine/Source/Runtime/NavigationSystem/Public/NavMesh/PImplRecastNavMesh.h
--- a/Engine/Source/Runtime/NavigationSystem/Public/NavMesh/PImplRecastNavMesh.h
+++ b/Engine/Source/Runtime/NavigationSystem/Public/NavMesh/PImplRecastNavMesh.h
@@ -183,6 +183,8 @@ public:
 	NAVIGATIONSYSTEM_API bool GetPolyVerts(NavNodeRef PolyID, TArray<FVector>& OutVerts) const;
 	/** Retrieves a random point inside the specified polygon. Returns false on error. */
 	NAVIGATIONSYSTEM_API bool GetRandomPointInPoly(NavNodeRef PolyID, FVector& OutPoint) const;
+	/** Retrieves the surface area of the specified polygon. Returns 0 on error. */
+	NAVIGATIONSYSTEM_API FVector::FReal GetPolySurfaceArea(NavNodeRef PolyID) const;
 	/** Retrieves the flags for the specified polygon. Returns false on error. */
 	NAVIGATIONSYSTEM_API bool GetPolyData(NavNodeRef PolyID, uint16& Flags, uint8& AreaType) const;
 	/** Retrieves area ID for the specified polygon. */
diff --git a/Engine/Source/Runtime/NavigationSystem/Public/NavMesh/RecastNavMesh.h b/Engine/Source/Runtime/NavigationSystem/Public/NavMesh/RecastNavMesh.h
--- a/Engine/Source/Runtime/NavigationSystem/Public/NavMesh/RecastNavMesh.h
+++ b/Engine/Source/Runtime/NavigationSystem/Public/NavMesh/RecastNavMesh.h
@@ -1349,6 +1349,9 @@ public:
 	/** Retrieves a random point inside the specified polygon. Returns false on error. */
 	NAVIGATIONSYSTEM_API bool GetRandomPointInPoly(NavNodeRef PolyID, FVector& OutPoint) const;
 
+	/** Retrieves the surface area of the specified polygon. Returns 0 on error. */
+	NAVIGATIONSYSTEM_API FVector::FReal GetPolySurfaceArea(NavNodeRef PolyID) const;
+
 	/** Retrieves area ID for the specified polygon. */
 	NAVIGATIONSYSTEM_API uint32 GetPolyAreaID(NavNodeRef PolyID) const;
 
