[Navigation] fix issue introduced in 32398037 causing navmesh to rebuild

diff --git a/Engine/Source/Runtime/NavigationSystem/Private/NavMesh/PImplRecastNavMesh.cpp b/Engine/Source/Runtime/NavigationSystem/Private/NavMesh/PImplRecastNavMesh.cpp
--- a/Engine/Source/Runtime/NavigationSystem/Private/NavMesh/PImplRecastNavMesh.cpp
+++ b/Engine/Source/Runtime/NavigationSystem/Private/NavMesh/PImplRecastNavMesh.cpp
@@ -606,7 +606,7 @@ void FPImplRecastNavMesh::Serialize( FArchive& Ar, int32 NavMeshVersion )
 			DefaultCellSize = NavMeshOwner->CellSize;
 		}
 		PRAGMA_ENABLE_DEPRECATION_WARNINGS
-		const FVector::FReal ActorsTileSize = (NavMeshOwner->TileSizeUU / DefaultCellSize) * DefaultCellSize;
+		const FVector::FReal ActorsTileSize = FMath::TruncToInt(NavMeshOwner->TileSizeUU / DefaultCellSize) * DefaultCellSize;
 
 		if (ActorsTileSize != Params.tileWidth)
 		{
