[NavigationSystem] updated some method display names to be consistent with other BP exposed functions

diff --git a/Engine/Source/Runtime/NavigationSystem/Public/NavigationSystem.h b/Engine/Source/Runtime/NavigationSystem/Public/NavigationSystem.h
--- a/Engine/Source/Runtime/NavigationSystem/Public/NavigationSystem.h
+++ b/Engine/Source/Runtime/NavigationSystem/Public/NavigationSystem.h
@@ -486,17 +486,17 @@ public:
 	static NAVIGATIONSYSTEM_API UNavigationSystemV1* GetNavigationSystem(UObject* WorldContextObject);
 	
 	/** Project a point onto the NavigationData */
-	UFUNCTION(BlueprintPure, Category = "AI|Navigation", meta = (WorldContext = "WorldContextObject", DisplayName = "ProjectPointToNavigation", ScriptName = "ProjectPointToNavigation"))
+	UFUNCTION(BlueprintPure, Category = "AI|Navigation", meta = (WorldContext = "WorldContextObject", DisplayName = "Project Point to Navigation", ScriptName = "ProjectPointToNavigation"))
 	static NAVIGATIONSYSTEM_API bool K2_ProjectPointToNavigation(UObject* WorldContextObject, const FVector& Point, FVector& ProjectedLocation, ANavigationData* NavData, TSubclassOf<UNavigationQueryFilter> FilterClass, const FVector QueryExtent = FVector::ZeroVector);
 
 	/** Generates a random location reachable from given Origin location.
 	 *	@return Return Value represents if the call was successful */
-	UFUNCTION(BlueprintPure, Category = "AI|Navigation", meta = (WorldContext = "WorldContextObject", DisplayName = "GetRandomReachablePointInRadius", ScriptName = "GetRandomReachablePointInRadius"))
+	UFUNCTION(BlueprintPure, Category = "AI|Navigation", meta = (WorldContext = "WorldContextObject", DisplayName = "Get Random Reachable Point in Radius", ScriptName = "GetRandomReachablePointInRadius"))
 	static NAVIGATIONSYSTEM_API bool K2_GetRandomReachablePointInRadius(UObject* WorldContextObject, const FVector& Origin, FVector& RandomLocation, float Radius, ANavigationData* NavData = NULL, TSubclassOf<UNavigationQueryFilter> FilterClass = NULL);
 
 	/** Generates a random location in navigable space within given radius of Origin.
 	 *	@return Return Value represents if the call was successful */
-	UFUNCTION(BlueprintCallable, Category = "AI|Navigation", meta = (WorldContext = "WorldContextObject", DisplayName = "GetRandomLocationInNavigableRadius", ScriptName = "GetRandomLocationInNavigableRadius"))
+	UFUNCTION(BlueprintCallable, Category = "AI|Navigation", meta = (WorldContext = "WorldContextObject", DisplayName = "Get Random Location in Navigable Radius", ScriptName = "GetRandomLocationInNavigableRadius"))
 	static NAVIGATIONSYSTEM_API bool K2_GetRandomLocationInNavigableRadius(UObject* WorldContextObject, const FVector& Origin, FVector& RandomLocation, float Radius, ANavigationData* NavData = NULL, TSubclassOf<UNavigationQueryFilter> FilterClass = NULL);
 	
 	/** Potentially expensive. Use with caution. Consider using UPathFollowingComponent::GetRemainingPathCost instead */
@@ -564,7 +564,7 @@ public:
 	UFUNCTION(BlueprintCallable, Category = "AI|Navigation|Generation")
 	NAVIGATIONSYSTEM_API void SetGeometryGatheringMode(ENavDataGatheringModeConfig NewMode);
 
-	UFUNCTION(BlueprintCallable, Category = "AI|Navigation", meta=(DisplayName="ReplaceAreaInOctreeData"))
+	UFUNCTION(BlueprintCallable, Category = "AI|Navigation", meta=(DisplayName="Replace Area in Octree Data"))
 	NAVIGATIONSYSTEM_API bool K2_ReplaceAreaInOctreeData(const UObject* Object, TSubclassOf<UNavArea> OldArea, TSubclassOf<UNavArea> NewArea);
 
 	FORCEINLINE bool IsActiveTilesGenerationEnabled() const{ return bGenerateNavigationOnlyAroundNavigationInvokers; }
