Add way to test if we have complete nav data ready around a segment (instead of a single point)

diff --git a/Engine/Source/Runtime/NavigationSystem/Private/NavMesh/NavTestRenderingComponent.cpp b/Engine/Source/Runtime/NavigationSystem/Private/NavMesh/NavTestRenderingComponent.cpp
--- a/Engine/Source/Runtime/NavigationSystem/Private/NavMesh/NavTestRenderingComponent.cpp
+++ b/Engine/Source/Runtime/NavigationSystem/Private/NavMesh/NavTestRenderingComponent.cpp
@@ -107,6 +107,21 @@ void FNavTestSceneProxy::GetDynamicMeshElements(const TArray<const FSceneView*>&
 						NavTestActor->RadiusUsedToValidateNavData, HalfHeight,  
 						NumSides, SDPG_World);
 				}
+
+				if (NavTestActor->bDrawIfNavDataIsReadyToQueryTargetActor && NavTestActor->QueryTargetActor)
+				{
+					const FVector QueryTargetLocation = NavTestActor->QueryTargetActor->GetActorLocation();
+					float Distance;
+					FVector Forward, Right, Up;
+					(QueryTargetLocation-ActorLocation).ToDirectionAndLength(Up, Distance);
+					Up.FindBestAxisVectors(Forward, Right);
+
+					constexpr int32 NumSides = 32;
+					DrawWireCapsule(PDI, ActorLocation + 0.5f * Distance * Up, Forward, Right, Up,
+						NavTestActor->bNavDataIsReadyToQueryTargetActor ? FColor::Green : FColor::Red, 
+						NavTestActor->RadiusUsedToValidateNavData, 0.5f*Distance + NavTestActor->RadiusUsedToValidateNavData,  
+						NumSides, SDPG_World);
+				}
 			}
 
 			// draw path
diff --git a/Engine/Source/Runtime/NavigationSystem/Private/NavMesh/RecastNavMesh.cpp b/Engine/Source/Runtime/NavigationSystem/Private/NavMesh/RecastNavMesh.cpp
--- a/Engine/Source/Runtime/NavigationSystem/Private/NavMesh/RecastNavMesh.cpp
+++ b/Engine/Source/Runtime/NavigationSystem/Private/NavMesh/RecastNavMesh.cpp
@@ -6,6 +6,7 @@
 #include "Engine/World.h"
 #include "NavigationSystem.h"
 #include "Engine/Engine.h"
+#include "Distance/DistSegment2AxisAlignedBox2.h"
 #include "DrawDebugHelpers.h"
 #include "Misc/ConfigCacheIni.h"
 #include "EngineUtils.h"
@@ -3507,6 +3508,22 @@ bool ARecastNavMesh::HasValidNavmesh() const
 }
 
 #if WITH_RECAST
+bool ARecastNavMesh::HasCompleteDataInTile(const int32 TileX, const int32 TileY) const
+{
+	check(RecastNavMeshImpl && RecastNavMeshImpl->DetourNavMesh);
+	const int32 NumTiles = RecastNavMeshImpl->DetourNavMesh->getTileCountAt(TileX, TileY);
+	if (NumTiles <= 0)
+	{
+		const bool bHasFailsafeData = bStoreEmptyTileLayers && RecastNavMeshImpl->HasTileCacheLayers(TileX, TileY);
+		if (!bHasFailsafeData)
+		{
+			return false;
+		}
+	}
+
+	return true;
+}
+
 bool ARecastNavMesh::HasCompleteDataInRadius(const FVector& TestLocation, FVector::FReal TestRadius) const
 {
 	if (HasValidNavmesh() == false)
@@ -3541,14 +3558,53 @@ bool ARecastNavMesh::HasCompleteDataInRadius(const FVector& TestLocation, FVecto
 			const bool bInside = FMath::SphereAABBIntersection(RcTestLocation, RadiusSq, FBox::BuildAABB(RcTileCenter, RcTileExtent2D));
 			if (bInside)
 			{
-				const int32 NumTiles = NavMesh->getTileCountAt(TileX, TileY);
-				if (NumTiles <= 0)
+				if (!HasCompleteDataInTile(TileX, TileY))
 				{
-					const bool bHasFailsafeData = bStoreEmptyTileLayers && RecastNavMeshImpl->HasTileCacheLayers(TileX, TileY);
-					if (!bHasFailsafeData)
-					{
-						return false;
-					}
+					return false;
+				}
+			}
+		}
+	}
+
+	return true;
+}
+
+bool ARecastNavMesh::HasCompleteDataAroundSegment(const FVector& StartLocation, const FVector& EndLocation, FVector::FReal TestRadius) const
+{
+	if (HasValidNavmesh() == false)
+	{
+		return false;
+	}
+
+	const dtNavMeshParams* NavParams = RecastNavMeshImpl->DetourNavMesh->getParams();
+	const float NavTileSize = GetTileSizeUU();
+	const FVector RcNavOrigin(NavParams->orig[0], NavParams->orig[1], NavParams->orig[2]);
+
+	FBox CapsuleBox(ForceInit);
+	CapsuleBox += StartLocation;
+	CapsuleBox += EndLocation;
+	CapsuleBox = CapsuleBox.ExpandBy(FVector(TestRadius, TestRadius, 0));
+	const FRcTileBox TileBox(CapsuleBox, RcNavOrigin, NavTileSize);
+
+	const FVector RcStartLocation = Unreal2RecastPoint(StartLocation);
+	const FVector RcEndLocation = Unreal2RecastPoint(EndLocation);
+	const UE::Geometry::TSegment2<FVector::FReal> Segment (FVector2D(RcStartLocation.X, RcStartLocation.Z), FVector2D(RcEndLocation.X, RcEndLocation.Z));
+	const float TestRadiusSquared = FMath::Square(TestRadius);
+
+	for (int32 TileX = TileBox.XMin; TileX <= TileBox.XMax; TileX++)
+	{
+		for (int32 TileY = TileBox.YMin; TileY <= TileBox.YMax; TileY++)
+		{
+			const FVector2D RcTileMin(RcNavOrigin.X + (TileX * NavTileSize), RcNavOrigin.Z + (TileY * NavTileSize));
+			const FBox2D RcTileAABB = FBox2D(RcTileMin, RcTileMin + FVector2D(NavTileSize, NavTileSize));
+
+			UE::Geometry::TDistSegment2AxisAlignedBox2<FVector::FReal> SegmentTileDist(Segment, RcTileAABB);
+			const float DistanceSquared = SegmentTileDist.GetSquared();
+			if (DistanceSquared <= TestRadiusSquared)
+			{
+				if (!HasCompleteDataInTile(TileX, TileY))
+				{
+					return false;
 				}
 			}
 		}
diff --git a/Engine/Source/Runtime/NavigationSystem/Private/NavigationTestingActor.cpp b/Engine/Source/Runtime/NavigationSystem/Private/NavigationTestingActor.cpp
--- a/Engine/Source/Runtime/NavigationSystem/Private/NavigationTestingActor.cpp
+++ b/Engine/Source/Runtime/NavigationSystem/Private/NavigationTestingActor.cpp
@@ -56,6 +56,7 @@ ANavigationTestingActor::ANavigationTestingActor(const FObjectInitializer& Objec
 	bDrawDistanceToWall = false;
 	ClosestWallLocation = FNavigationSystem::InvalidLocation;
 	bNavDataIsReadyInRadius = false;
+	bNavDataIsReadyToQueryTargetActor = false;
 	OffsetFromCornersDistance = 0.f;
 
 	QueryingExtent = FVector(DEFAULT_NAV_QUERY_EXTENT_HORIZONTAL, DEFAULT_NAV_QUERY_EXTENT_HORIZONTAL, DEFAULT_NAV_QUERY_EXTENT_VERTICAL);
@@ -91,6 +92,13 @@ void ANavigationTestingActor::BeginDestroy()
 		OtherActor->OtherActor = NULL;
 		OtherActor->LastPath.Reset();
 	}
+	if (QueryTargetActor)
+	{
+		if (USceneComponent* RootComp = QueryTargetActor->GetRootComponent())
+		{
+			RootComp->TransformUpdated.RemoveAll(this);
+		}
+	}
 	Super::BeginDestroy();
 }
 
@@ -98,16 +106,28 @@ void ANavigationTestingActor::BeginDestroy()
 void ANavigationTestingActor::PreEditChange(FProperty* PropertyThatWillChange)
 {
 	static const FName NAME_OtherActor = GET_MEMBER_NAME_CHECKED(ANavigationTestingActor, OtherActor);
+	static const FName NAME_QueryTargetActor = GET_MEMBER_NAME_CHECKED(ANavigationTestingActor, QueryTargetActor);
 
-	if (PropertyThatWillChange && PropertyThatWillChange->GetFName() == NAME_OtherActor && OtherActor && OtherActor->OtherActor == this)
+	if (PropertyThatWillChange)
 	{
-		OtherActor->OtherActor = NULL;
-		OtherActor->LastPath.Reset();
-		LastPath.Reset();
-#if WITH_EDITORONLY_DATA
-		OtherActor->EdRenderComp->MarkRenderStateDirty();
-		EdRenderComp->MarkRenderStateDirty();
-#endif
+		const FName ChangedPropName = PropertyThatWillChange->GetFName();
+		if (ChangedPropName == NAME_OtherActor && OtherActor && OtherActor->OtherActor == this)
+		{
+			OtherActor->OtherActor = NULL;
+			OtherActor->LastPath.Reset();
+			LastPath.Reset();
+	#if WITH_EDITORONLY_DATA
+			OtherActor->EdRenderComp->MarkRenderStateDirty();
+			EdRenderComp->MarkRenderStateDirty();
+	#endif
+		}
+		else if (ChangedPropName == NAME_QueryTargetActor && QueryTargetActor)
+		{
+			if (USceneComponent* RootComp = QueryTargetActor->GetRootComponent())
+			{
+				RootComp->TransformUpdated.RemoveAll(this);
+			}
+		}
 	}
 
 	Super::PreEditChange(PropertyThatWillChange);
@@ -117,6 +137,7 @@ void ANavigationTestingActor::PostEditChangeProperty(FPropertyChangedEvent& Prop
 {
 	static const FName NAME_ShouldBeVisibleInGame = GET_MEMBER_NAME_CHECKED(ANavigationTestingActor, bShouldBeVisibleInGame);
 	static const FName NAME_OtherActor = GET_MEMBER_NAME_CHECKED(ANavigationTestingActor, OtherActor);
+	static const FName NAME_QueryTargetActor = GET_MEMBER_NAME_CHECKED(ANavigationTestingActor, QueryTargetActor);
 	static const FName NAME_IsSearchStart = GET_MEMBER_NAME_CHECKED(ANavigationTestingActor, bSearchStart);
 	static const FName NAME_InvokerComponent = GET_MEMBER_NAME_CHECKED(ANavigationTestingActor, InvokerComponent);
 
@@ -128,7 +149,8 @@ void ANavigationTestingActor::PostEditChangeProperty(FPropertyChangedEvent& Prop
 		const FName ChangedCategory = FObjectEditorUtils::GetCategoryFName(PropertyChangedEvent.Property);
 
 		if (ChangedPropName == GET_MEMBER_NAME_CHECKED(FNavAgentProperties,AgentRadius) ||
-			ChangedPropName == GET_MEMBER_NAME_CHECKED(FNavAgentProperties,AgentHeight))
+			ChangedPropName == GET_MEMBER_NAME_CHECKED(FNavAgentProperties,AgentHeight) ||
+			ChangedPropName == GET_MEMBER_NAME_CHECKED(FNavAgentProperties,PreferredNavData))
 		{
 			MyNavData = NULL;
 			UpdateNavData();
@@ -159,20 +181,26 @@ void ANavigationTestingActor::PostEditChangeProperty(FPropertyChangedEvent& Prop
 		}
 		else if (ChangedCategory == TEXT("Query"))
 		{
+			if (ChangedPropName == NAME_QueryTargetActor && QueryTargetActor)
+			{
+				if (USceneComponent* RootComp = QueryTargetActor->GetRootComponent())
+				{
+					RootComp->TransformUpdated.AddUObject(this, &ANavigationTestingActor::OnQueryTargetActorTransformUpdated);
+				}
+			}
+
 			if (bDrawDistanceToWall)
 			{
 				ClosestWallLocation = FindClosestWallLocation();
 			}
-			else if (bDrawIfNavDataIsReadyInRadius)
+			if (bDrawIfNavDataIsReadyInRadius)
 			{
 				bNavDataIsReadyInRadius = CheckIfNavDataIsReadyInRadius();
 			}
-#if WITH_EDITORONLY_DATA
-			else
+			if (bDrawIfNavDataIsReadyToQueryTargetActor)
 			{
-				EdRenderComp->MarkRenderStateDirty();
+				bNavDataIsReadyToQueryTargetActor = CheckIfNavDataIsReadyToActor(QueryTargetActor);
 			}
-#endif
 		}
 		else if (ChangedCategory == TEXT("Pathfinding"))
 		{
@@ -260,6 +288,11 @@ void ANavigationTestingActor::PostEditMove(bool bFinished)
 		{
 			bNavDataIsReadyInRadius = CheckIfNavDataIsReadyInRadius();
 		}
+
+		if (bDrawIfNavDataIsReadyToQueryTargetActor)
+		{
+			bNavDataIsReadyToQueryTargetActor = CheckIfNavDataIsReadyToActor(QueryTargetActor);
+		}
 	}
 }
 
@@ -269,6 +302,14 @@ void ANavigationTestingActor::PostLoad()
 
 	InvokerComponent->bAutoActivate = bActAsNavigationInvoker;
 
+	if (QueryTargetActor)
+	{
+		if (USceneComponent* RootComp = QueryTargetActor->GetRootComponent())
+		{
+			RootComp->TransformUpdated.AddUObject(this, &ANavigationTestingActor::OnQueryTargetActorTransformUpdated);
+		}
+	}
+
 #if WITH_RECAST && WITH_EDITORONLY_DATA
 	if (GIsEditor)
 	{
@@ -424,6 +465,35 @@ bool ANavigationTestingActor::CheckIfNavDataIsReadyInRadius()
 	return false;
 }
 
+bool ANavigationTestingActor::CheckIfNavDataIsReadyToActor(const AActor* TargetActor)
+{
+#if WITH_EDITORONLY_DATA
+	if (EdRenderComp)
+	{
+		EdRenderComp->MarkRenderStateDirty();
+	}
+#endif // WITH_EDITORONLY_DATA
+
+#if WITH_RECAST
+	UpdateNavData();
+	const ARecastNavMesh* RecastNavMesh = Cast<ARecastNavMesh>(MyNavData);
+	if (RecastNavMesh && TargetActor)
+	{
+		return RecastNavMesh->HasCompleteDataAroundSegment(GetActorLocation(), TargetActor->GetActorLocation(), RadiusUsedToValidateNavData);
+	}
+#endif // WITH_RECAST
+
+	return false;
+}
+
+void ANavigationTestingActor::OnQueryTargetActorTransformUpdated(USceneComponent* InRootComponent, EUpdateTransformFlags UpdateTransformFlags, ETeleportType Teleport)
+{
+	if (bDrawIfNavDataIsReadyToQueryTargetActor)
+	{
+		bNavDataIsReadyToQueryTargetActor = CheckIfNavDataIsReadyToActor(QueryTargetActor);
+	}
+}
+
 void ANavigationTestingActor::SearchPathTo(ANavigationTestingActor* Goal)
 {
 #if WITH_EDITORONLY_DATA
diff --git a/Engine/Source/Runtime/NavigationSystem/Public/NavMesh/RecastNavMesh.h b/Engine/Source/Runtime/NavigationSystem/Public/NavMesh/RecastNavMesh.h
--- a/Engine/Source/Runtime/NavigationSystem/Public/NavMesh/RecastNavMesh.h
+++ b/Engine/Source/Runtime/NavigationSystem/Public/NavMesh/RecastNavMesh.h
@@ -1440,6 +1440,11 @@ public:
 	  * @returns true if ALL tiles inside are ready
 	  */
 	NAVIGATIONSYSTEM_API bool HasCompleteDataInRadius(const FVector& TestLocation, FVector::FReal TestRadius) const;
+	
+	/** Check if navmesh is defined (either built/streamed or recognized as empty tile by generator) within given radius around the given segment.
+	* @returns true if ALL tiles inside are ready
+	*/
+	NAVIGATIONSYSTEM_API bool HasCompleteDataAroundSegment(const FVector& StartLocation, const FVector& EndLocation, FVector::FReal TestRadius) const;
 
 	/** @return true is specified segment is fully on navmesh (respecting the optional filter) */
 	NAVIGATIONSYSTEM_API bool IsSegmentOnNavmesh(const FVector& SegmentStart, const FVector& SegmentEnd, FSharedConstNavQueryFilter Filter = NULL, const UObject* Querier = NULL) const;
@@ -1527,6 +1532,9 @@ private:
 	/** @return Navmesh data chunk that belongs to this actor */
 	NAVIGATIONSYSTEM_API URecastNavMeshDataChunk* GetNavigationDataChunk(const ANavigationDataChunkActor& InActor) const;
 
+	/** Check if navmesh is defined (either built/streamed or recognized as empty tile by generator) in given tile */
+	bool HasCompleteDataInTile(const int32 TileX, const int32 TileY) const;
+
 protected:
 	// retrieves RecastNavMeshImpl
 	FPImplRecastNavMesh* GetRecastNavMeshImpl() { return RecastNavMeshImpl; }
diff --git a/Engine/Source/Runtime/NavigationSystem/Public/NavigationTestingActor.h b/Engine/Source/Runtime/NavigationSystem/Public/NavigationTestingActor.h
--- a/Engine/Source/Runtime/NavigationSystem/Public/NavigationTestingActor.h
+++ b/Engine/Source/Runtime/NavigationSystem/Public/NavigationTestingActor.h
@@ -119,7 +119,15 @@ public:
 	/** If set, a cylinder is drawn to indicate if the navigation data is ready (has been generated) for the given radius (green when ready, red otherwise). */
 	UPROPERTY(EditAnywhere, Category=Query)
 	uint32 bDrawIfNavDataIsReadyInRadius : 1;
-	
+
+	/** If set, a capsule is drawn to indicate if the navigation data is ready (has been generated) for the given radius from the current actor to the query target (green when ready, red otherwise). */
+	UPROPERTY(EditAnywhere, Category=Query)
+	uint32 bDrawIfNavDataIsReadyToQueryTargetActor : 1;
+
+	/** Actor to use as a target for navigation data queries */
+	UPROPERTY(EditAnywhere, Category=Query)
+	TObjectPtr<AActor> QueryTargetActor;
+
 	/** show polys from open (orange) and closed (yellow) sets */
 	UPROPERTY(EditAnywhere, Category=Debug)
 	uint32 bShowNodePool : 1;
@@ -183,6 +191,7 @@ public:
 	FVector ClosestWallLocation;
 	
 	bool bNavDataIsReadyInRadius;
+	bool bNavDataIsReadyToQueryTargetActor;
 
 #if WITH_RECAST && WITH_EDITORONLY_DATA
 	/** detail data gathered from each step of regular A* algorithm */
@@ -244,4 +253,6 @@ public:
 protected:
 	NAVIGATIONSYSTEM_API FVector FindClosestWallLocation() const;
 	bool CheckIfNavDataIsReadyInRadius();
+	bool CheckIfNavDataIsReadyToActor(const AActor* TargetActor);
+	void OnQueryTargetActorTransformUpdated(USceneComponent* InRootComponent, EUpdateTransformFlags UpdateTransformFlags, ETeleportType Teleport);
 };
