PR #11558: Fixed RecastNavMeshGenerator to use TotalNavBounds when IsValid is 'true'

diff --git a/Engine/Source/Runtime/NavigationSystem/Private/NavMesh/RecastNavMeshGenerator.cpp b/Engine/Source/Runtime/NavigationSystem/Private/NavMesh/RecastNavMeshGenerator.cpp
--- a/Engine/Source/Runtime/NavigationSystem/Private/NavMesh/RecastNavMeshGenerator.cpp
+++ b/Engine/Source/Runtime/NavigationSystem/Private/NavMesh/RecastNavMeshGenerator.cpp
@@ -4921,7 +4921,7 @@ void FRecastNavMeshGenerator::UpdateNavigationBounds()
 		{
 			InclusionBounds.Reset(1);
 			TotalNavBounds = NavSys->GetWorldBounds();
-			if (!TotalNavBounds.IsValid)
+			if (TotalNavBounds.IsValid)
 			{
 				InclusionBounds.Add(TotalNavBounds);
 			}
