Stop trying to interpret very large position/rotation/scale values on a Nav Mesh. Instead, forcibly truncate those values to fit inside a 32 bit integer.

diff --git a/Engine/Source/Runtime/NavigationSystem/Public/NavMesh/RecastNavMeshGenerator.h b/Engine/Source/Runtime/NavigationSystem/Public/NavMesh/RecastNavMeshGenerator.h
--- a/Engine/Source/Runtime/NavigationSystem/Public/NavMesh/RecastNavMeshGenerator.h
+++ b/Engine/Source/Runtime/NavigationSystem/Public/NavMesh/RecastNavMeshGenerator.h
@@ -5,6 +5,7 @@
 #include "CoreMinimal.h"
 #endif
 #include "Stats/Stats.h"
+#include "AI/NavigationSystemBase.h"
 #include "AI/Navigation/NavigationTypes.h"
 #include "AI/Navigation/NavigationInvokerPriority.h"
 #include "AI/Navigation/NavigationRelevantData.h"
@@ -222,16 +223,16 @@ struct FRcTileBox
 			if (FMath::Modf(MaxAsFloat, &UnusedIntPart) == 0)
 			{
 				// Return the lower tile
-				return FMath::Max(IntCastChecked<int32>(FMath::FloorToInt(MaxAsFloat) - 1), MinCoord);
+				return FMath::Max(ClampToInt32(FMath::FloorToInt(MaxAsFloat) - 1), MinCoord);
 			}
 			// Otherwise use default behaviour
-			return IntCastChecked<int32>(FMath::FloorToInt(MaxAsFloat));
+			return ClampToInt32(FMath::FloorToInt(MaxAsFloat));
 		};
 
 		const FBox RcAreaBounds = Unreal2RecastBox(UnrealBounds);
-		XMin = IntCastChecked<int32>(FMath::FloorToInt((RcAreaBounds.Min.X - RcNavMeshOrigin.X) / TileSizeInWorldUnits));
+		XMin = ClampToInt32(FMath::FloorToInt((RcAreaBounds.Min.X - RcNavMeshOrigin.X) / TileSizeInWorldUnits));
 		XMax = CalcMaxCoordExclusive((RcAreaBounds.Max.X - RcNavMeshOrigin.X) / TileSizeInWorldUnits, XMin);
-		YMin = IntCastChecked<int32>(FMath::FloorToInt((RcAreaBounds.Min.Z - RcNavMeshOrigin.Z) / TileSizeInWorldUnits));
+		YMin = ClampToInt32(FMath::FloorToInt((RcAreaBounds.Min.Z - RcNavMeshOrigin.Z) / TileSizeInWorldUnits));
 		YMax = CalcMaxCoordExclusive((RcAreaBounds.Max.Z - RcNavMeshOrigin.Z) / TileSizeInWorldUnits, YMin);
 	}
 
@@ -240,6 +241,16 @@ struct FRcTileBox
 		return Point.X >= XMin && Point.X <= XMax
 			&& Point.Y >= YMin && Point.Y <= YMax;
 	}
+
+	static FORCEINLINE int32 ClampToInt32(const int64 Value)
+	{
+#if !NO_LOGGING
+		UE_CLOG(!IntFitsIn<int32>(Value), LogNavigation, Warning,
+			TEXT("FRcTileBox clamped a NavMesh transform value to fit in int32. Old value: %" UINT64_FMT), Value);
+#endif // !NO_LOGGING
+
+		return static_cast<int32>(FMath::Clamp(Value, MIN_int32, MAX_int32));
+	}
 };
 
 /**
