log actor and actor label to LogNavigationDirtyArea logs when they dirty a nav tile

diff --git a/Engine/Source/Runtime/NavigationSystem/Private/NavigationDirtyAreasController.cpp b/Engine/Source/Runtime/NavigationSystem/Private/NavigationDirtyAreasController.cpp
--- a/Engine/Source/Runtime/NavigationSystem/Private/NavigationDirtyAreasController.cpp
+++ b/Engine/Source/Runtime/NavigationSystem/Private/NavigationDirtyAreasController.cpp
@@ -184,11 +184,29 @@ void FNavigationDirtyAreasController::AddAreas(const TConstArrayView<FBox> NewAr
 					ObjectInfo = FString::Printf(TEXT(" | Element's owner: %s"), *GetFullNameSafe(ObjectOwner));
 				}
 
-				return FString::Printf(TEXT("From: %s | Object: %s %s | Bounds: %s"),
+				// attempt to find the actor which is/contains SourceObject
+				const AActor* Actor = Cast<AActor>(SourceObject);
+				{
+					const UObject* CurrentObject = SourceObject;
+					while (!Actor && CurrentObject)
+					{
+						CurrentObject = CurrentObject->GetOuter();
+						Actor = Cast<AActor>(CurrentObject);
+					}
+				}
+				FString ActorInfo;
+				if (Actor)
+				{
+					// useful to have actor label for those placed in editor
+					ActorInfo = FString::Printf(TEXT("| Actor: %s | Actor label: %s"), *GetFullNameSafe(Actor), *Actor->GetActorNameOrLabel());
+				}
+
+				return FString::Printf(TEXT("From: %s | Object: %s %s | Bounds: %s %s"),
 					*DebugReason.ToString(),
 					*GetFullNameSafe(SourceObject),
 					*ObjectInfo,
-					*BoundsSize.ToString());
+					*BoundsSize.ToString(),
+					*ActorInfo);
 		};
 
 		if (ShouldReportOversizedDirtyArea() && BoundsSize.GetMax() > DirtyAreaWarningSizeThreshold)
