Fix crash in ARecastNavMesh::FindPath when NavData is missing from the query.

diff --git a/Engine/Source/Runtime/NavigationSystem/Private/NavMesh/RecastNavMesh.cpp b/Engine/Source/Runtime/NavigationSystem/Private/NavMesh/RecastNavMesh.cpp
--- a/Engine/Source/Runtime/NavigationSystem/Private/NavMesh/RecastNavMesh.cpp
+++ b/Engine/Source/Runtime/NavigationSystem/Private/NavMesh/RecastNavMesh.cpp
@@ -2944,10 +2944,14 @@ FPathFindingResult ARecastNavMesh::FindPath(const FNavAgentProperties& AgentProp
 	CSV_SCOPED_TIMING_STAT_EXCLUSIVE(Pathfinding);
 
 	const ANavigationData* Self = Query.NavData.Get();
+	if (Self == nullptr)
+	{
+		return ENavigationQueryResult::Error;
+	}
+	
 	check(Cast<const ARecastNavMesh>(Self));
-
 	const ARecastNavMesh* RecastNavMesh = (const ARecastNavMesh*)Self;
-	if (Self == NULL || RecastNavMesh->RecastNavMeshImpl == NULL)
+	if (RecastNavMesh->RecastNavMeshImpl == nullptr)
 	{
 		return ENavigationQueryResult::Error;
 	}
