[NavMesh] NavLink generation improvements:
Refactor to keep solidHF and compactHF for the lifetime of FRecastTileGenerator
Allowed the sliced generator to build links
Fix issue with missing solidHF and compactHF when building links
Added experimental link builder configurations to RecastNavmesh
Increase heightfield border size when building links
Fix missing input edges on tile borders when generating links
Changed ARecastNavMesh "Draw Failed NavLinks" debug name to Draw Failed and Valid Links"
Added profile tags inside link generation
Small navlink geneneration debug display improvements

diff --git a/Engine/Source/Runtime/Engine/Public/EngineStats.h b/Engine/Source/Runtime/Engine/Public/EngineStats.h
--- a/Engine/Source/Runtime/Engine/Public/EngineStats.h
+++ b/Engine/Source/Runtime/Engine/Public/EngineStats.h
@@ -141,6 +141,8 @@ DECLARE_CYCLE_STAT_EXTERN(TEXT("Recast: build tile cache"),STAT_Navigation_Recas
 DECLARE_CYCLE_STAT_EXTERN(TEXT("Recast: build navmesh"),STAT_Navigation_RecastBuildNavigation,STATGROUP_Navigation, );
 DECLARE_CYCLE_STAT_EXTERN(TEXT("Recast: build navigation layer"), STAT_Navigation_GenerateNavigationDataLayer, STATGROUP_Navigation, );
 DECLARE_CYCLE_STAT_EXTERN(TEXT("Recast: build links"), STAT_Navigation_RecastBuildLinks, STATGROUP_Navigation, );
+DECLARE_CYCLE_STAT_EXTERN(TEXT("Recast: build links, find edges"), STAT_Navigation_RecastBuildLinks_FindEdges, STATGROUP_Navigation, );
+DECLARE_CYCLE_STAT_EXTERN(TEXT("Recast: build links, sample"), STAT_Navigation_RecastBuildLinks_Sample, STATGROUP_Navigation, );
 DECLARE_CYCLE_STAT_EXTERN(TEXT("Recast: build regions"), STAT_Navigation_RecastBuildRegions, STATGROUP_Navigation, );
 DECLARE_CYCLE_STAT_EXTERN(TEXT("Recast: mark areas"), STAT_Navigation_RecastMarkAreas, STATGROUP_Navigation, );
 DECLARE_CYCLE_STAT_EXTERN(TEXT("Recast: build contours"), STAT_Navigation_RecastBuildContours, STATGROUP_Navigation, );
diff --git a/Engine/Source/Runtime/NavigationSystem/Private/NavMesh/LinkGenerationConfig.cpp b/Engine/Source/Runtime/NavigationSystem/Private/NavMesh/LinkGenerationConfig.cpp
new file mode 100644
--- /dev/null
+++ b/Engine/Source/Runtime/NavigationSystem/Private/NavMesh/LinkGenerationConfig.cpp
@@ -0,0 +1,24 @@
+﻿// Copyright Epic Games, Inc. All Rights Reserved.
+
+#include "NavMesh/LinkGenerationConfig.h"
+
+#if WITH_RECAST
+#include "Detour/DetourNavLinkBuilderConfig.h"
+
+void FNavLinkGenerationJumpDownConfig::CopyToDetourConfig(dtNavLinkBuilderJumpDownConfig& OutDetourConfig) const
+{
+	OutDetourConfig.jumpLength = JumpLength;
+	OutDetourConfig.jumpDistanceFromEdge = JumpDistanceFromEdge;
+	OutDetourConfig.jumpMaxDepth = JumpMaxDepth;
+	OutDetourConfig.jumpEndsHeightTolerance	= JumpEndsHeightTolerance;
+}
+
+void FNavLinkGenerationJumpOverConfig::CopyToDetourConfig(dtNavLinkBuilderJumpOverConfig& OutDetourConfig) const
+{
+	OutDetourConfig.jumpLength = JumpLength;
+	OutDetourConfig.jumpDistanceFromEdge = JumpDistanceFromEdge;
+	OutDetourConfig.jumpHeight = JumpHeight;
+	OutDetourConfig.jumpHeightTolerance = JumpHeightTolerance;
+	OutDetourConfig.jumpEndsHeightTolerance	= JumpEndsHeightTolerance;
+}
+#endif //WITH_RECAST
diff --git a/Engine/Source/Runtime/NavigationSystem/Private/NavMesh/RecastNavMeshGenerator.cpp b/Engine/Source/Runtime/NavigationSystem/Private/NavMesh/RecastNavMeshGenerator.cpp
--- a/Engine/Source/Runtime/NavigationSystem/Private/NavMesh/RecastNavMeshGenerator.cpp
+++ b/Engine/Source/Runtime/NavigationSystem/Private/NavMesh/RecastNavMeshGenerator.cpp
@@ -298,6 +298,7 @@ FRecastGeometryCache::FRecastGeometryCache(const uint8* Memory)
 
 namespace RecastGeometryExport {
 
+#if SHOW_NAV_EXPORT_PREVIEW
 static UWorld* FindEditorWorld()
 {
 	if (GEngine)
@@ -313,6 +314,7 @@ static UWorld* FindEditorWorld()
 
 	return NULL;
 }
+#endif //SHOW_NAV_EXPORT_PREVIEW
 
 static void StoreCollisionCache(FRecastGeometryExport& GeomExport)
 {
@@ -1665,10 +1667,16 @@ FRecastTileGenerator::FRecastTileGenerator(FRecastNavMeshGenerator& ParentGenera
 
 	check(ParentGenerator.GetOwner());
 	TileTimeSliceSettings.FilterLedgeSpansMaxYProcess = ParentGenerator.GetOwner()->TimeSliceFilterLedgeSpansMaxYProcess;
+
+	SolidHF = nullptr;
+	CompactHF = nullptr;
 }
 
 FRecastTileGenerator::~FRecastTileGenerator()
 {
+	rcFreeHeightField(SolidHF);
+	rcFreeCompactHeightfield(CompactHF);
+
 	GenNavDataTimeSlicedGenerationContext.Reset();
 	GenNavDataTimeSlicedAllocator.Reset();
 	GenCompressedlayersTimeSlicedRasterContext.Reset();
@@ -1734,7 +1742,7 @@ void FRecastTileGenerator::Setup(const FRecastNavMeshGenerator& ParentGenerator,
 	}
 
 	// We have to regenerate layers data in case geometry is changed or tile cache is missing
-	bRegenerateCompressedLayers = (bGeometryChanged || CompressedLayers.Num() == 0);
+	bRegenerateCompressedLayers = (bGeometryChanged || TileConfig.bGenerateLinks || CompressedLayers.Num() == 0);
 	
 	// Gather geometry for tile if it's inside navigable bounds
 	if (InclusionBounds.Num())
@@ -2035,7 +2043,7 @@ void FRecastTileGenerator::GatherGeometry(const FRecastNavMeshGenerator& ParentG
 	const FBox NewBounds = ParentGenerator.GrowBoundingBox(TileBB, /*bIncludeAgentHeight*/ false);
 	
 	NavigationOctree->FindElementsWithBoundsTest(NewBounds,
-		[&RelevantDataArray, &OwnerNavDataConfig, &ParentGenerator, this, NavSys, bGeometryChanged, bUseVirtualGeometryFilteringAndDirtying](const FNavigationOctreeElement& Element)
+		[&RelevantDataArray, &OwnerNavDataConfig, &ParentGenerator, this, bUseVirtualGeometryFilteringAndDirtying](const FNavigationOctreeElement& Element)
 	{
 		const bool bShouldUse = bUseVirtualGeometryFilteringAndDirtying ?
 			ParentGenerator.ShouldGenerateGeometryForOctreeElement(Element, OwnerNavDataConfig) :
@@ -2172,6 +2180,36 @@ void FRecastTileGenerator::GatherNavigationDataGeometry(const TSharedRef<FNaviga
 	}
 }
 
+// Deprecated
+bool FRecastTileGenerator::CreateHeightField(FNavMeshBuildContext& BuildContext, FTileRasterizationContext& RasterContext)
+{
+	return CreateHeightField(BuildContext);
+}
+
+// Deprecated
+void FRecastTileGenerator::GenerateRecastFilter(FNavMeshBuildContext& BuildContext, FTileRasterizationContext& RasterContext)
+{
+	GenerateRecastFilter(BuildContext);
+}
+
+// Deprecated
+ETimeSliceWorkResult FRecastTileGenerator::GenerateRecastFilterTimeSliced(FNavMeshBuildContext& BuildContext, FTileRasterizationContext& RasterContext)
+{
+	return GenerateRecastFilterTimeSliced(BuildContext);
+}
+
+// Deprecated
+bool FRecastTileGenerator::BuildCompactHeightField(FNavMeshBuildContext& BuildContext, FTileRasterizationContext& RasterContext)
+{
+	return BuildCompactHeightField(BuildContext);
+}
+
+// Deprecated
+bool FRecastTileGenerator::RecastErodeWalkable(FNavMeshBuildContext& BuildContext, FTileRasterizationContext& RasterContext)
+{
+	return RecastErodeWalkable(BuildContext);
+}
+
 void FRecastTileGenerator::ApplyVoxelFilter(rcHeightfield* HF, FVector::FReal WalkableRadius)
 {
 	SCOPE_CYCLE_COUNTER(STAT_Navigation_TileVoxelFilteringAsync);
@@ -2281,9 +2319,9 @@ void FRecastTileGenerator::ApplyVoxelFilter(rcHeightfield* HF, FVector::FReal Wa
 	}
 }
 
-void FRecastTileGenerator::InitRasterizationMaskArray(const rcHeightfield* SolidHF, TInlineMaskArray& OutRasterizationMasks)
+void FRecastTileGenerator::InitRasterizationMaskArray(const rcHeightfield* InSolidHF, TInlineMaskArray& OutRasterizationMasks)
 {
-	const int CellCount = SolidHF->width * SolidHF->height;
+	const int CellCount = InSolidHF->width * InSolidHF->height;
 	OutRasterizationMasks.SetNumUninitialized(CellCount);
 	const uint8 AllowAllFlags = 0xFF;
 	FMemory::Memset(OutRasterizationMasks.GetData(), AllowAllFlags, CellCount*sizeof(TInlineMaskArray::ElementType));
@@ -2478,9 +2516,8 @@ ETimeSliceWorkResult FRecastTileGenerator::GenerateTileTimeSliced()
 	FNavMeshBuildContext BuildContext(*this);
 	ETimeSliceWorkResult WorkResult = ETimeSliceWorkResult::Succeeded;
 
-	UE_CLOG(TileConfig.bGenerateLinks, LogNavigation, Warning, TEXT("Generating links in time slice mode is currently not supported so it's been disabled. Turn off bGenerateNavLinks to avoid the warning."));
-	dtLinkBuilderData LinkBuiderData;
-	LinkBuiderData.generatingLinks = false; // Make sure it's disable for now.
+	dtLinkBuilderData linkBuiderData;
+	linkBuiderData.generatingLinks = TileConfig.bGenerateLinks;
 	
 	check(TimeSliceManager);
 
@@ -2522,7 +2559,7 @@ ETimeSliceWorkResult FRecastTileGenerator::GenerateTileTimeSliced()
 	} //fall through to next state
 	case EGenerateTileTimeSlicedState::GenerateNavigationData:
 	{
-		WorkResult = GenerateNavigationDataTimeSliced(BuildContext, LinkBuiderData);
+		WorkResult = GenerateNavigationDataTimeSliced(BuildContext, linkBuiderData);
 
 		if (WorkResult != ETimeSliceWorkResult::CallAgainNextTimeSlice)
 		{
@@ -2559,9 +2596,6 @@ bool FRecastTileGenerator::GenerateTile()
 	if (bRegenerateCompressedLayers)
 	{
 		CompressedLayers.Reset();
-
-		// CompactHF is part of FTileRasterizationContext RasterContext that gets deleted after the GenerateCompressedLayers
-		// so if we are generating links, keep the chf.
 		bSuccess = GenerateCompressedLayers(BuildContext, LinkBuiderData);
 
 #if RECAST_INTERNAL_DEBUG_DATA
@@ -2579,9 +2613,6 @@ bool FRecastTileGenerator::GenerateTile()
 	{
 		bSuccess = GenerateNavigationData(BuildContext, LinkBuiderData);
 	}
-
-	rcFreeHeightField(LinkBuiderData.solidHF);
-	rcFreeCompactHeightfield(LinkBuiderData.compactHF);
 	
 #if RECAST_INTERNAL_DEBUG_DATA	
 	const double EndStamp = FPlatformTime::Seconds();
@@ -2597,23 +2628,19 @@ bool FRecastTileGenerator::GenerateTile()
 
 struct FTileRasterizationContext
 {
-	FTileRasterizationContext() : SolidHF(0), LayerSet(0), CompactHF(0), RasterizationFlags(rcRasterizationFlags(0))
+	FTileRasterizationContext() : LayerSet(nullptr), RasterizationFlags(rcRasterizationFlags(0))
 	{
 	}
 
 	~FTileRasterizationContext()
 	{
-		rcFreeHeightField(SolidHF);
 		rcFreeHeightfieldLayerSet(LayerSet);
-		rcFreeCompactHeightfield(CompactHF);
 	}
 
 	rcRasterizationFlags GetRasterizationFlags() const { return RasterizationFlags; }
 	void SetRasterizationFlags(rcRasterizationFlags Value) { RasterizationFlags = Value; }
 
-	struct rcHeightfield* SolidHF;
 	struct rcHeightfieldLayerSet* LayerSet;
-	struct rcCompactHeightfield* CompactHF;
 	TArray<FNavMeshTileData> Layers;
 	FRecastTileGenerator::TInlineMaskArray RasterizationMasks;
 
@@ -2621,7 +2648,7 @@ private:
 	rcRasterizationFlags RasterizationFlags;
 };
 
-bool FRecastTileGenerator::CreateHeightField(FNavMeshBuildContext& BuildContext, FTileRasterizationContext& RasterContext)
+bool FRecastTileGenerator::CreateHeightField(FNavMeshBuildContext& BuildContext)
 {
 #if RECAST_INTERNAL_DEBUG_DATA
 	if (!IsTileDebugAllowingGeneration())
@@ -2649,13 +2676,13 @@ bool FRecastTileGenerator::CreateHeightField(FNavMeshBuildContext& BuildContext,
 	// Allocate voxel heightfield where we rasterize our input data to.
 	if (bHasGeometry)
 	{
-		RasterContext.SolidHF = rcAllocHeightfield();
-		if (RasterContext.SolidHF == nullptr)
+		SolidHF = rcAllocHeightfield();
+		if (SolidHF == nullptr)
 		{
 			BuildContext.log(RC_LOG_ERROR, "CreateHeightField: Out of memory 'SolidHF'.");
 			return false;
 		}
-		if (!rcCreateHeightfield(&BuildContext, *RasterContext.SolidHF, TileConfig.width, TileConfig.height, TileConfig.bmin, TileConfig.bmax, TileConfig.cs, TileConfig.ch))
+		if (!rcCreateHeightfield(&BuildContext, *SolidHF, TileConfig.width, TileConfig.height, TileConfig.bmin, TileConfig.bmax, TileConfig.cs, TileConfig.ch))
 		{
 			BuildContext.log(RC_LOG_ERROR, "CreateHeightField: Could not create solid heightfield.");
 			return false;
@@ -2704,7 +2731,7 @@ ETimeSliceWorkResult FRecastTileGenerator::RasterizeGeometryRecastTimeSliced(FNa
 		rcRasterizeTriangles(&BuildContext,
 			Coords.GetData(), NumVerts,
 			Indices.GetData(), RasterizeGeomRecastTriAreas.GetData(), NumFaces,
-			*RasterContext.SolidHF, TileConfig.walkableClimb, RasterizationFlags, MaskArray);
+			*SolidHF, TileConfig.walkableClimb, RasterizationFlags, MaskArray);
 
 #if RECAST_INTERNAL_DEBUG_DATA	
 		BuildContext.InternalDebugData.TriangleCount += NumFaces;
@@ -2753,7 +2780,7 @@ void FRecastTileGenerator::RasterizeGeometryRecast(FNavMeshBuildContext& BuildCo
 		rcRasterizeTriangles(&BuildContext,
 			Coords.GetData(), NumVerts,
 			Indices.GetData(), RasterizeGeomRecastTriAreas.GetData(), NumFaces,
-			*RasterContext.SolidHF, TileConfig.walkableClimb, RasterizationFlags, MaskArray);
+			*SolidHF, TileConfig.walkableClimb, RasterizationFlags, MaskArray);
 	}
 
 #if RECAST_INTERNAL_DEBUG_DATA
@@ -2960,7 +2987,7 @@ void FRecastTileGenerator::RasterizeTriangles(FNavMeshBuildContext& BuildContext
 	}
 }
 
-void FRecastTileGenerator::GenerateRecastFilter(FNavMeshBuildContext& BuildContext, FTileRasterizationContext& RasterContext)
+void FRecastTileGenerator::GenerateRecastFilter(FNavMeshBuildContext& BuildContext)
 {
 	SCOPE_CYCLE_COUNTER(STAT_Navigation_RecastFilter)
 
@@ -2971,17 +2998,17 @@ void FRecastTileGenerator::GenerateRecastFilter(FNavMeshBuildContext& BuildConte
 	// remove unwanted overhangs caused by the conservative rasterization
 	// as well as filter spans where the character cannot possibly stand.
 	{
-		rcFilterLowHangingWalkableObstacles(&BuildContext, TileConfig.walkableClimb, *RasterContext.SolidHF);
+		rcFilterLowHangingWalkableObstacles(&BuildContext, TileConfig.walkableClimb, *SolidHF);
 	}
 	{
 		SCOPE_CYCLE_COUNTER(STAT_Navigation_FilterLedgeSpans)
 		
 		rcFilterLedgeSpans(&BuildContext, TileConfig.walkableHeight, TileConfig.walkableClimb,
-			(rcNeighborSlopeFilterMode)TileConfig.LedgeSlopeFilterMode, TileConfig.maxStepFromWalkableSlope, TileConfig.ch, *RasterContext.SolidHF);
+			(rcNeighborSlopeFilterMode)TileConfig.LedgeSlopeFilterMode, TileConfig.maxStepFromWalkableSlope, TileConfig.ch, *SolidHF);
 	}
 	if (!TileConfig.bMarkLowHeightAreas)
 	{
-		rcFilterWalkableLowHeightSpans(&BuildContext, TileConfig.walkableHeight, *RasterContext.SolidHF);
+		rcFilterWalkableLowHeightSpans(&BuildContext, TileConfig.walkableHeight, *SolidHF);
 	}
 	else if (TileConfig.bFilterLowSpanFromTileCache)
 	{
@@ -2990,16 +3017,16 @@ void FRecastTileGenerator::GenerateRecastFilter(FNavMeshBuildContext& BuildConte
 
 		if (TileConfig.bFilterLowSpanSequences && bHasLowAreaModifiers)
 		{
-			rcFilterWalkableLowHeightSpansSequences(&BuildContext, FilterWalkableHeight, *RasterContext.SolidHF);
+			rcFilterWalkableLowHeightSpansSequences(&BuildContext, FilterWalkableHeight, *SolidHF);
 		}
 		else
 		{
-			rcFilterWalkableLowHeightSpans(&BuildContext, FilterWalkableHeight, *RasterContext.SolidHF);
+			rcFilterWalkableLowHeightSpans(&BuildContext, FilterWalkableHeight, *SolidHF);
 		}
 	}
 }
 
-ETimeSliceWorkResult FRecastTileGenerator::GenerateRecastFilterTimeSliced(FNavMeshBuildContext& BuildContext, FTileRasterizationContext& RasterContext)
+ETimeSliceWorkResult FRecastTileGenerator::GenerateRecastFilterTimeSliced(FNavMeshBuildContext& BuildContext)
 {
 	SCOPE_CYCLE_COUNTER(STAT_Navigation_RecastFilter)
 
@@ -3014,7 +3041,7 @@ ETimeSliceWorkResult FRecastTileGenerator::GenerateRecastFilterTimeSliced(FNavMe
 		// Once all geometry is rasterized, we do initial pass of filtering to
 		// remove unwanted overhangs caused by the conservative rasterization
 		// as well as filter spans where the character cannot possibly stand.
-		rcFilterLowHangingWalkableObstacles(&BuildContext, TileConfig.walkableClimb, *RasterContext.SolidHF);
+		rcFilterLowHangingWalkableObstacles(&BuildContext, TileConfig.walkableClimb, *SolidHF);
 		GenerateRecastFilterState = EGenerateRecastFilterTimeSlicedState::FilterLedgeSpans;
 	}// fall through to next state
 	case EGenerateRecastFilterTimeSlicedState::FilterLedgeSpans:
@@ -3027,11 +3054,11 @@ ETimeSliceWorkResult FRecastTileGenerator::GenerateRecastFilterTimeSliced(FNavMe
 		{
 			rcFilterLedgeSpans(&BuildContext, TileConfig.walkableHeight, TileConfig.walkableClimb,
 				(rcNeighborSlopeFilterMode)TileConfig.LedgeSlopeFilterMode, TileConfig.maxStepFromWalkableSlope, TileConfig.ch, 
-				GenRecastFilterLedgeSpansYStart, TileTimeSliceSettings.FilterLedgeSpansMaxYProcess, *RasterContext.SolidHF);
+				GenRecastFilterLedgeSpansYStart, TileTimeSliceSettings.FilterLedgeSpansMaxYProcess, *SolidHF);
 
 			GenRecastFilterLedgeSpansYStart += TileTimeSliceSettings.FilterLedgeSpansMaxYProcess;
 
-			if (GenRecastFilterLedgeSpansYStart >= RasterContext.SolidHF->height)
+			if (GenRecastFilterLedgeSpansYStart >= SolidHF->height)
 			{
 				GenRecastFilterLedgeSpansYStart = 0;
 				DoIter = false;
@@ -3054,7 +3081,7 @@ ETimeSliceWorkResult FRecastTileGenerator::GenerateRecastFilterTimeSliced(FNavMe
 
 		if (!TileConfig.bMarkLowHeightAreas)
 		{
-			rcFilterWalkableLowHeightSpans(&BuildContext, TileConfig.walkableHeight, *RasterContext.SolidHF);
+			rcFilterWalkableLowHeightSpans(&BuildContext, TileConfig.walkableHeight, *SolidHF);
 		}
 		else if (TileConfig.bFilterLowSpanFromTileCache)
 		{
@@ -3062,11 +3089,11 @@ ETimeSliceWorkResult FRecastTileGenerator::GenerateRecastFilterTimeSliced(FNavMe
 			// for now, switch on presence of those modifiers, will save memory as long as they are sparse (should be)
 			if (TileConfig.bFilterLowSpanSequences && bHasLowAreaModifiers)
 			{
-				rcFilterWalkableLowHeightSpansSequences(&BuildContext, FilterWalkableHeight, *RasterContext.SolidHF);
+				rcFilterWalkableLowHeightSpansSequences(&BuildContext, FilterWalkableHeight, *SolidHF);
 			}
 			else
 			{
-				rcFilterWalkableLowHeightSpans(&BuildContext, FilterWalkableHeight, *RasterContext.SolidHF);
+				rcFilterWalkableLowHeightSpans(&BuildContext, FilterWalkableHeight, *SolidHF);
 			}
 		}
 		GenerateRecastFilterState = EGenerateRecastFilterTimeSlicedState::FilterLowHangingWalkableObstacles;
@@ -3082,22 +3109,22 @@ ETimeSliceWorkResult FRecastTileGenerator::GenerateRecastFilterTimeSliced(FNavMe
 	return ETimeSliceWorkResult::Succeeded;
 }
 
-bool FRecastTileGenerator::BuildCompactHeightField(FNavMeshBuildContext& BuildContext, FTileRasterizationContext& RasterContext)
+bool FRecastTileGenerator::BuildCompactHeightField(FNavMeshBuildContext& BuildContext)
 {
 	SCOPE_CYCLE_COUNTER(STAT_Navigation_RecastBuildCompactHeightField);
 
 	// Compact the heightfield so that it is faster to handle from now on.
 	// This will result more cache coherent data as well as the neighbors
 	// between walkable cells will be calculated.
-	RasterContext.CompactHF = rcAllocCompactHeightfield();
-	if (RasterContext.CompactHF == nullptr)
+	CompactHF = rcAllocCompactHeightfield();
+	if (CompactHF == nullptr)
 	{
 		BuildContext.log(RC_LOG_ERROR, "BuildCompactHeightField: Out of memory 'CompactHF'.");
 		return false;
 	}
-	if (!rcBuildCompactHeightfield(&BuildContext, TileConfig.walkableHeight, TileConfig.walkableClimb, *RasterContext.SolidHF, *RasterContext.CompactHF))
+	if (!rcBuildCompactHeightfield(&BuildContext, TileConfig.walkableHeight, TileConfig.walkableClimb, *SolidHF, *CompactHF))
 	{
-		const int SpanCount = rcGetHeightFieldSpanCount(&BuildContext, *RasterContext.SolidHF);
+		const int SpanCount = rcGetHeightFieldSpanCount(&BuildContext, *SolidHF);
 		if (SpanCount > 0)
 		{
 			BuildContext.log(RC_LOG_ERROR, "BuildCompactHeightField: Could not build compact data.");
@@ -3112,7 +3139,7 @@ bool FRecastTileGenerator::BuildCompactHeightField(FNavMeshBuildContext& BuildCo
 	return true;
 }
 
-bool FRecastTileGenerator::RecastErodeWalkable(FNavMeshBuildContext& BuildContext, FTileRasterizationContext& RasterContext)
+bool FRecastTileGenerator::RecastErodeWalkable(FNavMeshBuildContext& BuildContext)
 {
 	SCOPE_CYCLE_COUNTER(STAT_Navigation_RecastErodeWalkable);
 
@@ -3128,8 +3155,8 @@ bool FRecastTileGenerator::RecastErodeWalkable(FNavMeshBuildContext& BuildContex
 		}
 
 		const bool bEroded = TileConfig.bMarkLowHeightAreas ?
-			rcErodeWalkableAndLowAreas(&BuildContext, TileConfig.walkableRadius, FilterWalkableHeight, RECAST_LOW_AREA, FilterFlags, *RasterContext.CompactHF) :
-			rcErodeWalkableArea(&BuildContext, TileConfig.walkableRadius, *RasterContext.CompactHF);
+			rcErodeWalkableAndLowAreas(&BuildContext, TileConfig.walkableRadius, FilterWalkableHeight, RECAST_LOW_AREA, FilterFlags, *CompactHF) :
+			rcErodeWalkableArea(&BuildContext, TileConfig.walkableRadius, *CompactHF);
 
 		if (!bEroded)
 		{
@@ -3140,7 +3167,7 @@ bool FRecastTileGenerator::RecastErodeWalkable(FNavMeshBuildContext& BuildContex
 	}
 	else if (TileConfig.bMarkLowHeightAreas)
 	{
-		rcMarkLowAreas(&BuildContext, FilterWalkableHeight, RECAST_LOW_AREA, *RasterContext.CompactHF);
+		rcMarkLowAreas(&BuildContext, FilterWalkableHeight, RECAST_LOW_AREA, *CompactHF);
 	}
 
 	return true;
@@ -3159,7 +3186,7 @@ bool FRecastTileGenerator::RecastBuildLayers(FNavMeshBuildContext& BuildContext,
 
 	if (TileConfig.regionPartitioning == RC_REGION_MONOTONE)
 	{
-		if (!rcBuildHeightfieldLayersMonotone(&BuildContext, *RasterContext.CompactHF, TileConfig.borderSize, TileConfig.walkableHeight, *RasterContext.LayerSet))
+		if (!rcBuildHeightfieldLayersMonotone(&BuildContext, *CompactHF, TileConfig.borderSize, TileConfig.walkableHeight, *RasterContext.LayerSet))
 		{
 			BuildContext.log(RC_LOG_ERROR, "RecastBuildLayers: Could not build heightfield layers.");
 			return false;
@@ -3168,13 +3195,13 @@ bool FRecastTileGenerator::RecastBuildLayers(FNavMeshBuildContext& BuildContext,
 	}
 	else if (TileConfig.regionPartitioning == RC_REGION_WATERSHED)
 	{
-		if (!rcBuildDistanceField(&BuildContext, *RasterContext.CompactHF))
+		if (!rcBuildDistanceField(&BuildContext, *CompactHF))
 		{
 			BuildContext.log(RC_LOG_ERROR, "RecastBuildLayers: Could not build distance field.");
 			return false;
 		}
 
-		if (!rcBuildHeightfieldLayers(&BuildContext, *RasterContext.CompactHF, TileConfig.borderSize, TileConfig.walkableHeight, *RasterContext.LayerSet))
+		if (!rcBuildHeightfieldLayers(&BuildContext, *CompactHF, TileConfig.borderSize, TileConfig.walkableHeight, *RasterContext.LayerSet))
 		{
 			BuildContext.log(RC_LOG_ERROR, "RecastBuildLayers: Could not build heightfield layers.");
 			return false;
@@ -3182,7 +3209,7 @@ bool FRecastTileGenerator::RecastBuildLayers(FNavMeshBuildContext& BuildContext,
 	}
 	else
 	{
-		if (!rcBuildHeightfieldLayersChunky(&BuildContext, *RasterContext.CompactHF, TileConfig.borderSize, TileConfig.walkableHeight, TileConfig.regionChunkSize, *RasterContext.LayerSet))
+		if (!rcBuildHeightfieldLayersChunky(&BuildContext, *CompactHF, TileConfig.borderSize, TileConfig.walkableHeight, TileConfig.regionChunkSize, *RasterContext.LayerSet))
 		{
 			BuildContext.log(RC_LOG_ERROR, "RecastBuildLayers: Could not build heightfield layers.");
 			return false;
@@ -3308,7 +3335,7 @@ ETimeSliceWorkResult FRecastTileGenerator::GenerateCompressedLayersTimeSliced(FN
 	} // fall through to next state
 	case EGenerateCompressedLayersTimeSliced::CreateHeightField:
 	{
-		if (!CreateHeightField(BuildContext, *RasterContext))
+		if (!CreateHeightField(BuildContext))
 		{
 			GenCompressedLayersTimeSlicedState = EGenerateCompressedLayersTimeSliced::Invalid;
 			//no need to check time slice as not much work done
@@ -3341,7 +3368,7 @@ ETimeSliceWorkResult FRecastTileGenerator::GenerateCompressedLayersTimeSliced(FN
 	} // fall through to next state
 	case EGenerateCompressedLayersTimeSliced::EmptyLayers:
 	{
-		if (!RasterContext->SolidHF || RasterContext->SolidHF->pools == 0)
+		if (!SolidHF || SolidHF->pools == 0)
 		{
 			BuildContext.log(RC_LOG_WARNING, "GenerateCompressedLayersTimeSliced: empty tile - aborting");
 
@@ -3359,7 +3386,7 @@ ETimeSliceWorkResult FRecastTileGenerator::GenerateCompressedLayersTimeSliced(FN
 		// Reject voxels outside generation boundaries
 		if (TileConfig.bPerformVoxelFiltering && !bFullyEncapsulatedByInclusionBounds)
 		{
-			ApplyVoxelFilter(RasterContext->SolidHF, TileConfig.walkableRadius);
+			ApplyVoxelFilter(SolidHF, TileConfig.walkableRadius);
 
 			MARK_TIMESLICE_SECTION_DEBUG(TimeSliceManager->GetTimeSlicer(), VoxelFilter);
 
@@ -3371,7 +3398,7 @@ ETimeSliceWorkResult FRecastTileGenerator::GenerateCompressedLayersTimeSliced(FN
 	}// fall through to next state
 	case EGenerateCompressedLayersTimeSliced::RecastFilter:
 	{
-		const ETimeSliceWorkResult WorkResult = GenerateRecastFilterTimeSliced(BuildContext, *RasterContext);
+		const ETimeSliceWorkResult WorkResult = GenerateRecastFilterTimeSliced(BuildContext);
 
 		// Non timesliced code this is based on did not care about success or failure here.
 		if (WorkResult != ETimeSliceWorkResult::CallAgainNextTimeSlice)
@@ -3386,7 +3413,7 @@ ETimeSliceWorkResult FRecastTileGenerator::GenerateCompressedLayersTimeSliced(FN
 	}// fall through to next state
 	case EGenerateCompressedLayersTimeSliced::CompactHeightField:
 	{
-		if (!BuildCompactHeightField(BuildContext, *RasterContext))
+		if (!BuildCompactHeightField(BuildContext))
 		{
 			//no need to check time slice as not much work done
 			GenCompressedLayersTimeSlicedState = EGenerateCompressedLayersTimeSliced::Invalid;
@@ -3405,7 +3432,7 @@ ETimeSliceWorkResult FRecastTileGenerator::GenerateCompressedLayersTimeSliced(FN
 	}// fall through to next state
 	case EGenerateCompressedLayersTimeSliced::ErodeWalkable:
 	{
-		if (!RecastErodeWalkable(BuildContext, *RasterContext))
+		if (!RecastErodeWalkable(BuildContext))
 		{
 			//no need to check time slice as not much work done
 			GenCompressedLayersTimeSlicedState = EGenerateCompressedLayersTimeSliced::Invalid;
@@ -3518,14 +3545,14 @@ namespace UE::NavMesh::Private
 #endif //RECAST_INTERNAL_DEBUG_DATA
 };
 
-bool FRecastTileGenerator::GenerateCompressedLayers(FNavMeshBuildContext& BuildContext, dtLinkBuilderData& OutLinkBuilderData)
+bool FRecastTileGenerator::GenerateCompressedLayers(FNavMeshBuildContext& BuildContext, const dtLinkBuilderData& InLinkBuilderData)
 {
 	SCOPE_CYCLE_COUNTER(STAT_Navigation_RecastBuildCompressedLayers);
 
 	FTileRasterizationContext RasterContext;
 	CompressedLayers.Reset();
 
-	if (!CreateHeightField(BuildContext, RasterContext))
+	if (!CreateHeightField(BuildContext))
 	{
 		return false;
 	}
@@ -3533,7 +3560,7 @@ bool FRecastTileGenerator::GenerateCompressedLayers(FNavMeshBuildContext& BuildC
 	ComputeRasterizationMasks(BuildContext, RasterContext);
 
 	RasterizeTriangles(BuildContext, RasterContext);
-	if (!RasterContext.SolidHF || RasterContext.SolidHF->pools == 0)
+	if (!SolidHF || SolidHF->pools == 0)
 	{
 		BuildContext.log(RC_LOG_WARNING, "GenerateCompressedLayers: empty tile - aborting");
 		return true;
@@ -3544,11 +3571,11 @@ bool FRecastTileGenerator::GenerateCompressedLayers(FNavMeshBuildContext& BuildC
 	{
 		if (TileDebugSettings.bHeightfieldFromRasterization)
 		{
-			UE::NavMesh::Private::DrawHeightfield(TileDebugSettings.HeightFieldRenderMode, &BuildContext.InternalDebugData, *RasterContext.SolidHF);
+			UE::NavMesh::Private::DrawHeightfield(TileDebugSettings.HeightFieldRenderMode, &BuildContext.InternalDebugData, *SolidHF);
 		}
 		if (TileDebugSettings.bHeightfieldBounds)
 		{
-			duDebugDrawHeightfieldBounds(&BuildContext.InternalDebugData, *RasterContext.SolidHF);
+			duDebugDrawHeightfieldBounds(&BuildContext.InternalDebugData, *SolidHF);
 		}
 	}
 #endif
@@ -3556,26 +3583,26 @@ bool FRecastTileGenerator::GenerateCompressedLayers(FNavMeshBuildContext& BuildC
 	// Reject voxels outside generation boundaries
 	if (TileConfig.bPerformVoxelFiltering && !bFullyEncapsulatedByInclusionBounds)
 	{
-		ApplyVoxelFilter(RasterContext.SolidHF, TileConfig.walkableRadius);
+		ApplyVoxelFilter(SolidHF, TileConfig.walkableRadius);
 	}
 
 #if RECAST_INTERNAL_DEBUG_DATA
 	if (IsTileDebugActive() && TileDebugSettings.bHeightfieldPostInclusionBoundsFiltering)
 	{
-		UE::NavMesh::Private::DrawHeightfield(TileDebugSettings.HeightFieldRenderMode, &BuildContext.InternalDebugData, *RasterContext.SolidHF);
+		UE::NavMesh::Private::DrawHeightfield(TileDebugSettings.HeightFieldRenderMode, &BuildContext.InternalDebugData, *SolidHF);
 	}
 #endif
 
-	GenerateRecastFilter(BuildContext, RasterContext);
+	GenerateRecastFilter(BuildContext);
 
 #if RECAST_INTERNAL_DEBUG_DATA
 	if (IsTileDebugActive() && TileDebugSettings.bHeightfieldPostHeightFiltering)
 	{
-		UE::NavMesh::Private::DrawHeightfield(TileDebugSettings.HeightFieldRenderMode, &BuildContext.InternalDebugData, *RasterContext.SolidHF);
+		UE::NavMesh::Private::DrawHeightfield(TileDebugSettings.HeightFieldRenderMode, &BuildContext.InternalDebugData, *SolidHF);
 	}
 #endif
 
-	if (!BuildCompactHeightField(BuildContext, RasterContext))
+	if (!BuildCompactHeightField(BuildContext))
 	{
 		return false;
 	}
@@ -3583,11 +3610,11 @@ bool FRecastTileGenerator::GenerateCompressedLayers(FNavMeshBuildContext& BuildC
 #if RECAST_INTERNAL_DEBUG_DATA
 	if (IsTileDebugActive() && TileDebugSettings.bCompactHeightfield)
 	{
-		duDebugDrawCompactHeightfieldSolid(&BuildContext.InternalDebugData, *RasterContext.CompactHF);
+		duDebugDrawCompactHeightfieldSolid(&BuildContext.InternalDebugData, *CompactHF);
 	}
 #endif
 
-	if (!RecastErodeWalkable(BuildContext, RasterContext))
+	if (!RecastErodeWalkable(BuildContext))
 	{
 		return false;
 	}
@@ -3595,7 +3622,7 @@ bool FRecastTileGenerator::GenerateCompressedLayers(FNavMeshBuildContext& BuildC
 #if RECAST_INTERNAL_DEBUG_DATA
 	if (IsTileDebugActive() && TileDebugSettings.bCompactHeightfieldEroded)
 	{
-		duDebugDrawCompactHeightfieldSolid(&BuildContext.InternalDebugData, *RasterContext.CompactHF);
+		duDebugDrawCompactHeightfieldSolid(&BuildContext.InternalDebugData, *CompactHF);
 	}
 #endif
 
@@ -3609,26 +3636,15 @@ bool FRecastTileGenerator::GenerateCompressedLayers(FNavMeshBuildContext& BuildC
 	{
 		if (TileDebugSettings.bCompactHeightfieldRegions)
 		{
-			duDebugDrawCompactHeightfieldRegions(&BuildContext.InternalDebugData, *RasterContext.CompactHF);	
+			duDebugDrawCompactHeightfieldRegions(&BuildContext.InternalDebugData, *CompactHF);	
 		}
 
 		if (TileDebugSettings.bCompactHeightfieldDistances)
 		{
-			duDebugDrawCompactHeightfieldDistance(&BuildContext.InternalDebugData, *RasterContext.CompactHF);	
+			duDebugDrawCompactHeightfieldDistance(&BuildContext.InternalDebugData, *CompactHF);	
 		}
 	}
 #endif
-
-	if (OutLinkBuilderData.generatingLinks)
-	{
-		OutLinkBuilderData.solidHF = RasterContext.SolidHF;
-		OutLinkBuilderData.compactHF = RasterContext.CompactHF;
-
-		// When generating links, set SolidHF and CompactHF to null to prevent free 
-		// on the destruction of RasterContext since we will need those later.
-		RasterContext.SolidHF = nullptr;
-		RasterContext.CompactHF = nullptr;
-	}
 	
 	return RecastBuildTileCache(BuildContext, RasterContext);
 }
@@ -3677,7 +3693,7 @@ struct FTileGenerationContext
 };
 
 dtStatus FRecastTileGenerator::BuildTileCacheLinks(FNavMeshBuildContext& BuildContext, dtTileCacheAlloc* alloc, const dtTileCacheLayer& layer,
-	const dtTileCacheContourSet& lcset, const dtLinkBuilderData& linkBuilderData, TArray<FNavigationLink>& OutGeneratedLinks) const
+	const dtTileCacheContourSet& lcset, TArray<FNavigationLink>& OutGeneratedLinks) const
 {
 	duDebugDraw* dd = nullptr;
 	int32 DebugEdge = -1;
@@ -3693,13 +3709,15 @@ dtStatus FRecastTileGenerator::BuildTileCacheLinks(FNavMeshBuildContext& BuildCo
 
 	dtAssert(alloc);
 
-	if (!linkBuilderData.solidHF || !linkBuilderData.compactHF)
+	if (!SolidHF || !CompactHF)
 	{
 		return DT_FAILURE;
 	}
 
 	const dtReal* orig = layer.header->bmin;
 	dtLinkBuilderConfig linkBuilderConfig;
+	linkBuilderConfig.jumpDownConfig = TileConfig.JumpDownConfig;
+	linkBuilderConfig.jumpOverConfig = TileConfig.JumpOverConfig;
 	linkBuilderConfig.agentRadius = TileConfig.walkableRadius * TileConfig.cs;
 	linkBuilderConfig.agentHeight = TileConfig.walkableHeight * TileConfig.ch;
 	linkBuilderConfig.agentClimb = TileConfig.walkableClimb * TileConfig.ch;
@@ -3707,15 +3725,22 @@ dtStatus FRecastTileGenerator::BuildTileCacheLinks(FNavMeshBuildContext& BuildCo
 	linkBuilderConfig.cellHeight = TileConfig.ch;
 
 	dtNavLinkBuilder linkBuilder;
-	if (!linkBuilder.findEdges(BuildContext, TileConfig, linkBuilderConfig, lcset, orig, linkBuilderData))
 	{
-		return DT_FAILURE;
+		SCOPE_CYCLE_COUNTER(STAT_Navigation_RecastBuildLinks_FindEdges);
+
+		if (!linkBuilder.findEdges(BuildContext, TileConfig, linkBuilderConfig, lcset, orig, SolidHF, CompactHF))
+		{
+			return DT_FAILURE;
+		}
 	}
 
 	if (DebugEdge == -1)
 	{
-		linkBuilder.buildForAllEdges(linkBuilderConfig, DT_LINK_ACTION_JUMP_DOWN);
-		linkBuilder.buildForAllEdges(linkBuilderConfig, DT_LINK_ACTION_JUMP_OVER);
+		{
+			SCOPE_CYCLE_COUNTER(STAT_Navigation_RecastBuildLinks_Sample);
+			linkBuilder.buildForAllEdges(linkBuilderConfig, DT_LINK_ACTION_JUMP_DOWN);
+			linkBuilder.buildForAllEdges(linkBuilderConfig, DT_LINK_ACTION_JUMP_OVER);
+		}
 
 #if RECAST_INTERNAL_DEBUG_DATA
 		duDebugDrawNavLinkBuilder(dd, linkBuilder, DebugFlags, nullptr);
@@ -3752,7 +3777,7 @@ dtStatus FRecastTileGenerator::BuildTileCacheLinks(FNavMeshBuildContext& BuildCo
 }
 
 bool FRecastTileGenerator::GenerateNavigationDataLayer(FNavMeshBuildContext& BuildContext, FTileCacheCompressor& TileCompressor,
-	FTileCacheAllocator& GenNavAllocator, FTileGenerationContext& GenerationContext, dtLinkBuilderData& InOutLinkBuilderData, int32 LayerIdx)
+	FTileCacheAllocator& GenNavAllocator, FTileGenerationContext& GenerationContext, const dtLinkBuilderData& InLinkBuilderData, int32 LayerIdx)
 {
 	SCOPE_CYCLE_COUNTER(STAT_Navigation_GenerateNavigationDataLayer)
 		
@@ -3949,11 +3974,14 @@ bool FRecastTileGenerator::GenerateNavigationDataLayer(FNavMeshBuildContext& Bui
 
 	// Build Links
 	TArray<FNavigationLink> GeneratedLinks;
-	if (InOutLinkBuilderData.generatingLinks)
+	if (InLinkBuilderData.generatingLinks)
 	{
 		SCOPE_CYCLE_COUNTER(STAT_Navigation_RecastBuildLinks);
 
-		status = BuildTileCacheLinks(BuildContext, &GenNavAllocator, *GenerationContext.Layer, *GenerationContext.ContourSet, InOutLinkBuilderData, GeneratedLinks);
+		if (SolidHF && CompactHF)
+		{
+			status = BuildTileCacheLinks(BuildContext, &GenNavAllocator, *GenerationContext.Layer, *GenerationContext.ContourSet, GeneratedLinks);	
+		}
 
 		if (dtStatusFailed(status))
 		{
@@ -4066,15 +4094,12 @@ bool FRecastTileGenerator::GenerateNavigationDataLayer(FNavMeshBuildContext& Bui
 	return true;
 }
 
-ETimeSliceWorkResult FRecastTileGenerator::GenerateNavigationDataTimeSliced(FNavMeshBuildContext& BuildContext, dtLinkBuilderData& InOutLinkBuilderData)
+ETimeSliceWorkResult FRecastTileGenerator::GenerateNavigationDataTimeSliced(FNavMeshBuildContext& BuildContext, const dtLinkBuilderData& InLinkBuilderData)
 {
 	SCOPE_CYCLE_COUNTER(STAT_Navigation_RecastBuildNavigation);
 
 	check(TimeSliceManager);
-
-	FTileCacheCompressor TileCompressor;
 	ETimeSliceWorkResult WorkResult = ETimeSliceWorkResult::Succeeded;
-	dtStatus status = DT_SUCCESS;
 
 	switch (GenerateNavDataTimeSlicedState)
 	{
@@ -4108,7 +4133,9 @@ ETimeSliceWorkResult FRecastTileGenerator::GenerateNavigationDataTimeSliced(FNav
 				break;
 			}
 
-			const bool bGenDataLayer = GenerateNavigationDataLayer(BuildContext, TileCompressor, *GenNavDataTimeSlicedAllocator, *GenNavDataTimeSlicedGenerationContext, InOutLinkBuilderData, GenNavDataLayerTimeSlicedIdx);
+			FTileCacheCompressor TileCompressor;
+			const bool bGenDataLayer = GenerateNavigationDataLayer(BuildContext, TileCompressor, *GenNavDataTimeSlicedAllocator,
+				*GenNavDataTimeSlicedGenerationContext, InLinkBuilderData, GenNavDataLayerTimeSlicedIdx);
 
 			MARK_TIMESLICE_SECTION_DEBUG(TimeSliceManager->GetTimeSlicer(), GenerateLayers);
 
@@ -4146,7 +4173,7 @@ ETimeSliceWorkResult FRecastTileGenerator::GenerateNavigationDataTimeSliced(FNav
 	return WorkResult;
 }
 
-bool FRecastTileGenerator::GenerateNavigationData(FNavMeshBuildContext& BuildContext, dtLinkBuilderData& LinkBuilderData)
+bool FRecastTileGenerator::GenerateNavigationData(FNavMeshBuildContext& BuildContext, const dtLinkBuilderData& InLinkBuilderData)
 {
 	SCOPE_CYCLE_COUNTER(STAT_Navigation_RecastBuildNavigation);
 
@@ -4165,7 +4192,7 @@ bool FRecastTileGenerator::GenerateNavigationData(FNavMeshBuildContext& BuildCon
 			continue;
 		}
 
-		bGenDataLayer = GenerateNavigationDataLayer(BuildContext, TileCompressor, GenNavAllocator, GenerationContext, LinkBuilderData, LayerIdx);
+		bGenDataLayer = GenerateNavigationDataLayer(BuildContext, TileCompressor, GenNavAllocator, GenerationContext, InLinkBuilderData, LayerIdx);
 
 		if (!bGenDataLayer)
 		{
@@ -4196,7 +4223,7 @@ void FRecastTileGenerator::ComputeRasterizationMasks(FNavMeshBuildContext& Build
 	{
 		if (ModifierElement.bMaskFillCollisionUnderneathForNavmesh)
 		{
-			if (RasterContext.SolidHF == nullptr)
+			if (SolidHF == nullptr)
 			{
 				return;
 			}
@@ -4206,12 +4233,12 @@ void FRecastTileGenerator::ComputeRasterizationMasks(FNavMeshBuildContext& Build
 			{
 				for (const FTransform& LocalToWorld : ModifierElement.PerInstanceTransform)
 				{
-					MarkRasterizationMask(&BuildContext, RasterContext.SolidHF, ModifierArea, LocalToWorld, Mask, RasterContext.RasterizationMasks);
+					MarkRasterizationMask(&BuildContext, SolidHF, ModifierArea, LocalToWorld, Mask, RasterContext.RasterizationMasks);
 				}
 
 				if (ModifierElement.PerInstanceTransform.Num() == 0)
 				{
-					MarkRasterizationMask(&BuildContext, RasterContext.SolidHF, ModifierArea, FTransform::Identity, Mask, RasterContext.RasterizationMasks);
+					MarkRasterizationMask(&BuildContext, SolidHF, ModifierArea, FTransform::Identity, Mask, RasterContext.RasterizationMasks);
 				}
 			}
 		}
@@ -4415,7 +4442,7 @@ void MarkConvexMask(const int mask, const FVector::FReal* verts, const int nv, r
 	}
 }
 
-void FRecastTileGenerator::MarkRasterizationMask(rcContext* /*BuildContext*/, rcHeightfield* SolidHF,
+void FRecastTileGenerator::MarkRasterizationMask(rcContext* /*BuildContext*/, rcHeightfield* InSolidHF,
 	const FAreaNavModifier& Modifier, const FTransform& LocalToWorld, const int32 Mask, TInlineMaskArray& OutMaskArray)
 {
 	FBox ModifierBounds = Modifier.GetBounds().TransformBy(LocalToWorld);
@@ -4427,7 +4454,7 @@ void FRecastTileGenerator::MarkRasterizationMask(rcContext* /*BuildContext*/, rc
 	// Init on first use
 	if (OutMaskArray.Num() == 0)
 	{
-		InitRasterizationMaskArray(SolidHF, OutMaskArray);
+		InitRasterizationMaskArray(InSolidHF, OutMaskArray);
 	}
 
 	switch (Modifier.GetShapeType())
@@ -4441,8 +4468,8 @@ void FRecastTileGenerator::MarkRasterizationMask(rcContext* /*BuildContext*/, rc
 		FVector RecastPos;
 		FVector RecastExtent;
 		RecastBox.GetCenterAndExtents(RecastPos, RecastExtent);
-		check(OutMaskArray.Num() == SolidHF->width*SolidHF->height);
-		MarkBoxMask(&(RecastPos.X), &(RecastExtent.X), Mask, *SolidHF, OutMaskArray.GetData());
+		check(OutMaskArray.Num() == InSolidHF->width*InSolidHF->height);
+		MarkBoxMask(&(RecastPos.X), &(RecastExtent.X), Mask, *InSolidHF, OutMaskArray.GetData());
 	}
 	break;
 
@@ -4469,7 +4496,7 @@ void FRecastTileGenerator::MarkRasterizationMask(rcContext* /*BuildContext*/, rc
 				*ItCoord = RecastV.Z; ItCoord++;
 			}
 
-			MarkConvexMask(Mask, ConvexCoords.GetData(), ConvexVerts.Num(), *SolidHF, OutMaskArray.GetData());
+			MarkConvexMask(Mask, ConvexCoords.GetData(), ConvexVerts.Num(), *InSolidHF, OutMaskArray.GetData());
 		}
 	}
 	break;
@@ -4794,7 +4821,16 @@ void FRecastNavMeshGenerator::SetupTileConfig(const ENavigationDataResolution Ti
 	OutConfig.walkableRadius = FMath::CeilToInt(DestNavMesh->AgentRadius / CellSize);
 	OutConfig.maxStepFromWalkableSlope = OutConfig.cs * FMath::Tan(FMath::DegreesToRadians(OutConfig.walkableSlopeAngle));
 
-	OutConfig.borderSize = OutConfig.walkableRadius + 3; // +1 for voxelization rounding, +1 for ledge neighbor access, +1 for occasional errors
+	int MaxLinkDistanceVx = 0;
+	if (DestNavMesh->bGenerateNavLinks)
+	{
+		MaxLinkDistanceVx = FMath::CeilToInt(OutConfig.LinkSpillDistance / CellSize);
+	}
+
+	// +1 for voxelization rounding, +1 for ledge neighbor access, +1 for occasional errors
+	const int BorderForAgentVx = OutConfig.walkableRadius + 3;
+	OutConfig.borderSize = BorderForAgentVx + MaxLinkDistanceVx; 
+
 	OutConfig.maxEdgeLen = (int32)(1200.0f / CellSize);
 
 	OutConfig.minRegionArea = (int32)rcSqr(DestNavMesh->MinRegionArea / CellSize);
@@ -4862,13 +4898,29 @@ void FRecastNavMeshGenerator::ConfigureBuildProperties(FRecastBuildConfig& OutCo
 				*UEnum::GetDisplayValueAsText(Resolution).ToString(), AgentMaxSlope, static_cast<float>(RequiredClimbVx-1)*TempCellHeight);
 		}
 	}
+
+	int MaxLinkDistanceVx = 0;
+	if (DestNavMesh->bGenerateNavLinks)
+	{
+		// NavLink builder configuration
+		const FNavLinkGenerationJumpDownConfig& JumpDown = DestNavMesh->NavLinkJumpDownConfig;
+		const FNavLinkGenerationJumpOverConfig& JumpOver = DestNavMesh->NavLinkJumpOverConfig;
+		JumpDown.CopyToDetourConfig(OutConfig.JumpDownConfig);
+		JumpOver.CopyToDetourConfig(OutConfig.JumpOverConfig);
+
+		OutConfig.LinkSpillDistance = FMath::Max((JumpDown.JumpLength-JumpDown.JumpDistanceFromEdge), (JumpOver.JumpLength-JumpOver.JumpDistanceFromEdge));
+		MaxLinkDistanceVx = FMath::CeilToInt(OutConfig.LinkSpillDistance / CellSize);
+	}
 	
 	// store original sizes
 	OutConfig.AgentHeight = AgentHeight;
 	OutConfig.AgentMaxClimb = AgentMaxClimb;
 	OutConfig.AgentRadius = AgentRadius;
 
-	OutConfig.borderSize = OutConfig.walkableRadius + 3; // +1 for voxelization rounding, +1 for ledge neighbor access, +1 for occasional errors
+	// +1 for voxelization rounding, +1 for ledge neighbor access, +1 for occasional errors
+	const int BorderForAgentVx = OutConfig.walkableRadius + 3;
+	OutConfig.borderSize = BorderForAgentVx + MaxLinkDistanceVx; 
+
 	OutConfig.maxEdgeLen = (int32)(1200.0f / CellSize);
 
 	// hardcoded, but can be overridden by RecastNavMesh params later
diff --git a/Engine/Source/Runtime/NavigationSystem/Private/NavigationSystem.cpp b/Engine/Source/Runtime/NavigationSystem/Private/NavigationSystem.cpp
--- a/Engine/Source/Runtime/NavigationSystem/Private/NavigationSystem.cpp
+++ b/Engine/Source/Runtime/NavigationSystem/Private/NavigationSystem.cpp
@@ -118,6 +118,8 @@ DEFINE_STAT(STAT_Navigation_RecastBuildContours);
 DEFINE_STAT(STAT_Navigation_RecastBuildNavigation);
 DEFINE_STAT(STAT_Navigation_GenerateNavigationDataLayer);
 DEFINE_STAT(STAT_Navigation_RecastBuildLinks);
+DEFINE_STAT(STAT_Navigation_RecastBuildLinks_FindEdges);
+DEFINE_STAT(STAT_Navigation_RecastBuildLinks_Sample);
 DEFINE_STAT(STAT_Navigation_RecastBuildRegions);
 DEFINE_STAT(STAT_Navigation_UpdateNavOctree);
 DEFINE_STAT(STAT_Navigation_CollisionTreeMemory);
diff --git a/Engine/Source/Runtime/NavigationSystem/Public/NavMesh/LinkGenerationConfig.h b/Engine/Source/Runtime/NavigationSystem/Public/NavMesh/LinkGenerationConfig.h
new file mode 100644
--- /dev/null
+++ b/Engine/Source/Runtime/NavigationSystem/Public/NavMesh/LinkGenerationConfig.h
@@ -0,0 +1,62 @@
+﻿// Copyright Epic Games, Inc. All Rights Reserved.
+#pragma once
+
+#include "LinkGenerationConfig.generated.h"
+
+#if WITH_RECAST
+struct dtNavLinkBuilderJumpDownConfig;
+struct dtNavLinkBuilderJumpOverConfig;
+#endif //WITH_RECAST
+
+/** Experimental configuration for generated jump down links. */
+USTRUCT()
+struct FNavLinkGenerationJumpDownConfig
+{
+	GENERATED_USTRUCT_BODY()
+
+	// @todo: Rename, describe and find best defaults for those parameters.
+	UPROPERTY(EditAnywhere, Config, Category = Settings, meta=(Units=cm, UIMin=0, ClampMin=0))
+	float JumpLength = 150.f; 
+
+	UPROPERTY(EditAnywhere, Config, Category = Settings, meta=(Units=cm, UIMin=0, ClampMin=0))
+	float JumpDistanceFromEdge = 10.f; 
+	
+	UPROPERTY(EditAnywhere, Config, Category = Settings, meta=(Units=cm, UIMin=0, ClampMin=0))
+	float JumpMaxDepth = 150.f;
+
+	UPROPERTY(EditAnywhere, Config, Category = Settings, meta=(Units=cm, UIMin=0, ClampMin=0))
+	float JumpEndsHeightTolerance = 50.f;
+
+#if WITH_RECAST	
+	/** Copy configuration to dtNavLinkBuilderJumpDownConfig. */
+	void CopyToDetourConfig(dtNavLinkBuilderJumpDownConfig& OutDetourConfig) const;
+#endif //WITH_RECAST
+};
+
+/** Experimental configuration for generated jump over links. */
+USTRUCT()
+struct FNavLinkGenerationJumpOverConfig
+{
+	GENERATED_USTRUCT_BODY()
+
+	// @todo: Rename, describe and find best defaults for those parameters.
+	UPROPERTY(EditAnywhere, Config, Category = Settings, meta=(Units=cm, UIMin=0, ClampMin=0))
+	float JumpLength = 200.f; 
+
+	UPROPERTY(EditAnywhere, Config, Category = Settings, meta=(Units=cm, UIMin=0, ClampMin=0))
+	float JumpDistanceFromEdge = 100.f; 
+	
+	UPROPERTY(EditAnywhere, Config, Category = Settings, meta=(Units=cm, UIMin=0, ClampMin=0))
+	float JumpHeight = 100.f; 
+
+	UPROPERTY(EditAnywhere, Config, Category = Settings, meta=(Units=cm, UIMin=0, ClampMin=0))
+	float JumpHeightTolerance = 100.f; 
+
+	UPROPERTY(EditAnywhere, Config, Category = Settings, meta=(Units=cm, UIMin=0, ClampMin=0))
+	float JumpEndsHeightTolerance = 50.f;
+
+#if WITH_RECAST
+	/** Copy configuration to dtNavLinkBuilderJumpOverConfig. */
+	void CopyToDetourConfig(dtNavLinkBuilderJumpOverConfig& OutDetourConfig) const;
+#endif //WITH_RECAST
+};
diff --git a/Engine/Source/Runtime/NavigationSystem/Public/NavMesh/RecastNavMesh.h b/Engine/Source/Runtime/NavigationSystem/Public/NavMesh/RecastNavMesh.h
--- a/Engine/Source/Runtime/NavigationSystem/Public/NavMesh/RecastNavMesh.h
+++ b/Engine/Source/Runtime/NavigationSystem/Public/NavMesh/RecastNavMesh.h
@@ -8,6 +8,7 @@
 #include "UObject/ObjectMacros.h"
 #include "Templates/SubclassOf.h"
 #include "EngineDefines.h"
+#include "LinkGenerationConfig.h"
 #include "AI/Navigation/NavigationTypes.h"
 #include "AI/Navigation/NavigationDataResolution.h"
 #include "NavigationSystemTypes.h"
@@ -720,10 +721,12 @@ class ARecastNavMesh : public ANavigationData
 	UPROPERTY(EditAnywhere, Category=Display)
 	uint32 bDrawLabelsOnPathNodes:1;
 
+	/** Draw valid links (both ends are valid). */
 	UPROPERTY(EditAnywhere, Category=Display)
 	uint32 bDrawNavLinks:1;
 
-	UPROPERTY(EditAnywhere, Category=Display)
+	/** Draw failed links and valid links. */
+	UPROPERTY(EditAnywhere, Category=Display, Meta = (DisplayName = "Draw Failed and Valid NavLinks"))
 	uint32 bDrawFailedNavLinks:1;
 	
 	/** Draw navmesh's clusters and cluster links. (Requires WITH_NAVMESH_CLUSTER_LINKS=1) */
@@ -968,6 +971,14 @@ protected:
 	UPROPERTY() 
 	uint32 bAllowWorldPartitionedNavMesh : 1;
 #endif // WITH_EDITORONLY_DATA
+
+	/** Experimental: paramters for building jump down links.  */
+	UPROPERTY(EditAnywhere, Category=Generation, config)
+	FNavLinkGenerationJumpDownConfig NavLinkJumpDownConfig;
+
+	/** Experimental: paramters for building jump over links.  */
+	UPROPERTY(EditAnywhere, Category=Generation, config)
+	FNavLinkGenerationJumpOverConfig NavLinkJumpOverConfig;
 	
 private:
 	/** Cache rasterized voxels instead of just collision vertices/indices in navigation octree */
diff --git a/Engine/Source/Runtime/NavigationSystem/Public/NavMesh/RecastNavMeshGenerator.h b/Engine/Source/Runtime/NavigationSystem/Public/NavMesh/RecastNavMeshGenerator.h
--- a/Engine/Source/Runtime/NavigationSystem/Public/NavMesh/RecastNavMeshGenerator.h
+++ b/Engine/Source/Runtime/NavigationSystem/Public/NavMesh/RecastNavMeshGenerator.h
@@ -23,6 +23,7 @@
 
 #include "Recast/Recast.h"
 #include "Detour/DetourNavMesh.h"
+#include "Detour/DetourNavLinkBuilderConfig.h"
 
 #if RECAST_INTERNAL_DEBUG_DATA
 #include "NavMesh/RecastInternalDebugData.h"
@@ -77,6 +78,10 @@ struct FRecastBuildConfig : public rcConfig
 	 *	factors - DO NOT SET IT TO ARBITRARY VALUE */
 	int32 MaxPolysPerTile;
 
+	/** NavLink building configuration */
+	dtNavLinkBuilderJumpDownConfig JumpDownConfig;
+	dtNavLinkBuilderJumpOverConfig JumpOverConfig;
+	
 	/** Actual agent height (in uu)*/
 	float AgentHeight;
 	/** Actual agent climb (in uu)*/
@@ -91,6 +96,8 @@ struct FRecastBuildConfig : public rcConfig
 	ENavigationLedgeSlopeFilterMode LedgeSlopeFilterMode;
 	/** Is the config completely setup */
 	bool bIsTileSetupConfigCompleted;
+	/** Used when generating links automatically. Distance representing how far generated links can go outside a tile (in uu). */ 
+	float LinkSpillDistance;
 	
 	FRecastBuildConfig()
 	{
@@ -110,10 +117,13 @@ struct FRecastBuildConfig : public rcConfig
 		bGenerateLinks = false;
 		// Still initializing, even though the property is deprecated, to avoid static analysis warnings
 		MaxPolysPerTile = -1;
+		JumpDownConfig = dtNavLinkBuilderJumpDownConfig();
+		JumpOverConfig = dtNavLinkBuilderJumpOverConfig();
 		AgentIndex = 0;
 		TileResolution = ENavigationDataResolution::Default;
 		LedgeSlopeFilterMode = ENavigationLedgeSlopeFilterMode::Recast;
 		bIsTileSetupConfigCompleted = false;
+		LinkSpillDistance = 0.f;
 	}
 
 	rcReal GetTileSizeUU() const { return tileSize * cs; }
@@ -423,8 +433,19 @@ protected:
 	/** Gather geometry from a specified Navigation Data */
 	NAVIGATIONSYSTEM_API void GatherNavigationDataGeometry(const TSharedRef<FNavigationRelevantData, ESPMode::ThreadSafe>& ElementData, UNavigationSystemV1& NavSys, const FNavDataConfig& OwnerNavDataConfig, bool bGeometryChanged);
 
-	/** Start functions used by GenerateCompressedLayersTimeSliced / GenerateCompressedLayers */
+	UE_DEPRECATED(5.5, "Use the new version without RasterContext instead.")
 	NAVIGATIONSYSTEM_API bool CreateHeightField(FNavMeshBuildContext& BuildContext, FTileRasterizationContext& RasterContext);
+	UE_DEPRECATED(5.5, "Use the new version without RasterContext instead.")
+	NAVIGATIONSYSTEM_API void GenerateRecastFilter(FNavMeshBuildContext& BuildContext, FTileRasterizationContext& RasterContext);
+	UE_DEPRECATED(5.5, "Use the new version without RasterContext instead.")
+	NAVIGATIONSYSTEM_API ETimeSliceWorkResult GenerateRecastFilterTimeSliced(FNavMeshBuildContext& BuildContext, FTileRasterizationContext& RasterContext);
+	UE_DEPRECATED(5.5, "Use the new version without RasterContext instead.")
+	NAVIGATIONSYSTEM_API bool BuildCompactHeightField(FNavMeshBuildContext& BuildContext, FTileRasterizationContext& RasterContext);
+	UE_DEPRECATED(5.5, "Use the new version without RasterContext instead.")
+	NAVIGATIONSYSTEM_API bool RecastErodeWalkable(FNavMeshBuildContext& BuildContext, FTileRasterizationContext& RasterContext);
+	
+	/** Start functions used by GenerateCompressedLayersTimeSliced / GenerateCompressedLayers */
+	NAVIGATIONSYSTEM_API bool CreateHeightField(FNavMeshBuildContext& BuildContext);
 	NAVIGATIONSYSTEM_API ETimeSliceWorkResult RasterizeTrianglesTimeSliced(FNavMeshBuildContext& BuildContext, FTileRasterizationContext& RasterContext);
 	NAVIGATIONSYSTEM_API void RasterizeTriangles(FNavMeshBuildContext& BuildContext, FTileRasterizationContext& RasterContext);
 	NAVIGATIONSYSTEM_API ETimeSliceWorkResult RasterizeGeometryRecastTimeSliced(FNavMeshBuildContext& BuildContext, const TArray<FVector::FReal>& Coords, const TArray<int32>& Indices, const rcRasterizationFlags RasterizationFlags, FTileRasterizationContext& RasterContext);
@@ -434,10 +455,10 @@ protected:
 	NAVIGATIONSYSTEM_API void RasterizeGeometryTransformCoords(const TArray<FVector::FReal>& Coords, const FTransform& LocalToWorld);
 	NAVIGATIONSYSTEM_API ETimeSliceWorkResult RasterizeGeometryTimeSliced(FNavMeshBuildContext& BuildContext, const TArray<FVector::FReal>& Coords, const TArray<int32>& Indices, const FTransform& LocalToWorld, const rcRasterizationFlags RasterizationFlags, FTileRasterizationContext& RasterContext);
 	NAVIGATIONSYSTEM_API void RasterizeGeometry(FNavMeshBuildContext& BuildContext, const TArray<FVector::FReal>& Coords, const TArray<int32>& Indices, const FTransform& LocalToWorld, const rcRasterizationFlags RasterizationFlags, FTileRasterizationContext& RasterContext);
-	NAVIGATIONSYSTEM_API void GenerateRecastFilter(FNavMeshBuildContext& BuildContext, FTileRasterizationContext& RasterContext);
-	NAVIGATIONSYSTEM_API ETimeSliceWorkResult GenerateRecastFilterTimeSliced(FNavMeshBuildContext& BuildContext, FTileRasterizationContext& RasterContext);
-	NAVIGATIONSYSTEM_API bool BuildCompactHeightField(FNavMeshBuildContext& BuildContext, FTileRasterizationContext& RasterContext);
-	NAVIGATIONSYSTEM_API bool RecastErodeWalkable(FNavMeshBuildContext& BuildContext, FTileRasterizationContext& RasterContext);
+	NAVIGATIONSYSTEM_API void GenerateRecastFilter(FNavMeshBuildContext& BuildContext);
+	NAVIGATIONSYSTEM_API ETimeSliceWorkResult GenerateRecastFilterTimeSliced(FNavMeshBuildContext& BuildContext);
+	NAVIGATIONSYSTEM_API bool BuildCompactHeightField(FNavMeshBuildContext& BuildContext);
+	NAVIGATIONSYSTEM_API bool RecastErodeWalkable(FNavMeshBuildContext& BuildContext);
 	NAVIGATIONSYSTEM_API bool RecastBuildLayers(FNavMeshBuildContext& BuildContext, FTileRasterizationContext& RasterContext);
 	NAVIGATIONSYSTEM_API bool RecastBuildTileCache(FNavMeshBuildContext& BuildContext, FTileRasterizationContext& RasterContext);
 	/** End functions used by GenerateCompressedLayersTimeSliced / GenerateCompressedLayers */
@@ -464,27 +485,27 @@ protected:
 	
 	
 	/** Builds CompressedLayers array (geometry + modifiers) */
-	NAVIGATIONSYSTEM_API virtual bool GenerateCompressedLayers(FNavMeshBuildContext& BuildContext, dtLinkBuilderData& OutLinkBuilderData);
+	NAVIGATIONSYSTEM_API virtual bool GenerateCompressedLayers(FNavMeshBuildContext& BuildContext, const dtLinkBuilderData& InLinkBuilderData);
 
 	/** Builds a navigation data layer */
-	NAVIGATIONSYSTEM_API bool GenerateNavigationDataLayer(FNavMeshBuildContext& BuildContext, FTileCacheCompressor& TileCompressor, FTileCacheAllocator& GenNavAllocator, FTileGenerationContext& GenerationContext, dtLinkBuilderData& InOutLinkBuilderData, int32 LayerIdx);
+	NAVIGATIONSYSTEM_API bool GenerateNavigationDataLayer(FNavMeshBuildContext& BuildContext, FTileCacheCompressor& TileCompressor, FTileCacheAllocator& GenNavAllocator, FTileGenerationContext& GenerationContext, const dtLinkBuilderData& InLinkBuilderData, int32 LayerIdx);
 
 	/** Builds NavigationData array (layers + obstacles) time sliced */
-	NAVIGATIONSYSTEM_API ETimeSliceWorkResult GenerateNavigationDataTimeSliced(FNavMeshBuildContext& BuildContext, dtLinkBuilderData& InOutLinkBuilderData);
+	NAVIGATIONSYSTEM_API ETimeSliceWorkResult GenerateNavigationDataTimeSliced(FNavMeshBuildContext& BuildContext, const dtLinkBuilderData& InLinkBuilderData);
 
 	/** Builds NavigationData array (layers + obstacles) */
-	NAVIGATIONSYSTEM_API bool GenerateNavigationData(FNavMeshBuildContext& BuildContext, dtLinkBuilderData& LinkBuilderData);
+	NAVIGATIONSYSTEM_API bool GenerateNavigationData(FNavMeshBuildContext& BuildContext, const dtLinkBuilderData& InLinkBuilderData);
 
 	/** Builds navigation links */
 	dtStatus BuildTileCacheLinks(FNavMeshBuildContext& BuildContext, struct dtTileCacheAlloc* alloc, const dtTileCacheLayer& layer,
-		const struct dtTileCacheContourSet& lcset, const dtLinkBuilderData& linkBuilderData, TArray<FNavigationLink>& OutGeneratedLinks) const;
+		const struct dtTileCacheContourSet& lcset, TArray<FNavigationLink>& OutGeneratedLinks) const;
 	
 	NAVIGATIONSYSTEM_API virtual void ApplyVoxelFilter(struct rcHeightfield* SolidHF, FVector::FReal WalkableRadius);
 
 	/** Compute rasterization mask */
-	NAVIGATIONSYSTEM_API void InitRasterizationMaskArray(const rcHeightfield* SolidHF, TInlineMaskArray& OutRasterizationMasks);
+	NAVIGATIONSYSTEM_API void InitRasterizationMaskArray(const rcHeightfield* InSolidHF, TInlineMaskArray& OutRasterizationMasks);
 	NAVIGATIONSYSTEM_API void ComputeRasterizationMasks(FNavMeshBuildContext& BuildContext, FTileRasterizationContext& RasterContext);
-	NAVIGATIONSYSTEM_API void MarkRasterizationMask(rcContext* /*BuildContext*/, rcHeightfield* SolidHF,
+	NAVIGATIONSYSTEM_API void MarkRasterizationMask(rcContext* /*BuildContext*/, rcHeightfield* InSolidHF,
 		const FAreaNavModifier& Modifier, const FTransform& LocalToWorld, const int32 Mask, TInlineMaskArray& OutMaskArray);
 
 	/** apply areas from DynamicAreas to layer */
@@ -587,6 +608,9 @@ protected:
 	TWeakObjectPtr<UNavigationSystemV1> NavSystem; 
 	FNavDataConfig NavDataConfig;
 
+	rcHeightfield* SolidHF;
+	rcCompactHeightfield* CompactHF;
+
 	FRecastNavMeshTileGenerationDebug TileDebugSettings;
 
 #if RECAST_INTERNAL_DEBUG_DATA
diff --git a/Engine/Source/Runtime/Navmesh/Private/DebugUtils/DebugDraw.cpp b/Engine/Source/Runtime/Navmesh/Private/DebugUtils/DebugDraw.cpp
--- a/Engine/Source/Runtime/Navmesh/Private/DebugUtils/DebugDraw.cpp
+++ b/Engine/Source/Runtime/Navmesh/Private/DebugUtils/DebugDraw.cpp
@@ -76,7 +76,7 @@ unsigned int duColor::brown			= duRGBA(0x6d4c41ff);
 unsigned int duColor::white			= duRGBA(0xffffffff);
 unsigned int duColor::lightGrey		= duRGBA(0xe0e0e0ff);
 unsigned int duColor::grey			= duRGBA(0x9e9e9eff);
-unsigned int duColor::drakGrey		= duRGBA(0x616161ff);
+unsigned int duColor::darkGrey		= duRGBA(0x616161ff);
 unsigned int duColor::black			= duRGBA(0x000000ff);
 //@UE END
 
diff --git a/Engine/Source/Runtime/Navmesh/Private/DebugUtils/DetourNavLinkDebugDraw.cpp b/Engine/Source/Runtime/Navmesh/Private/DebugUtils/DetourNavLinkDebugDraw.cpp
--- a/Engine/Source/Runtime/Navmesh/Private/DebugUtils/DetourNavLinkDebugDraw.cpp
+++ b/Engine/Source/Runtime/Navmesh/Private/DebugUtils/DetourNavLinkDebugDraw.cpp
@@ -106,8 +106,8 @@ namespace UE::Detour::Private
 
 	void drawTrajectorySlice(duDebugDraw* dd, const dtReal* pa, const dtReal* pb, const dtNavLinkBuilder::Trajectory2D* trajectory, const unsigned int color)
 	{
-		unsigned int colt = duTransCol(color, 255);
-		unsigned int colb = duTransCol(duLerpCol(color,duColor::black,96), 255);
+		unsigned int colt = duTransCol(color, 50);
+		unsigned int colb = duTransCol(duLerpCol(color,duColor::black,96), 50);
 		
 		dd->begin(DU_DRAW_QUADS);
 		dtReal p0[3], p1[3], p2[3], p3[3];
@@ -448,12 +448,12 @@ void duDebugDrawNavLinkBuilder(duDebugDraw* dd, const dtNavLinkBuilder& linkBuil
 						dtVlerp(spt, es->start.p, es->start.q, u);
 						dtVlerp(ept, es->start.p, es->start.q, u);
 						
-						unsigned int col = duRGBA(48,16,16,255);
+						unsigned int col = duColor::darkGrey;
 						dtReal off = 1;
 						if ((s->flags & dtNavLinkBuilder::HAS_GROUND) == 0)
 						{
 							off = 0;
-							col = duRGBA(220,32,32,255);
+							col = duColor::orangeRed;
 						}
 						
 						spt[1] = s->height + off;
@@ -475,7 +475,7 @@ void duDebugDrawNavLinkBuilder(duDebugDraw* dd, const dtNavLinkBuilder& linkBuil
 							off = 1.f;
 						}
 						spt[1] = s->height + off;
-						dd->vertex(spt, duColor::white);
+						dd->vertex(spt, duColor::green);
 					}
 					dd->end();
 					
@@ -492,12 +492,12 @@ void duDebugDrawNavLinkBuilder(duDebugDraw* dd, const dtNavLinkBuilder& linkBuil
 						dtVlerp(spt, es->end.p, es->end.q, u);
 						dtVlerp(ept, es->end.p, es->end.q, u);
 						
-						unsigned int col = duRGBA(48,16,16,255);
+						unsigned int col = duColor::darkGrey;
 						float off = 1.f;
 						if ((s->flags & dtNavLinkBuilder::HAS_GROUND) == 0)
 						{
 							off = 0;
-							col = duRGBA(220,32,32,255);
+							col = duColor::orangeRed;
 						}
 						
 						spt[1] = s->height + off;
@@ -520,7 +520,7 @@ void duDebugDrawNavLinkBuilder(duDebugDraw* dd, const dtNavLinkBuilder& linkBuil
 							off = 1.f;
 						}
 						spt[1] = s->height + off;
-						dd->vertex(spt, duColor::white);
+						dd->vertex(spt, duColor::green);
 					}
 					dd->end();
 				}
@@ -548,9 +548,9 @@ void duDebugDrawNavLinkBuilder(duDebugDraw* dd, const dtNavLinkBuilder& linkBuil
 						ept[1] = esmp->height;
 						
 						if (ssmp->flags & dtNavLinkBuilder::UNRESTRICTED)
-							UE::Detour::Private::drawTrajectorySlice(dd, spt, ept, &es->trajectory, duRGBA(32,255,96,255));
+							UE::Detour::Private::drawTrajectorySlice(dd, spt, ept, &es->trajectory, duColor::green);
 						else
-							UE::Detour::Private::drawTrajectorySlice(dd, spt, ept, &es->trajectory, duRGBA(255,32,32,255));
+							UE::Detour::Private::drawTrajectorySlice(dd, spt, ept, &es->trajectory, duColor::orangeRed);
 					}
 					dd->depthMask(true);
 				}
diff --git a/Engine/Source/Runtime/Navmesh/Private/Detour/DetourNavLinkBuilder.cpp b/Engine/Source/Runtime/Navmesh/Private/Detour/DetourNavLinkBuilder.cpp
--- a/Engine/Source/Runtime/Navmesh/Private/Detour/DetourNavLinkBuilder.cpp
+++ b/Engine/Source/Runtime/Navmesh/Private/Detour/DetourNavLinkBuilder.cpp
@@ -242,14 +242,15 @@ void dtNavLinkBuilder::cleanup()
 }
 
 bool dtNavLinkBuilder::findEdges(rcContext& ctx, const rcConfig& cfg, const dtLinkBuilderConfig& builderConfig,
-                                 const dtTileCacheContourSet& lcset, const dtReal* orig, const dtLinkBuilderData& linkBuilderData)
+                                 const dtTileCacheContourSet& lcset, const dtReal* orig,
+                                 const rcHeightfield* solidHF, const rcCompactHeightfield* compactHF)
 {
 	cleanup();
 	m_linkBuilderConfig = builderConfig;
 
 	m_cs = cfg.cs;
-	m_solid = linkBuilderData.solidHF;
-	m_chf = linkBuilderData.compactHF;
+	m_solid = solidHF;
+	m_chf = compactHF;
 
 	// Build edges.
 	m_nedges = 0;
@@ -287,9 +288,6 @@ bool dtNavLinkBuilder::findEdges(rcContext& ctx, const rcConfig& cfg, const dtLi
 			const unsigned short* va = &c.verts[k*4];
 			const unsigned short* vb = &c.verts[j*4];
 			
-			if ((va[3] & 0xf) != 0xf)	// A direction is set, so it's a portal edge.
-				continue;
-			
 			// Check k-j for matching contour
 			bool matchFound = false;
 			for (int ii = 0; ii < lcset.nconts; ++ii)
@@ -360,7 +358,7 @@ dtNavLinkBuilder::JumpLink* dtNavLinkBuilder::addLink()
 	return link;
 }
 
-void dtNavLinkBuilder::addEdgeLinks(const dtLinkBuilderConfig& acfg, const EdgeSampler* es)
+void dtNavLinkBuilder::addEdgeLinks(const dtLinkBuilderConfig& builderConfig, const EdgeSampler* es)
 {
 	using namespace UE::Detour::NavLink::Private;
 	
@@ -390,7 +388,7 @@ void dtNavLinkBuilder::addEdgeLinks(const dtLinkBuilderConfig& acfg, const EdgeS
 		nflags[i] = kernel[(nkernel+1)/2];
 	}
 
-	const int agentRadiusVx = acfg.agentRadius / acfg.cellSize;
+	const int agentRadiusVx = builderConfig.agentRadius / builderConfig.cellSize;
 	
 	// Build segments
 	int start = -1;
@@ -519,7 +517,7 @@ void dtNavLinkBuilder::buildForAllEdges(const dtLinkBuilderConfig& builderConfig
 	for (int i = 0; i < m_nedges; ++i)
 	{
 		EdgeSampler sampler;
-		const bool success = sampleEdge(action, m_edges[i].sp, m_edges[i].sq, &sampler);
+		const bool success = sampleEdge(builderConfig, action, m_edges[i].sp, m_edges[i].sq, &sampler);
 		if (success)
 		{
 			addEdgeLinks(builderConfig, &sampler);
@@ -538,7 +536,7 @@ void dtNavLinkBuilder::debugBuildEdge(const dtLinkBuilderConfig& builderConfig,
 	
 	m_debugSelectedEdge = edgeIndex;
 
-	const bool success = sampleEdge(action, m_edges[edgeIndex].sp, m_edges[edgeIndex].sq, &sampler);
+	const bool success = sampleEdge(builderConfig, action, m_edges[edgeIndex].sp, m_edges[edgeIndex].sq, &sampler);
 	if (success)
 	{
 		addEdgeLinks(builderConfig, &sampler);
@@ -737,7 +735,8 @@ void dtNavLinkBuilder::initTrajectory(Trajectory2D* tra) const
 
 		const float y = dtLerp(pa[1], pb[1], u);
 		
-		s->ymin = dtMin(y0,y1) + m_linkBuilderConfig.agentClimb - y;
+		// Todo: revisit starting height and trajectory heights to handle ledge corner when jumping down.
+		s->ymin = dtMin(y0,y1) + (2*m_linkBuilderConfig.agentClimb) - y;	
 		s->ymax = dtMax(y0,y1) + m_linkBuilderConfig.agentHeight - y; 
 	}
 }
@@ -966,27 +965,28 @@ void dtNavLinkBuilder::initJumpOverRig(EdgeSampler* es, const dtReal* sp, const
 	es->groundRange = groundRange;
 }
 
-bool dtNavLinkBuilder::sampleEdge(dtNavLinkAction desiredAction, const dtReal* sp, const dtReal* sq, dtNavLinkBuilder::EdgeSampler* es) const
+bool dtNavLinkBuilder::sampleEdge(const dtLinkBuilderConfig& builderConfig, dtNavLinkAction desiredAction, const dtReal* sp, const dtReal* sq, dtNavLinkBuilder::EdgeSampler* es) const
 {
 	using namespace UE::Detour::NavLink::Private;
-
-	constexpr float toUU = 100.f;
 	
 	if (desiredAction == DT_LINK_ACTION_JUMP_DOWN)
 	{
-		constexpr float jumpStartDist = -0.25f * toUU;
-		constexpr float jumpEndDist = 2.f * toUU;
-		constexpr float jumpDownDist = -2.f * toUU;
-		constexpr float groundRange = 1.f * toUU;
-		initJumpDownRig(es, sp, sq, jumpStartDist, jumpEndDist, jumpDownDist, groundRange);
+		const dtNavLinkBuilderJumpDownConfig& config = builderConfig.jumpDownConfig;
+		const float jumpStartDist = config.jumpDistanceFromEdge;
+		const float jumpEndDist = config.jumpLength;
+		const float jumpDownDist = config.jumpMaxDepth;
+		const float groundRange = config.jumpEndsHeightTolerance;
+		initJumpDownRig(es, sp, sq, -jumpStartDist, jumpEndDist, -jumpDownDist, groundRange);
 	}
 	else if (desiredAction == DT_LINK_ACTION_JUMP_OVER)
 	{
-		constexpr float jumpDist = 2.5f * toUU;
-		constexpr float heightRange = 1.f * toUU;
+		const dtNavLinkBuilderJumpOverConfig& config = builderConfig.jumpOverConfig;
+		const float jumpDist = config.jumpLength;
+		const float heightRange = config.jumpHeightTolerance;
 		static constexpr int NSEGS = 8;
 		dtReal segs[NSEGS*6];
 		int nsegs = findPotentialJumpOverEdges(sp, sq, jumpDist, heightRange, segs, NSEGS);
+
 		int ibest = -1;
 		float dbest = 0;
 		for (int i = 0; i < nsegs; ++i)
@@ -1004,9 +1004,10 @@ bool dtNavLinkBuilder::sampleEdge(dtNavLinkAction desiredAction, const dtReal* s
 			return false;
 		}
 
-		constexpr float jumpHeight = 1.f * toUU;
-		constexpr float groundRange = 0.5f * toUU;
-		initJumpOverRig(es, &segs[ibest*6+0], &segs[ibest*6+3], -jumpDist*0.5f, jumpDist*0.5f, jumpHeight, groundRange);
+		const float jumpStartDist = config.jumpDistanceFromEdge; 
+		const float jumpHeight = config.jumpHeight;
+		const float groundRange = config.jumpEndsHeightTolerance;
+		initJumpOverRig(es, &segs[ibest*6+0], &segs[ibest*6+3], -jumpStartDist, jumpStartDist, jumpHeight, groundRange);
 	}
 	
 	initTrajectory(&es->trajectory);
diff --git a/Engine/Source/Runtime/Navmesh/Public/DebugUtils/DebugDraw.h b/Engine/Source/Runtime/Navmesh/Public/DebugUtils/DebugDraw.h
--- a/Engine/Source/Runtime/Navmesh/Public/DebugUtils/DebugDraw.h
+++ b/Engine/Source/Runtime/Navmesh/Public/DebugUtils/DebugDraw.h
@@ -238,7 +238,7 @@ struct duColor
 	static unsigned int white;
 	static unsigned int lightGrey;
 	static unsigned int grey;
-	static unsigned int drakGrey;
+	static unsigned int darkGrey;
 	static unsigned int black;
 };
 //@UE END
diff --git a/Engine/Source/Runtime/Navmesh/Public/Detour/DetourNavLinkBuilder.h b/Engine/Source/Runtime/Navmesh/Public/Detour/DetourNavLinkBuilder.h
--- a/Engine/Source/Runtime/Navmesh/Public/Detour/DetourNavLinkBuilder.h
+++ b/Engine/Source/Runtime/Navmesh/Public/Detour/DetourNavLinkBuilder.h
@@ -4,6 +4,7 @@
 #include "CoreTypes.h"
 #include "DetourAlloc.h"
 #include "Detour/DetourLargeWorldCoordinates.h"
+#include "Detour/DetourNavLinkBuilderConfig.h"
 
 //@UE BEGIN
 enum dtNavLinkAction
@@ -15,6 +16,9 @@ enum dtNavLinkAction
 
 struct dtLinkBuilderConfig
 {
+	dtNavLinkBuilderJumpDownConfig jumpDownConfig;
+	dtNavLinkBuilderJumpOverConfig jumpOverConfig;
+	
 	dtReal agentRadius = 0;
 	dtReal agentHeight = 0;
 	dtReal agentClimb = 0;
@@ -28,8 +32,6 @@ struct rcCompactHeightfield;
 struct NAVMESH_API dtLinkBuilderData
 {
 	bool generatingLinks = false;
-	rcHeightfield* solidHF = nullptr;
-	rcCompactHeightfield* compactHF = nullptr;
 };
 
 class rcContext;
@@ -40,8 +42,8 @@ class dtNavLinkBuilder
 	dtLinkBuilderConfig m_linkBuilderConfig;
 
 	dtReal m_cs = 0;
-	rcHeightfield* m_solid;
-	rcCompactHeightfield* m_chf;
+	const rcHeightfield* m_solid;
+	const rcCompactHeightfield* m_chf;
 
 	struct Edge
 	{
@@ -149,7 +151,8 @@ public:
 
 	// Loops through contours to store edge points in world coordinates.
 	NAVMESH_API bool findEdges(rcContext& ctx, const rcConfig& cfg, const dtLinkBuilderConfig& builderConfig,
-							   const struct dtTileCacheContourSet& lcset, const dtReal* orig, const dtLinkBuilderData& linkBuilderData);
+							   const struct dtTileCacheContourSet& lcset, const dtReal* orig,
+							   const rcHeightfield* solidHF, const rcCompactHeightfield* compactHF);
 
 	// For all edges, sample edges (sampleEdge) and add links to m_links
 	NAVMESH_API void buildForAllEdges(const dtLinkBuilderConfig& acfg, dtNavLinkAction action);
@@ -182,8 +185,8 @@ private:
 	
 	void filterJumpOverLinks() const;
 
-	bool sampleEdge(dtNavLinkAction desiredAction, const dtReal* sp, const dtReal* sq, dtNavLinkBuilder::EdgeSampler* sampler) const;
+	bool sampleEdge(const dtLinkBuilderConfig& builderConfig, dtNavLinkAction desiredAction, const dtReal* sp, const dtReal* sq, dtNavLinkBuilder::EdgeSampler* sampler) const;
 	JumpLink* addLink();
-	void addEdgeLinks(const dtLinkBuilderConfig& acfg, const EdgeSampler* es);
+	void addEdgeLinks(const dtLinkBuilderConfig& builderConfig, const EdgeSampler* es);
 };
 //@UE END
\ No newline at end of file
diff --git a/Engine/Source/Runtime/Navmesh/Public/Detour/DetourNavLinkBuilderConfig.h b/Engine/Source/Runtime/Navmesh/Public/Detour/DetourNavLinkBuilderConfig.h
new file mode 100644
--- /dev/null
+++ b/Engine/Source/Runtime/Navmesh/Public/Detour/DetourNavLinkBuilderConfig.h
@@ -0,0 +1,21 @@
+﻿// Copyright Epic Games, Inc. All Rights Reserved.
+#pragma once
+
+/** Configuration for generated jump down links. */
+struct dtNavLinkBuilderJumpDownConfig
+{
+	float jumpLength = 150.f; 
+	float jumpDistanceFromEdge = 10.f; 
+	float jumpMaxDepth = 150.f;
+	float jumpEndsHeightTolerance = 50.f;
+};
+
+/** Configuration for generated jump over links. */
+struct dtNavLinkBuilderJumpOverConfig
+{
+	float jumpLength = 200.f; 
+	float jumpDistanceFromEdge = 100.f; 
+	float jumpHeight = 100.f; 
+	float jumpHeightTolerance = 100.f; 
+	float jumpEndsHeightTolerance = 50.f; 
+};
