autortfm: fix issues related to the movement component being created uninstrumented

- revert the code that puts MovementComponent->RegisterComponent(); into the 'open' with a custom undo handler - our new treatment of UObject construction and rollback handles this naturally now, and adding components uninstrumented means that we don't properly capture the state of component lists. this was leading to memory corruption / crashes

- a few additional fixes across a couple of other modules which are necessary to run RegisterComponent with autortfm instrumentation

diff --git a/Engine/Source/Runtime/NavigationSystem/NavigationSystem.Build.cs b/Engine/Source/Runtime/NavigationSystem/NavigationSystem.Build.cs
--- a/Engine/Source/Runtime/NavigationSystem/NavigationSystem.Build.cs
+++ b/Engine/Source/Runtime/NavigationSystem/NavigationSystem.Build.cs
@@ -76,6 +76,8 @@ namespace UnrealBuildTool.Rules
                 PublicDependencyModuleNames.Add("UnrealEd");
                 CircularlyReferencedDependentModules.Add("UnrealEd");
             }
+
+			bAllowAutoRTFMInstrumentation = true;
         }
     }
 }
diff --git a/Engine/Source/Runtime/Renderer/Renderer.Build.cs b/Engine/Source/Runtime/Renderer/Renderer.Build.cs
--- a/Engine/Source/Runtime/Renderer/Renderer.Build.cs
+++ b/Engine/Source/Runtime/Renderer/Renderer.Build.cs
@@ -57,5 +57,7 @@ public class Renderer : ModuleRules
 		DynamicallyLoadedModuleNames.AddRange(new string[] { "EyeTracker" });
 
 		RuntimeDependencies.Add("$(EngineDir)/Content/Renderer/TessellationTable.bin");
+
+		bAllowAutoRTFMInstrumentation = true;
 	}
 }
