Expose FRecastGeometryExport struct. This allows us to :
- share the struct with other navigation data generators
- use it directly where needed instead of using the RecastNavMeshGenerator (e.g. NagivationSystem, NavigationDataHandler, etc.)

diff --git a/Engine/Source/Runtime/NavigationSystem/Private/NavMesh/RecastNavMeshGenerator.cpp b/Engine/Source/Runtime/NavigationSystem/Private/NavMesh/RecastNavMeshGenerator.cpp
--- a/Engine/Source/Runtime/NavigationSystem/Private/NavMesh/RecastNavMeshGenerator.cpp
+++ b/Engine/Source/Runtime/NavigationSystem/Private/NavMesh/RecastNavMeshGenerator.cpp
@@ -22,6 +22,7 @@
 #include "Chaos/HeightField.h"
 #include "Chaos/TriangleMeshImplicitObject.h"
 #include "NavMesh/PImplRecastNavMesh.h"
+#include "NavMesh/RecastGeometryExport.h"
 #include "VisualLogger/VisualLogger.h"
 
 // recast includes
@@ -223,32 +224,6 @@ static void ExportGeomToOBJFile(const FString& InFileName, const TNavStatArray<F
 #endif
 }
 
-//----------------------------------------------------------------------//
-// 
-// 
-
-struct FRecastGeometryExport : public FNavigableGeometryExport
-{
-	FRecastGeometryExport(FNavigationRelevantData& InData) : Data(&InData) 
-	{
-		Data->Bounds = FBox(ForceInit);
-	}
-
-	FNavigationRelevantData* Data;
-	TNavStatArray<FVector::FReal> VertexBuffer;
-	TNavStatArray<int32> IndexBuffer;
-	FWalkableSlopeOverride SlopeOverride;
-
-	virtual void ExportChaosTriMesh(const Chaos::FTriangleMeshImplicitObject* const TriMesh, const FTransform& LocalToWorld) override;
-	virtual void ExportChaosConvexMesh(const FKConvexElem* const Convex, const FTransform& LocalToWorld) override;
-	virtual void ExportChaosHeightField(const Chaos::FHeightField* const Heightfield, const FTransform& LocalToWorld) override;
-	virtual void ExportChaosHeightFieldSlice(const FNavHeightfieldSamples& PrefetchedHeightfieldSamples, const int32 NumRows, const int32 NumCols, const FTransform& LocalToWorld, const FBox& SliceBox) override;
-	virtual void ExportCustomMesh(const FVector* InVertices, int32 NumVerts, const int32* InIndices, int32 NumIndices, const FTransform& LocalToWorld) override;
-	virtual void ExportRigidBodySetup(UBodySetup& BodySetup, const FTransform& LocalToWorld) override;
-	virtual void AddNavModifiers(const FCompositeNavModifier& Modifiers) override;
-	virtual void SetNavDataPerInstanceTransformDelegate(const FNavDataPerInstanceTransformDelegate& InDelegate) override;
-};
-
 FRecastVoxelCache::FRecastVoxelCache(const uint8* Memory)
 {
 	uint8* BytesArr = (uint8*)Memory;
@@ -668,6 +643,11 @@ void ExportCustomMesh(const FVector* InVertices, int32 NumVerts, const int32* In
 		return;
 	}
 
+	if (InVertices == nullptr || InIndices == nullptr)
+	{
+		return;
+	}
+
 	int32 VertOffset = VertexBuffer.Num() / 3;
 	VertexBuffer.Reserve(VertexBuffer.Num() + NumVerts*3);
 	IndexBuffer.Reserve(IndexBuffer.Num() + NumIndices);
@@ -1072,8 +1052,13 @@ FORCEINLINE void TransformVertexSoupToRecast(const TArray<FVector>& VertexSoup,
 	}
 }
 
-FORCEINLINE void CovertCoordDataToRecast(TNavStatArray<FVector::FReal>& Coords)
+FORCEINLINE void ConvertCoordDataToRecast(TNavStatArray<FVector::FReal>& Coords)
 {
+	if (Coords.Num() == 0)
+	{
+		return;
+	}
+
 	FVector::FReal* CoordPtr = Coords.GetData();
 	const int32 MaxIt = Coords.Num() / 3;
 	for (int32 i = 0; i < MaxIt; i++)
@@ -1123,6 +1108,12 @@ void ExportVertexSoup(const TArray<FVector>& VertexSoup, TNavStatArray<FVector::
 
 } // namespace RecastGeometryExport
 
+FRecastGeometryExport::FRecastGeometryExport(FNavigationRelevantData& InData) 
+	: Data(&InData) 
+{
+	Data->Bounds = FBox(ForceInit);
+}
+
 void FRecastGeometryExport::ExportChaosTriMesh(const Chaos::FTriangleMeshImplicitObject* const TriMesh, const FTransform& LocalToWorld)
 {
 	RecastGeometryExport::ExportChaosTriMesh(TriMesh, LocalToWorld, VertexBuffer, IndexBuffer, Data->Bounds);
@@ -1145,6 +1136,12 @@ void FRecastGeometryExport::ExportChaosHeightFieldSlice(const FNavHeightfieldSam
 
 void FRecastGeometryExport::ExportCustomMesh(const FVector* InVertices, int32 NumVerts, const int32* InIndices, int32 NumIndices, const FTransform& LocalToWorld)
 {
+	if (NumIndices % 3 != 0)
+	{
+		UE_LOG(LogNavigation, Warning, TEXT("%hs: InIndices doesn't represent a list of triangles. Skipping it [Data Owner: %s]"), __FUNCTION__, *GetDataOwnerName());
+		return;
+	}
+
 	RecastGeometryExport::ExportCustomMesh(InVertices, NumVerts, InIndices, NumIndices, LocalToWorld, VertexBuffer, IndexBuffer, Data->Bounds);
 }
 
@@ -1163,6 +1160,32 @@ void FRecastGeometryExport::SetNavDataPerInstanceTransformDelegate(const FNavDat
 	Data->NavDataPerInstanceTransformDelegate = InDelegate;
 }
 
+void FRecastGeometryExport::ConvertVertexBufferToRecast()
+{
+	if (VertexBuffer.Num() % 3 != 0)
+	{
+		UE_LOG(LogNavigation, Warning, TEXT("%hs: try to convert a vertex buffer that doesn't contain a list of vertex triplets. Skipping it [Data Owner: %s]"), __FUNCTION__, *GetDataOwnerName());
+		return;
+	}
+
+	RecastGeometryExport::ConvertCoordDataToRecast(VertexBuffer);
+}
+
+void FRecastGeometryExport::StoreCollisionCache()
+{
+	RecastGeometryExport::StoreCollisionCache(*this);
+}
+
+void FRecastGeometryExport::TransformVertexSoupToRecast(const TArray<FVector>& VertexSoup, TNavStatArray<FVector>& Verts, TNavStatArray<int32>& Faces)
+{
+	RecastGeometryExport::TransformVertexSoupToRecast(VertexSoup, Verts, Faces);
+}
+
+FString FRecastGeometryExport::GetDataOwnerName() const
+{
+	return Data ? GetNameSafe(Data->GetOwner()) : TEXT("No Data");
+}
+
 FORCEINLINE void GrowConvexHull(const FVector::FReal ExpandBy, const TArray<FVector>& Verts, TArray<FVector>& OutResult)
 {
 	if (Verts.Num() < 3)
@@ -2116,7 +2139,7 @@ void FRecastTileGenerator::GatherNavigationDataGeometry(const TSharedRef<FNaviga
 			// overlap landscape's tile bounds
 			NavRelevant->GatherGeometrySlice(GeomExport, TileBBExpandedForAgent);
 
-			RecastGeometryExport::CovertCoordDataToRecast(GeomExport.VertexBuffer);
+			RecastGeometryExport::ConvertCoordDataToRecast(GeomExport.VertexBuffer);
 			RecastGeometryExport::StoreCollisionCache(GeomExport);
 			bDumpGeometryData = true;
 		}
@@ -7101,39 +7124,53 @@ void FRecastNavMeshGenerator::GetDebugGeometry(const FNavigationRelevantData& En
 }
 #endif // !UE_BUILD_SHIPPING
 
+// Deprecated
 void FRecastNavMeshGenerator::ExportComponentGeometry(UActorComponent* InOutComponent, FNavigationRelevantData& OutData)
 {
 	if (INavRelevantInterface* NavRelevantInterface = Cast<INavRelevantInterface>(InOutComponent))
 	{
-		ExportNavRelevantObjectGeometry(*NavRelevantInterface, OutData);
+		FRecastGeometryExport::ExportNavRelevantObjectGeometry(*NavRelevantInterface, OutData);
 	}
 }
 
+// Deprecated
 void FRecastNavMeshGenerator::ExportNavRelevantObjectGeometry(INavRelevantInterface& InOutNavRelevantInterface, FNavigationRelevantData& OutData)
+{
+	FRecastGeometryExport::ExportNavRelevantObjectGeometry(InOutNavRelevantInterface, OutData);
+}
+
+void FRecastGeometryExport::ExportNavRelevantObjectGeometry(INavRelevantInterface& InOutNavRelevantInterface, FNavigationRelevantData& OutData)
 {
 	FRecastGeometryExport GeomExport(OutData);
 	RecastGeometryExport::ExportObject(InOutNavRelevantInterface, GeomExport);
 
-#if !UE_BUILD_SHIPPING	
+#if !UE_BUILD_SHIPPING
 	RecastGeometryExport::ValidateGeometryExport(GeomExport);
 #endif
 	
-	RecastGeometryExport::CovertCoordDataToRecast(GeomExport.VertexBuffer);
+	RecastGeometryExport::ConvertCoordDataToRecast(GeomExport.VertexBuffer);
 	RecastGeometryExport::StoreCollisionCache(GeomExport);
 }
 
+// Deprecated
 void FRecastNavMeshGenerator::ExportVertexSoupGeometry(const TArray<FVector>& InVerts, FNavigationRelevantData& OutData)
+{
+	FRecastGeometryExport::ExportVertexSoupGeometry(InVerts, OutData);
+}
+
+void FRecastGeometryExport::ExportVertexSoupGeometry(const TArray<FVector>& InVerts, FNavigationRelevantData& OutData)
 {
 	FRecastGeometryExport GeomExport(OutData);
 	RecastGeometryExport::ExportVertexSoup(InVerts, GeomExport.VertexBuffer, GeomExport.IndexBuffer, GeomExport.Data->Bounds);
 
-#if !UE_BUILD_SHIPPING	
+#if !UE_BUILD_SHIPPING
 	RecastGeometryExport::ValidateGeometryExport(GeomExport);
 #endif
-	
+
 	RecastGeometryExport::StoreCollisionCache(GeomExport);
 }
 
+// Deprecated
 void FRecastNavMeshGenerator::ExportRigidBodyGeometry(
 	UBodySetup& InOutBodySetup,
 	TNavStatArray<FVector>& OutVertexBuffer,
@@ -7141,15 +7178,21 @@ void FRecastNavMeshGenerator::ExportRigidBodyGeometry(
 	const FTransform& LocalToWorld)
 {
 	FBox TempBounds;
-	ExportRigidBodyGeometry(InOutBodySetup, OutVertexBuffer, OutIndexBuffer, TempBounds, LocalToWorld);
+	FRecastGeometryExport::ExportRigidBodyGeometry(InOutBodySetup, OutVertexBuffer, OutIndexBuffer, TempBounds, LocalToWorld);
 }
 
+// Deprecated
 void FRecastNavMeshGenerator::ExportRigidBodyGeometry(
 	UBodySetup& InOutBodySetup,
 	TNavStatArray<FVector>& OutVertexBuffer,
 	TNavStatArray<int32>& OutIndexBuffer,
 	FBox& OutBounds,
 	const FTransform& LocalToWorld)
+{
+	FRecastGeometryExport::ExportRigidBodyGeometry(InOutBodySetup, OutVertexBuffer, OutIndexBuffer, OutBounds, LocalToWorld);
+}
+
+void FRecastGeometryExport::ExportRigidBodyGeometry(UBodySetup& InOutBodySetup, TNavStatArray<FVector>& OutVertexBuffer, TNavStatArray<int32>& OutIndexBuffer, FBox& OutBounds, const FTransform& LocalToWorld)
 {
 	TNavStatArray<FVector::FReal> VertCoords;
 	RecastGeometryExport::ExportRigidBodySetup(InOutBodySetup, VertCoords, OutIndexBuffer, OutBounds, LocalToWorld);
@@ -7161,6 +7204,7 @@ void FRecastNavMeshGenerator::ExportRigidBodyGeometry(
 	}
 }
 
+// Deprecated
 void FRecastNavMeshGenerator::ExportRigidBodyGeometry(
 	UBodySetup& InOutBodySetup,
 	TNavStatArray<FVector>& OutTriMeshVertexBuffer,
@@ -7171,7 +7215,7 @@ void FRecastNavMeshGenerator::ExportRigidBodyGeometry(
 	const FTransform& LocalToWorld)
 {
 	FBox TempBounds;
-	ExportRigidBodyGeometry(
+	FRecastGeometryExport::ExportRigidBodyGeometry(
 		InOutBodySetup,
 		OutTriMeshVertexBuffer,
 		OutTriMeshIndexBuffer,
@@ -7182,6 +7226,7 @@ void FRecastNavMeshGenerator::ExportRigidBodyGeometry(
 		LocalToWorld);
 }
 
+// Deprecated
 void FRecastNavMeshGenerator::ExportRigidBodyGeometry(
 	UBodySetup& InOutBodySetup,
 	TNavStatArray<FVector>& OutTriMeshVertexBuffer,
@@ -7191,6 +7236,19 @@ void FRecastNavMeshGenerator::ExportRigidBodyGeometry(
 	TNavStatArray<int32>& OutShapeBuffer,
 	FBox& OutBounds,
 	const FTransform& LocalToWorld)
+{
+	FRecastGeometryExport::ExportRigidBodyGeometry(
+		InOutBodySetup,
+		OutTriMeshVertexBuffer,
+		OutTriMeshIndexBuffer,
+		OutConvexVertexBuffer,
+		OutConvexIndexBuffer,
+		OutShapeBuffer,
+		OutBounds,
+		LocalToWorld);
+}
+
+void FRecastGeometryExport::ExportRigidBodyGeometry(UBodySetup& InOutBodySetup, TNavStatArray<FVector>& OutTriMeshVertexBuffer, TNavStatArray<int32>& OutTriMeshIndexBuffer, TNavStatArray<FVector>& OutConvexVertexBuffer, TNavStatArray<int32>& OutConvexIndexBuffer, TNavStatArray<int32>& OutShapeBuffer, FBox& OutBounds, const FTransform& LocalToWorld)
 {
 	InOutBodySetup.CreatePhysicsMeshes();
 
@@ -7217,6 +7275,7 @@ void FRecastNavMeshGenerator::ExportRigidBodyGeometry(
 	}
 }
 
+// Deprecated
 void FRecastNavMeshGenerator::ExportAggregatedGeometry(
 	const FKAggregateGeom& AggGeom,
 	TNavStatArray<FVector>& OutConvexVertexBuffer,
@@ -7225,9 +7284,10 @@ void FRecastNavMeshGenerator::ExportAggregatedGeometry(
 	const FTransform& LocalToWorld)
 {
 	FBox TempBounds;
-	ExportAggregatedGeometry(AggGeom, OutConvexVertexBuffer, OutConvexIndexBuffer, OutShapeBuffer, TempBounds, LocalToWorld);
+	FRecastGeometryExport::ExportAggregatedGeometry(AggGeom, OutConvexVertexBuffer, OutConvexIndexBuffer, OutShapeBuffer, TempBounds, LocalToWorld);
 }
 
+// Deprecated
 void FRecastNavMeshGenerator::ExportAggregatedGeometry(
 	const FKAggregateGeom& AggGeom,
 	TNavStatArray<FVector>& OutConvexVertexBuffer,
@@ -7235,6 +7295,11 @@ void FRecastNavMeshGenerator::ExportAggregatedGeometry(
 	TNavStatArray<int32>& OutShapeBuffer,
 	FBox& OutBounds,
 	const FTransform& LocalToWorld)
+{
+	FRecastGeometryExport::ExportAggregatedGeometry(AggGeom, OutConvexVertexBuffer, OutConvexIndexBuffer, OutShapeBuffer, OutBounds, LocalToWorld);
+}
+
+void FRecastGeometryExport::ExportAggregatedGeometry(const FKAggregateGeom& AggGeom, TNavStatArray<FVector>& OutConvexVertexBuffer, TNavStatArray<int32>& OutConvexIndexBuffer, TNavStatArray<int32>& OutShapeBuffer, FBox& OutBounds, const FTransform& LocalToWorld)
 {
 	TNavStatArray<FVector::FReal> VertCoords;
 	const int32 NumExistingVerts = OutConvexVertexBuffer.Num();
diff --git a/Engine/Source/Runtime/NavigationSystem/Private/NavigationDataHandler.cpp b/Engine/Source/Runtime/NavigationSystem/Private/NavigationDataHandler.cpp
--- a/Engine/Source/Runtime/NavigationSystem/Private/NavigationDataHandler.cpp
+++ b/Engine/Source/Runtime/NavigationSystem/Private/NavigationDataHandler.cpp
@@ -2,8 +2,9 @@
 
 #include "NavigationDataHandler.h"
 #include "Engine/Level.h"
-#include "NavMesh/RecastNavMeshGenerator.h"
+#include "GameFramework/Actor.h"
 #include "NavAreas/NavArea.h"
+#include "NavMesh/RecastGeometryExport.h"
 
 DEFINE_LOG_CATEGORY_STATIC(LogNavOctree, Warning, All);
 
@@ -496,7 +497,7 @@ void FNavigationDataHandler::AddLevelCollisionToOctree(ULevel& Level)
 				BSPElem.Data->bLoadedData = Level.HasVisibilityChangeRequestPending();
 			}
 			
-			FRecastNavMeshGenerator::ExportVertexSoupGeometry(*LevelGeom, *BSPElem.Data);
+			FRecastGeometryExport::ExportVertexSoupGeometry(*LevelGeom, *BSPElem.Data);
 
 			const FBox& Bounds = BSPElem.Data->Bounds;
 			if (!Bounds.GetExtent().IsNearlyZero())
diff --git a/Engine/Source/Runtime/NavigationSystem/Private/NavigationSystem.cpp b/Engine/Source/Runtime/NavigationSystem/Private/NavigationSystem.cpp
--- a/Engine/Source/Runtime/NavigationSystem/Private/NavigationSystem.cpp
+++ b/Engine/Source/Runtime/NavigationSystem/Private/NavigationSystem.cpp
@@ -35,7 +35,7 @@
 #if WITH_RECAST
 #include "NavMesh/RecastNavMesh.h"
 #include "NavMesh/RecastHelpers.h"
-#include "NavMesh/RecastNavMeshGenerator.h"
+#include "NavMesh/RecastGeometryExport.h"
 #endif // WITH_RECAST
 #if WITH_EDITOR
 #include "EditorModeManager.h"
@@ -1037,7 +1037,7 @@ bool UNavigationSystemV1::ConditionalPopulateNavOctree()
 			if (bStoreNavGeometry)
 			{
 #if WITH_RECAST
-				DefaultOctreeController.NavOctree->NavRelevantGeometryExportDelegate = FNavigationOctree::FNavRelevantGeometryExportDelegate::CreateStatic(&FRecastNavMeshGenerator::ExportNavRelevantObjectGeometry);
+				DefaultOctreeController.NavOctree->NavRelevantGeometryExportDelegate = FNavigationOctree::FNavRelevantGeometryExportDelegate::CreateStatic(&FRecastGeometryExport::ExportNavRelevantObjectGeometry);
 #endif // WITH_RECAST
 			}
 
diff --git a/Engine/Source/Runtime/NavigationSystem/Private/NavigationSystemTypes.cpp b/Engine/Source/Runtime/NavigationSystem/Private/NavigationSystemTypes.cpp
--- a/Engine/Source/Runtime/NavigationSystem/Private/NavigationSystemTypes.cpp
+++ b/Engine/Source/Runtime/NavigationSystem/Private/NavigationSystemTypes.cpp
@@ -2,8 +2,8 @@
 
 #include "NavigationSystemTypes.h"
 #include "NavLinkCustomInterface.h"
-#include "NavMesh/RecastNavMeshGenerator.h"
 #include "NavigationSystem.h"
+#include "NavMesh/RecastGeometryExport.h"
 #include "Engine/World.h"
 #include "NavCollision.h"
 #include "PhysicsEngine/BodySetup.h"
@@ -48,7 +48,7 @@ namespace NavigationHelper
 			return;
 		}
 #if WITH_RECAST
-		FRecastNavMeshGenerator::ExportRigidBodyGeometry(*RigidBody, OutVertexBuffer, OutIndexBuffer, OutBounds, LocalToWorld);
+		FRecastGeometryExport::ExportRigidBodyGeometry(*RigidBody, OutVertexBuffer, OutIndexBuffer, OutBounds, LocalToWorld);
 #endif // WITH_RECAST
 	}
 
@@ -59,7 +59,7 @@ namespace NavigationHelper
 			return;
 		}
 #if WITH_RECAST
-		FRecastNavMeshGenerator::ExportRigidBodyGeometry(*RigidBody
+		FRecastGeometryExport::ExportRigidBodyGeometry(*RigidBody
 			, NavCollision->GetMutableTriMeshCollision().VertexBuffer, NavCollision->GetMutableTriMeshCollision().IndexBuffer
 			, NavCollision->GetMutableConvexCollision().VertexBuffer, NavCollision->GetMutableConvexCollision().IndexBuffer
 			, NavCollision->ConvexShapeIndices
@@ -70,7 +70,7 @@ namespace NavigationHelper
 	void GatherCollision(const FKAggregateGeom& AggGeom, UNavCollision& NavCollision)
 	{
 #if WITH_RECAST
-		FRecastNavMeshGenerator::ExportAggregatedGeometry(
+		FRecastGeometryExport::ExportAggregatedGeometry(
 			AggGeom,
 			NavCollision.GetMutableConvexCollision().VertexBuffer,
 			 NavCollision.GetMutableConvexCollision().IndexBuffer,
diff --git a/Engine/Source/Runtime/NavigationSystem/Public/NavMesh/RecastGeometryExport.h b/Engine/Source/Runtime/NavigationSystem/Public/NavMesh/RecastGeometryExport.h
new file mode 100644
--- /dev/null
+++ b/Engine/Source/Runtime/NavigationSystem/Public/NavMesh/RecastGeometryExport.h
@@ -0,0 +1,87 @@
+// Copyright Epic Games, Inc. All Rights Reserved.
+#pragma once
+
+#include "AI/NavigationSystemHelpers.h"
+#include "Engine/EngineTypes.h"
+
+struct FNavigationRelevantData;
+
+#if WITH_RECAST
+
+/**
+ * Class that handles geometry exporting for Recast navmesh generation.
+ */
+struct FRecastGeometryExport : public FNavigableGeometryExport
+{
+	NAVIGATIONSYSTEM_API FRecastGeometryExport(FNavigationRelevantData& InData);
+
+	FNavigationRelevantData* Data;
+	TNavStatArray<FVector::FReal> VertexBuffer;
+	TNavStatArray<int32> IndexBuffer;
+	FWalkableSlopeOverride SlopeOverride;
+
+	/** Export the collision of a Chaos triangle mesh into the Vertex and Index buffer. */
+	NAVIGATIONSYSTEM_API virtual void ExportChaosTriMesh(const Chaos::FTriangleMeshImplicitObject* const TriMesh, const FTransform& LocalToWorld) override;
+	/** Export the collision of a Chaos convex mesh into the Vertex and Index buffer. */
+	NAVIGATIONSYSTEM_API virtual void ExportChaosConvexMesh(const FKConvexElem* const Convex, const FTransform& LocalToWorld) override;
+	/** Export the collision of a Chaos height field into the Vertex and Index buffer. */
+	NAVIGATIONSYSTEM_API virtual void ExportChaosHeightField(const Chaos::FHeightField* const Heightfield, const FTransform& LocalToWorld) override;
+	/** Export a slice of the collision of a Chaos height field into the Vertex and Index buffer.
+	 * @param SliceBox Box that defines the slice we want to extract from the height field */
+	NAVIGATIONSYSTEM_API virtual void ExportChaosHeightFieldSlice(const FNavHeightfieldSamples& PrefetchedHeightfieldSamples, const int32 NumRows, const int32 NumCols, const FTransform& LocalToWorld, const FBox& SliceBox) override;
+	/** Export a custom mesh into the Vertex and Index buffer.*/
+	NAVIGATIONSYSTEM_API virtual void ExportCustomMesh(const FVector* InVertices, int32 NumVerts, const int32* InIndices, int32 NumIndices, const FTransform& LocalToWorld) override;
+	/** Export a rigid body into the Vertex and Index buffer.*/
+	NAVIGATIONSYSTEM_API virtual void ExportRigidBodySetup(UBodySetup& BodySetup, const FTransform& LocalToWorld) override;
+	/** Add Nav Modifiers to the owned NavigationRelevantData. */
+	NAVIGATIONSYSTEM_API virtual void AddNavModifiers(const FCompositeNavModifier& Modifiers) override;
+	/** Optional delegate for geometry per instance transforms. */
+	NAVIGATIONSYSTEM_API virtual void SetNavDataPerInstanceTransformDelegate(const FNavDataPerInstanceTransformDelegate& InDelegate) override;
+
+	/** Convert the vertices in VertexBuffer from Unreal to Recast coordinates. */
+	NAVIGATIONSYSTEM_API void ConvertVertexBufferToRecast();
+	/** Store Vertex and Index buffer data in associated FNavigationRelevantData. */
+	NAVIGATIONSYSTEM_API void StoreCollisionCache();
+
+	/** Collect the collision information from an INavRelevantInterface and store it into the FNavigationRelevantData's CollisionData. */
+	static NAVIGATIONSYSTEM_API void ExportNavRelevantObjectGeometry(INavRelevantInterface& InOutNavRelevantInterface, FNavigationRelevantData& OutData);
+
+	/** Convert a list of vertices into the navigation format and store it into the FNavigationRelevantData's CollisionData.
+	 * @param InVerts Array of triangles vertices position. Each triangle will be created from 3 consecutive vertices in the array. Its size must be a multiple a 3.*/
+	static NAVIGATIONSYSTEM_API void ExportVertexSoupGeometry(const TArray<FVector>& InVerts, FNavigationRelevantData& OutData);
+
+	/** Collect the collision information of a BodySetup as a triangle mesh. */
+	static NAVIGATIONSYSTEM_API void ExportRigidBodyGeometry(UBodySetup& InOutBodySetup,
+		TNavStatArray<FVector>& OutVertexBuffer,
+		TNavStatArray<int32>& OutIndexBuffer,
+		FBox& OutBounds,
+		const FTransform& LocalToWorld = FTransform::Identity);
+
+	/** Collect the collision information of a BodySetup as a triangle mesh and a series of convex shapes. */
+	static NAVIGATIONSYSTEM_API void ExportRigidBodyGeometry(
+		UBodySetup& InOutBodySetup,
+		TNavStatArray<FVector>& OutTriMeshVertexBuffer,
+		TNavStatArray<int32>& OutTriMeshIndexBuffer,
+		TNavStatArray<FVector>& OutConvexVertexBuffer,
+		TNavStatArray<int32>& OutConvexIndexBuffer,
+		TNavStatArray<int32>& OutShapeBuffer,
+		FBox& OutBounds,
+		const FTransform& LocalToWorld = FTransform::Identity);
+
+	/** Collect the collision information of an AggregateGeometry as a series of convex shapes. */
+	static NAVIGATIONSYSTEM_API void ExportAggregatedGeometry(
+		const FKAggregateGeom& AggGeom,
+		TNavStatArray<FVector>& OutConvexVertexBuffer,
+		TNavStatArray<int32>& OutConvexIndexBuffer,
+		TNavStatArray<int32>& OutShapeBuffer,
+		FBox& OutBounds,
+		const FTransform& LocalToWorld = FTransform::Identity);
+
+	/** Transform a list of vertex triplets from Unreal to Recast coordinate and generate an Index buffer. */
+	static NAVIGATIONSYSTEM_API void TransformVertexSoupToRecast(const TArray<FVector>& VertexSoup, TNavStatArray<FVector>& Verts, TNavStatArray<int32>& Faces);
+
+private:
+	FString GetDataOwnerName() const;
+};
+
+#endif // WITH_RECAST
diff --git a/Engine/Source/Runtime/NavigationSystem/Public/NavMesh/RecastNavMeshGenerator.h b/Engine/Source/Runtime/NavigationSystem/Public/NavMesh/RecastNavMeshGenerator.h
--- a/Engine/Source/Runtime/NavigationSystem/Public/NavMesh/RecastNavMeshGenerator.h
+++ b/Engine/Source/Runtime/NavigationSystem/Public/NavMesh/RecastNavMeshGenerator.h
@@ -841,17 +841,17 @@ public:
 	NAVIGATIONSYSTEM_API virtual void GrabDebugSnapshot(struct FVisualLogEntry* Snapshot, const FBox& BoundingBox, const FName& CategoryName, ELogVerbosity::Type Verbosity) const override;
 #endif
 
-	UE_DEPRECATED(5.4, "Use ExportNavRelevantObjectGeometry")
+	UE_DEPRECATED(5.4, "Use FRecastGeometryExport::ExportNavRelevantObjectGeometry")
 	static NAVIGATIONSYSTEM_API void ExportComponentGeometry(UActorComponent* InOutComponent, FNavigationRelevantData& OutData);
 
-	UE_DEPRECATED(5.4, "Use ExportRigidBodyGeometry that takes bounds as parameter.")
+	UE_DEPRECATED(5.4, "Use FRecastGeometryExport::ExportRigidBodyGeometry.")
 	static NAVIGATIONSYSTEM_API void ExportRigidBodyGeometry(
 		UBodySetup& InOutBodySetup,
 		TNavStatArray<FVector>& OutVertexBuffer,
 		TNavStatArray<int32>& OutIndexBuffer,
 		const FTransform& LocalToWorld = FTransform::Identity);
 
-	UE_DEPRECATED(5.4, "Use ExportRigidBodyGeometry that takes bounds as parameter.")
+	UE_DEPRECATED(5.4, "Use FRecastGeometryExport::ExportRigidBodyGeometry.")
 	static NAVIGATIONSYSTEM_API void ExportRigidBodyGeometry(
 		UBodySetup& InOutBodySetup,
 		TNavStatArray<FVector>& OutTriMeshVertexBuffer,
@@ -861,7 +861,7 @@ public:
 		TNavStatArray<int32>& OutShapeBuffer,
 		const FTransform& LocalToWorld = FTransform::Identity);
 
-	UE_DEPRECATED(5.4, "Use ExportAggregatedGeometry that takes bounds as parameter.")
+	UE_DEPRECATED(5.4, "Use FRecastGeometryExport::ExportAggregatedGeometry.")
 	static NAVIGATIONSYSTEM_API void ExportAggregatedGeometry(
 		const FKAggregateGeom& AggGeom,
 		TNavStatArray<FVector>& OutConvexVertexBuffer,
@@ -869,15 +869,19 @@ public:
 		TNavStatArray<int32>& OutShapeBuffer,
 		const FTransform& LocalToWorld = FTransform::Identity);
 
+	UE_DEPRECATED(5.5, "Use FRecastGeometryExport::ExportNavRelevantObjectGeometry.")
 	static NAVIGATIONSYSTEM_API void ExportNavRelevantObjectGeometry(INavRelevantInterface& InOutNavRelevantInterface, FNavigationRelevantData& OutData);
+	UE_DEPRECATED(5.5, "Use FRecastGeometryExport::ExportVertexSoupGeometry.")
 	static NAVIGATIONSYSTEM_API void ExportVertexSoupGeometry(const TArray<FVector>& InVerts, FNavigationRelevantData& OutData);
 
+	UE_DEPRECATED(5.5, "Use FRecastGeometryExport::ExportRigidBodyGeometry.")
 	static NAVIGATIONSYSTEM_API void ExportRigidBodyGeometry(UBodySetup& InOutBodySetup,
 		TNavStatArray<FVector>& OutVertexBuffer,
 		TNavStatArray<int32>& OutIndexBuffer,
 		FBox& OutBounds,
 		const FTransform& LocalToWorld = FTransform::Identity);
 
+	UE_DEPRECATED(5.5, "Use FRecastGeometryExport::ExportRigidBodyGeometry.")
 	static NAVIGATIONSYSTEM_API void ExportRigidBodyGeometry(
 		UBodySetup& InOutBodySetup,
 		TNavStatArray<FVector>& OutTriMeshVertexBuffer,
@@ -888,6 +892,7 @@ public:
 		FBox& OutBounds,
 		const FTransform& LocalToWorld = FTransform::Identity);
 
+	UE_DEPRECATED(5.5, "Use FRecastGeometryExport::ExportAggregatedGeometry.")
 	static NAVIGATIONSYSTEM_API void ExportAggregatedGeometry(
 		const FKAggregateGeom& AggGeom,
 		TNavStatArray<FVector>& OutConvexVertexBuffer,
