Extract loop that notify the navigation data about which invokers are currently active

diff --git a/Engine/Source/Runtime/NavigationSystem/Private/NavigationSystem.cpp b/Engine/Source/Runtime/NavigationSystem/Private/NavigationSystem.cpp
--- a/Engine/Source/Runtime/NavigationSystem/Private/NavigationSystem.cpp
+++ b/Engine/Source/Runtime/NavigationSystem/Private/NavigationSystem.cpp
@@ -5382,15 +5382,7 @@ void UNavigationSystemV1::UpdateInvokers()
 #endif // ENABLE_VISUAL_LOG
 		}
 
-#if WITH_RECAST
-		const double UpdateStartTime = FPlatformTime::Seconds();
-		for (TActorIterator<ARecastNavMesh> It(World); It; ++It)
-		{
-			It->UpdateActiveTiles(InvokerLocations);
-		}
-		const double UpdateEndTime = FPlatformTime::Seconds();
-		UE_VLOG(this, LogNavInvokers, Log, TEXT("Marking tiles to update %fms (%d invokers)"), (UpdateEndTime - UpdateStartTime) * 1000, InvokerLocations.Num());
-#endif
+		UpdateNavDataActiveTiles();
 
 		// once per second
 		NextInvokersUpdateTime = CurrentTime + ActiveTilesUpdateInterval;
@@ -5430,6 +5422,19 @@ void UNavigationSystemV1::UpdateInvokers()
 #endif // !UE_BUILD_SHIPPING
 }
 
+void UNavigationSystemV1::UpdateNavDataActiveTiles()
+{
+#if WITH_RECAST
+	const double UpdateStartTime = FPlatformTime::Seconds();
+	for (TActorIterator<ARecastNavMesh> It(GetWorld()); It; ++It)
+	{
+		It->UpdateActiveTiles(InvokerLocations);
+	}
+	const double UpdateEndTime = FPlatformTime::Seconds();
+	UE_VLOG(this, LogNavInvokers, Log, TEXT("Marking tiles to update %fms (%d invokers)"), (UpdateEndTime - UpdateStartTime) * 1000, InvokerLocations.Num());
+#endif
+}
+
 void UNavigationSystemV1::DirtyTilesInBuildBounds()
 {
 #if WITH_RECAST
diff --git a/Engine/Source/Runtime/NavigationSystem/Public/NavigationSystem.h b/Engine/Source/Runtime/NavigationSystem/Public/NavigationSystem.h
--- a/Engine/Source/Runtime/NavigationSystem/Public/NavigationSystem.h
+++ b/Engine/Source/Runtime/NavigationSystem/Public/NavigationSystem.h
@@ -466,7 +466,11 @@ private:
 
 	double NextInvokersUpdateTime;
 	NAVIGATIONSYSTEM_API void UpdateInvokers();
+	
+protected:
+	NAVIGATIONSYSTEM_API virtual void UpdateNavDataActiveTiles();
 
+private:
 	NAVIGATIONSYSTEM_API void DirtyTilesInBuildBounds();
 
 	void UnregisterInvoker_Internal(const UObject& Invoker);
