Prevent crash when using NavQueryFilter with no implementation in the NavTestingActor

diff --git a/Engine/Source/Runtime/NavigationSystem/Private/NavigationTestingActor.cpp b/Engine/Source/Runtime/NavigationSystem/Private/NavigationTestingActor.cpp
--- a/Engine/Source/Runtime/NavigationSystem/Private/NavigationTestingActor.cpp
+++ b/Engine/Source/Runtime/NavigationSystem/Private/NavigationTestingActor.cpp
@@ -523,16 +523,21 @@ void ANavigationTestingActor::SearchPathTo(ANavigationTestingActor* Goal)
 
 	FPathFindingQuery Query = BuildPathFindingQuery(Goal);
 
+	FSharedConstNavQueryFilter NavQueryFilter = Query.QueryFilter ? Query.QueryFilter : NavData->GetDefaultQueryFilter();
+	if (NavQueryFilter->GetImplementation() == nullptr)
+	{
+		return;
+	}
+
 	if (bBacktracking)
 	{
-		FSharedConstNavQueryFilter NavQueryFilter = Query.QueryFilter ? Query.QueryFilter : NavData->GetDefaultQueryFilter();
 		FSharedNavQueryFilter NavigationFilterCopy = NavQueryFilter->GetCopy();
 		NavigationFilterCopy->SetBacktrackingEnabled(true);
 		Query.QueryFilter = NavigationFilterCopy;
+		NavQueryFilter = NavigationFilterCopy;
 	}
-	
+
 	//Apply cost limit factor
-	FSharedConstNavQueryFilter NavQueryFilter = Query.QueryFilter ? Query.QueryFilter : NavData->GetDefaultQueryFilter();
 	const float HeuristicScale = NavQueryFilter->GetHeuristicScale();
 	Query.CostLimit = FPathFindingQuery::ComputeCostLimitFromHeuristic(Query.StartLocation, Query.EndLocation, HeuristicScale, CostLimitFactor, MinimumCostLimit);
 
