[NavMesh] Optimisation: navlink generation now uses asymmetric border sizes.
It's now possible to build navmesh tiles with border sizes that are not equal all around.
When computing the extra borders to generate navlinks, borders are now asymmetric (increased only on one side of the axis)
This reduces substantially the size of the required heightfield and reduces by half the added cost of building links.
Added the ai.nav.UseAsymetricBorderSizes console command to deactivate the feature
Added AverageLinkBuildTime in the navmesh debug display
Also made sure not to compute border distance for disabled configurations

diff --git a/Engine/Source/Runtime/NavigationSystem/Private/NavMesh/LinkGenerationConfig.cpp b/Engine/Source/Runtime/NavigationSystem/Private/NavMesh/LinkGenerationConfig.cpp
--- a/Engine/Source/Runtime/NavigationSystem/Private/NavMesh/LinkGenerationConfig.cpp
+++ b/Engine/Source/Runtime/NavigationSystem/Private/NavMesh/LinkGenerationConfig.cpp
@@ -7,7 +7,7 @@
 
 void FNavLinkGenerationJumpDownConfig::CopyToDetourConfig(dtNavLinkBuilderJumpDownConfig& OutDetourConfig) const
 {
-	OutDetourConfig.enabled = bEnable;
+	OutDetourConfig.enabled = bEnabled;
 	OutDetourConfig.jumpLength = JumpLength;
 	OutDetourConfig.jumpDistanceFromEdge = JumpDistanceFromEdge;
 	OutDetourConfig.jumpMaxDepth = JumpMaxDepth;
diff --git a/Engine/Source/Runtime/NavigationSystem/Private/NavMesh/NavMeshRenderingComponent.cpp b/Engine/Source/Runtime/NavigationSystem/Private/NavMesh/NavMeshRenderingComponent.cpp
--- a/Engine/Source/Runtime/NavigationSystem/Private/NavMesh/NavMeshRenderingComponent.cpp
+++ b/Engine/Source/Runtime/NavigationSystem/Private/NavMesh/NavMeshRenderingComponent.cpp
@@ -526,6 +526,7 @@ void FNavMeshSceneProxyData::GatherData(const ARecastNavMesh* NavMesh, int32 InN
 		double AverageBuildTime = 0.;
 		double AverageCompLayersBuildTime = 0.;
 		double AverageNavLayersBuildTime = 0.;
+		double AverageLinkBuildTime = 0.;
 		double TotalBuildLinkTime = 0.;
 		if (bGatherTileBuildTimes || NavMeshGeometry.bGatherTileBuildTimesHeatMap)
 		{
@@ -544,16 +545,18 @@ void FNavMeshSceneProxyData::GatherData(const ARecastNavMesh* NavMesh, int32 InN
 					AverageBuildTime += Pair.Value.BuildTime;
 					AverageCompLayersBuildTime += Pair.Value.BuildCompressedLayerTime;
 					AverageNavLayersBuildTime += Pair.Value.BuildNavigationDataTime;
-					TotalBuildLinkTime += Pair.Value.BuildLinkTime;
+					AverageLinkBuildTime += Pair.Value.BuildLinkTime;
 				}
 			}
 
 			TotalTileBuildTime = AverageBuildTime;
+			TotalBuildLinkTime = AverageLinkBuildTime;
 			if (DebugDataMap->Num() != 0)
 			{
 				AverageBuildTime /= DebugDataMap->Num();
 				AverageCompLayersBuildTime /= DebugDataMap->Num();
 				AverageNavLayersBuildTime /= DebugDataMap->Num();
+				AverageLinkBuildTime /= DebugDataMap->Num();
 			}
 		}
 
@@ -630,8 +633,9 @@ void FNavMeshSceneProxyData::GatherData(const ARecastNavMesh* NavMesh, int32 InN
 			{
 				DebugLabels.Add(FDebugText(FString::Printf(TEXT("Tile count: %i"), DebugDataMap->Num())));
 				DebugLabels.Add(FDebugText(FString::Printf(TEXT("Avg tile build time: %0.2f ms"), AverageBuildTime*1000.)));
-				DebugLabels.Add(FDebugText(FString::Printf(TEXT("   Avg comp layers time: %0.1f ms"), AverageCompLayersBuildTime*1000.)));
-				DebugLabels.Add(FDebugText(FString::Printf(TEXT("   Avg nav layers time: %0.1f ms"), AverageNavLayersBuildTime*1000.)));
+				DebugLabels.Add(FDebugText(FString::Printf(TEXT("   Avg comp layers time: %0.2f ms"), AverageCompLayersBuildTime*1000.)));
+				DebugLabels.Add(FDebugText(FString::Printf(TEXT("   Avg nav layers time: %0.2f ms"), AverageNavLayersBuildTime*1000.)));
+				DebugLabels.Add(FDebugText(FString::Printf(TEXT("   Avg link build time: %0.2f ms"), AverageLinkBuildTime*1000.)));
 				DebugLabels.Add(FDebugText(FString::Printf(TEXT("Min: %0.2f ms  Max: %0.2f ms"), NavMeshGeometry.MinTileBuildTime*1000., NavMeshGeometry.MaxTileBuildTime*1000.)));
 				DebugLabels.Add(FDebugText(FString::Printf(TEXT("Total: %0.3f s"), TotalTileBuildTime)));
 				DebugLabels.Add(FDebugText(FString::Printf(TEXT("   Total link build time: %.4f s"), TotalBuildLinkTime)));
diff --git a/Engine/Source/Runtime/NavigationSystem/Private/NavMesh/RecastNavMeshGenerator.cpp b/Engine/Source/Runtime/NavigationSystem/Private/NavMesh/RecastNavMeshGenerator.cpp
--- a/Engine/Source/Runtime/NavigationSystem/Private/NavMesh/RecastNavMeshGenerator.cpp
+++ b/Engine/Source/Runtime/NavigationSystem/Private/NavMesh/RecastNavMeshGenerator.cpp
@@ -85,6 +85,9 @@ namespace UE::NavMesh::Private
 
 	static bool bUseTightBoundExpansion = true;
 	static FAutoConsoleVariableRef CVarUseTightBoundExpansion(TEXT("ai.nav.UseTightBoundExpansion"), bUseTightBoundExpansion, TEXT("Active by default. Use an expansion of one AgentRadius. Set to false to revert to the previous behavior (2 AgentRadius)."), ECVF_Default);
+
+	static bool bUseAsymetricBorderSizes = true;
+	static FAutoConsoleVariableRef CVarUseAsymetricBorderSizes(TEXT("ai.nav.UseAsymetricBorderSizes"), bUseAsymetricBorderSizes, TEXT("Active by default. When generating links, use asymetric tile border sizes to improve generation speed."), ECVF_Default);
 }
 
 static FOodleDataCompression::ECompressor GNavmeshTileCacheCompressor = FOodleDataCompression::ECompressor::Mermaid;
@@ -1622,17 +1625,17 @@ struct FVoxelCacheRasterizeContext
 		rcResetHeightfield(*RasterizeHF);
 	}
 
-	void SetupForTile(const FVector::FReal* TileBMin, const FVector::FReal* TileBMax, const FVector::FReal RasterizationPadding)
+	void SetupForTile(const FVector::FReal* TileBMin, const FVector::FReal* TileBMax, const FVector::FReal RasterizationLowPadding, const FVector::FReal RasterizationHighPadding)
 	{
 		Reset();
 
 		rcVcopy(RasterizeHF->bmin, TileBMin);
 		rcVcopy(RasterizeHF->bmax, TileBMax);
 
-		RasterizeHF->bmin[0] -= RasterizationPadding;
-		RasterizeHF->bmin[2] -= RasterizationPadding;
-		RasterizeHF->bmax[0] += RasterizationPadding;
-		RasterizeHF->bmax[2] += RasterizationPadding;
+		RasterizeHF->bmin[0] -= RasterizationLowPadding;
+		RasterizeHF->bmin[2] -= RasterizationLowPadding;
+		RasterizeHF->bmax[0] += RasterizationHighPadding;
+		RasterizeHF->bmax[2] += RasterizationHighPadding;
 	}
 
 	rcHeightfield* RasterizeHF;
@@ -2355,11 +2358,12 @@ void FRecastTileGenerator::PrepareVoxelCache(const TNavStatArray<uint8>& RawColl
 	// tile's geometry: voxel cache (only for synchronous rebuilds)
 	const int32 WalkableClimbVX = TileConfig.walkableClimb;
 	const FVector::FReal WalkableSlopeCos = FMath::Cos(FMath::DegreesToRadians(TileConfig.walkableSlopeAngle));
-	const FVector::FReal RasterizationPadding = TileConfig.borderSize * TileConfig.cs;
+	const FVector::FReal RasterizationLowPadding = TileConfig.borderSize.low * TileConfig.cs;
+	const FVector::FReal RasterizationHighPadding = TileConfig.borderSize.high * TileConfig.cs;
 
 	FRecastGeometryCache CachedCollisions(RawCollisionCache.GetData());
 
-	VoxelCacheContext.SetupForTile(TileConfig.bmin, TileConfig.bmax, RasterizationPadding);
+	VoxelCacheContext.SetupForTile(TileConfig.bmin, TileConfig.bmax, RasterizationLowPadding, RasterizationHighPadding);
 
 	float SlopeCosPerActor = UE_REAL_TO_FLOAT(WalkableSlopeCos);
 	CachedCollisions.Header.SlopeOverride.ModifyWalkableFloorZ(SlopeCosPerActor);
@@ -2682,14 +2686,16 @@ bool FRecastTileGenerator::CreateHeightField(FNavMeshBuildContext& BuildContext)
 
 	SCOPE_CYCLE_COUNTER(STAT_Navigation_RecastCreateHeightField);
 
-	TileConfig.width = TileConfig.tileSize + TileConfig.borderSize * 2;
-	TileConfig.height = TileConfig.tileSize + TileConfig.borderSize * 2;
+	const int size = TileConfig.tileSize + (TileConfig.borderSize.low + TileConfig.borderSize.high);
+	TileConfig.width = size;
+	TileConfig.height = size;
 
-	const FVector::FReal BBoxPadding = TileConfig.borderSize * TileConfig.cs;
-	TileConfig.bmin[0] -= BBoxPadding;
-	TileConfig.bmin[2] -= BBoxPadding;
-	TileConfig.bmax[0] += BBoxPadding;
-	TileConfig.bmax[2] += BBoxPadding;
+	const FVector::FReal BBoxPaddingLow = TileConfig.borderSize.low * TileConfig.cs;
+	const FVector::FReal BBoxPaddingHigh = TileConfig.borderSize.high * TileConfig.cs;
+	TileConfig.bmin[0] -= BBoxPaddingLow;
+	TileConfig.bmin[2] -= BBoxPaddingLow;
+	TileConfig.bmax[0] += BBoxPaddingHigh;
+	TileConfig.bmax[2] += BBoxPaddingHigh;
 
 	BuildContext.log(RC_LOG_PROGRESS, "CreateHeightField:");
 	BuildContext.log(RC_LOG_PROGRESS, " - %d x %d cells", TileConfig.width, TileConfig.height);
@@ -4887,7 +4893,8 @@ void FRecastNavMeshGenerator::SetupTileConfig(const ENavigationDataResolution Ti
 
 	// +1 for voxelization rounding, +1 for ledge neighbor access, +1 for occasional errors
 	const int BorderForAgentVx = OutConfig.walkableRadius + 3;
-	OutConfig.borderSize = BorderForAgentVx + MaxLinkDistanceVx; 
+	OutConfig.borderSize.low = UE::NavMesh::Private::bUseAsymetricBorderSizes ? BorderForAgentVx : BorderForAgentVx + MaxLinkDistanceVx;
+	OutConfig.borderSize.high = BorderForAgentVx + MaxLinkDistanceVx;
 
 	OutConfig.maxEdgeLen = (int32)(1200.0f / CellSize);
 
@@ -4966,7 +4973,9 @@ void FRecastNavMeshGenerator::ConfigureBuildProperties(FRecastBuildConfig& OutCo
 		JumpDown.CopyToDetourConfig(OutConfig.JumpDownConfig);
 		JumpOver.CopyToDetourConfig(OutConfig.JumpOverConfig);
 
-		OutConfig.LinkSpillDistance = FMath::Max((JumpDown.JumpLength-JumpDown.JumpDistanceFromEdge), (JumpOver.JumpLength-JumpOver.JumpDistanceFromEdge));
+		const float JumpDownSpillDistance = JumpDown.bEnabled ? JumpDown.JumpLength - JumpDown.JumpDistanceFromEdge : 0.f;
+		const float JumpOverSpillDistance = JumpOver.bEnabled ? JumpOver.JumpLength - JumpOver.JumpDistanceFromEdge : 0.f;
+		OutConfig.LinkSpillDistance = FMath::Max(JumpDownSpillDistance, JumpOverSpillDistance);
 		MaxLinkDistanceVx = FMath::CeilToInt(OutConfig.LinkSpillDistance / CellSize);
 	}
 	
@@ -4977,7 +4986,8 @@ void FRecastNavMeshGenerator::ConfigureBuildProperties(FRecastBuildConfig& OutCo
 
 	// +1 for voxelization rounding, +1 for ledge neighbor access, +1 for occasional errors
 	const int BorderForAgentVx = OutConfig.walkableRadius + 3;
-	OutConfig.borderSize = BorderForAgentVx + MaxLinkDistanceVx; 
+	OutConfig.borderSize.low = UE::NavMesh::Private::bUseAsymetricBorderSizes ? BorderForAgentVx : BorderForAgentVx + MaxLinkDistanceVx;
+	OutConfig.borderSize.high = BorderForAgentVx + MaxLinkDistanceVx;
 
 	OutConfig.maxEdgeLen = (int32)(1200.0f / CellSize);
 
@@ -5039,13 +5049,27 @@ void FRecastNavMeshGenerator::Init()
 
 	if (UE::NavMesh::Private::bUseTightBoundExpansion)
 	{
-		const rcReal HorizontalGrowth = Config.borderSize * Config.cs;
-		BBoxGrowth = FVector(HorizontalGrowth, HorizontalGrowth, Config.cs);
+		// Recast axis are inverted so growth direction are inverted here (using the positive side border for the low box growth and vice versa).
+		const rcReal HorizontalGrowhtLow = Config.borderSize.high * Config.cs;
+		BBoxGrowthLow = FVector(HorizontalGrowhtLow, HorizontalGrowhtLow, Config.cs);
+
+		const rcReal HorizontalGrowthHigh = Config.borderSize.low * Config.cs;
+		BBoxGrowthHigh = FVector(HorizontalGrowthHigh, HorizontalGrowthHigh, Config.cs);
+
+    PRAGMA_DISABLE_DEPRECATION_WARNINGS
+	    BBoxGrowth = FVector(static_cast<rcReal>(FMath::Max(Config.borderSize.low, Config.borderSize.high)) * Config.cs);
+    PRAGMA_ENABLE_DEPRECATION_WARNINGS
 	}
 	else
 	{
-		// Deprecated
-		BBoxGrowth = FVector(2.0 * static_cast<rcReal>(Config.borderSize) * Config.cs);
+		// Not using bUseTightBoundExpansion is deprecated, setting values mimicking previous behavior.
+
+    PRAGMA_DISABLE_DEPRECATION_WARNINGS
+	    BBoxGrowth = FVector(2.0 * static_cast<rcReal>(FMath::Max(Config.borderSize.low, Config.borderSize.high)) * Config.cs);
+    PRAGMA_ENABLE_DEPRECATION_WARNINGS
+
+		BBoxGrowthLow = FVector(2.0 * static_cast<rcReal>(Config.borderSize.high) * Config.cs);
+		BBoxGrowthHigh = FVector(2.0 * static_cast<rcReal>(Config.borderSize.low) * Config.cs);
 	}
 	RcNavMeshOrigin = Unreal2RecastPoint(DestNavMesh->NavMeshOriginOffset);
 	
@@ -5073,7 +5097,7 @@ void FRecastNavMeshGenerator::Init()
 	// prepare voxel cache if needed
 	if (ARecastNavMesh::IsVoxelCacheEnabled())
 	{
-		VoxelCacheContext.Create(Config.tileSize + Config.borderSize * 2, Config.cs, Config.ch);
+		VoxelCacheContext.Create(Config.tileSize + Config.borderSize.low + Config.borderSize.high, Config.cs, Config.ch);
 	}
 
 	bInitialized = true;
@@ -6122,7 +6146,14 @@ FBox FRecastNavMeshGenerator::GrowBoundingBox(const FBox& BBox, bool bIncludeAge
 {
 	const FVector BBoxGrowOffsetMin = FVector(0, 0, bIncludeAgentHeight ? Config.AgentHeight : 0.0f);
 
-	return FBox(BBox.Min - BBoxGrowth - BBoxGrowOffsetMin, BBox.Max + BBoxGrowth);
+	return FBox(BBox.Min - BBoxGrowthLow - BBoxGrowOffsetMin, BBox.Max + BBoxGrowthHigh);
+}
+
+FBox FRecastNavMeshGenerator::GrowDirtyBounds(const FBox& BBox, bool bIncludeAgentHeight) const
+{
+	const FVector BBoxGrowOffsetMin = FVector(0, 0, bIncludeAgentHeight ? Config.AgentHeight : 0.0f);
+
+	return FBox(BBox.Min - BBoxGrowthHigh - BBoxGrowOffsetMin, BBox.Max + BBoxGrowthLow);
 }
 
 bool FRecastNavMeshGenerator::ShouldGenerateGeometryForOctreeElement(const FNavigationOctreeElement& Element, const FNavDataConfig& NavDataConfig) const
@@ -6302,7 +6333,7 @@ void FRecastNavMeshGenerator::MarkDirtyTiles(const TArray<FNavigationDirtyArea>&
 				}
 
 				const FBox CutDownArea = CalculateBoxIntersection(GetTotalBounds(), DirtyArea.Bounds);
-				AdjustedAreaBounds = GrowBoundingBox(CutDownArea, DirtyArea.HasFlag(ENavigationDirtyFlag::UseAgentHeight));
+				AdjustedAreaBounds = GrowDirtyBounds(CutDownArea, DirtyArea.HasFlag(ENavigationDirtyFlag::UseAgentHeight));
 
 				// @TODO this and the following test share some work in common
 				if (IntersectBounds(AdjustedAreaBounds, InclusionBounds) == false)
diff --git a/Engine/Source/Runtime/NavigationSystem/Public/NavMesh/LinkGenerationConfig.h b/Engine/Source/Runtime/NavigationSystem/Public/NavMesh/LinkGenerationConfig.h
--- a/Engine/Source/Runtime/NavigationSystem/Public/NavMesh/LinkGenerationConfig.h
+++ b/Engine/Source/Runtime/NavigationSystem/Public/NavMesh/LinkGenerationConfig.h
@@ -16,7 +16,7 @@ struct FNavLinkGenerationJumpDownConfig
 
 	/** Should this config be used to generate links. */
 	UPROPERTY(EditAnywhere, Config, Category = Settings)
-	bool bEnable = true;
+	bool bEnabled = true;
 	
 	// @todo: Rename, describe and find best defaults for those parameters.
 	UPROPERTY(EditAnywhere, Config, Category = Settings, meta=(Units=cm, UIMin=0, ClampMin=0))
diff --git a/Engine/Source/Runtime/NavigationSystem/Public/NavMesh/RecastNavMeshGenerator.h b/Engine/Source/Runtime/NavigationSystem/Public/NavMesh/RecastNavMeshGenerator.h
--- a/Engine/Source/Runtime/NavigationSystem/Public/NavMesh/RecastNavMeshGenerator.h
+++ b/Engine/Source/Runtime/NavigationSystem/Public/NavMesh/RecastNavMeshGenerator.h
@@ -1026,7 +1026,15 @@ protected:
 		return TileGenerator;
 	}
 
-	void SetBBoxGrowth(const FVector& InBBox) { BBoxGrowth = InBBox; }
+	UE_DEPRECATED(5.5, "Use BBoxGrowthLow and BBoxGrowthHigh properties instead.")
+	void SetBBoxGrowth(const FVector& InBBox)
+	{
+PRAGMA_DISABLE_DEPRECATION_WARNINGS		
+		BBoxGrowth = InBBox;
+PRAGMA_ENABLE_DEPRECATION_WARNINGS
+	}
+
+	FBox GrowDirtyBounds(const FBox& BBox, bool bIncludeAgentHeight) const;
 
 	//----------------------------------------------------------------------//
 	// debug
@@ -1061,8 +1069,17 @@ protected:
 
 	/** Used to grow generic element bounds to match this generator's properties
 	 *	(most notably Config.borderSize) */
+	UE_DEPRECATED(5.5, "Use BBoxGrowthLow and BBoxGrowthHigh instead.")
 	FVector BBoxGrowth;
 	
+	/** Growth in the negative axis direction. 
+	 * Used to grow generic element bounds to match this generator's properties (most notably Config.borderSize) */
+	FVector BBoxGrowthLow;
+
+	/** Growth in the positive axis direction. 
+	 * Used to grow generic element bounds to match this generator's properties (most notably Config.borderSize) */
+	FVector BBoxGrowthHigh;
+	
 	int32 NumActiveTiles;
 	/** the limit to number of asynchronous tile generators running at one time,
 	 *  this is also used by the sync non time sliced regeneration ProcessTileTasksSync,
diff --git a/Engine/Source/Runtime/Navmesh/Private/Recast/RecastContour.cpp b/Engine/Source/Runtime/Navmesh/Private/Recast/RecastContour.cpp
--- a/Engine/Source/Runtime/Navmesh/Private/Recast/RecastContour.cpp
+++ b/Engine/Source/Runtime/Navmesh/Private/Recast/RecastContour.cpp
@@ -598,7 +598,7 @@ static bool mergeContours(rcContour& ca, rcContour& cb, int ia, int ib)
 /// Simplified contours are generated such that the vertices for portals between areas match up. 
 /// (They are considered mandatory vertices.)
 ///
-/// Setting @p maxEdgeLength to zero will disabled the edge length feature.
+/// Setting @p maxEdgeLength to zero will disable the edge length feature.
 /// 
 /// See the #rcConfig documentation for more information on the configuration parameters.
 /// 
@@ -611,25 +611,27 @@ bool rcBuildContours(rcContext* ctx, rcCompactHeightfield& chf,
 	
 	const int w = chf.width;
 	const int h = chf.height;
-	const int borderSize = chf.borderSize;
+	const rcBorderSize borderSize = chf.borderSize;		//@UE
 	
 	ctx->startTimer(RC_TIMER_BUILD_CONTOURS);
 	
 	rcVcopy(cset.bmin, chf.bmin);
 	rcVcopy(cset.bmax, chf.bmax);
-	if (borderSize > 0)
+	if (borderSize.low > 0 || borderSize.high > 0)		//@UE
 	{
 		// If the heightfield was build with bordersize, remove the offset.
-		const rcReal pad = borderSize*chf.cs;
+		const rcReal pad = borderSize.low*chf.cs;		//@UE
 		cset.bmin[0] += pad;
 		cset.bmin[2] += pad;
-		cset.bmax[0] -= pad;
-		cset.bmax[2] -= pad;
+		
+		const rcReal padHigh = borderSize.high*chf.cs;	//@UE
+		cset.bmax[0] -= padHigh;						//@UE
+		cset.bmax[2] -= padHigh;						//@UE
 	}
 	cset.cs = chf.cs;
 	cset.ch = chf.ch;
-	cset.width = chf.width - chf.borderSize*2;
-	cset.height = chf.height - chf.borderSize*2;
+	cset.width = chf.width - (chf.borderSize.low+chf.borderSize.high);		//@UE
+	cset.height = chf.height - (chf.borderSize.low+chf.borderSize.high);	//@UE
 	cset.borderSize = chf.borderSize;
 	
 	int maxContours = rcMax((int)chf.maxRegions, 8);
@@ -749,14 +751,14 @@ bool rcBuildContours(rcContext* ctx, rcCompactHeightfield& chf,
 						return false;
 					}
 					memcpy(cont->verts, &simplified[0], sizeof(int)*cont->nverts*4);
-					if (borderSize > 0)
+					if (borderSize.low > 0)				//@UE
 					{
 						// If the heightfield was build with bordersize, remove the offset.
 						for (int j = 0; j < cont->nverts; ++j)
 						{
 							int* v = &cont->verts[j*4];
-							v[0] -= borderSize;
-							v[2] -= borderSize;
+							v[0] -= borderSize.low;		//@UE
+							v[2] -= borderSize.low;		//@UE
 						}
 					}
 					
@@ -768,28 +770,17 @@ bool rcBuildContours(rcContext* ctx, rcCompactHeightfield& chf,
 						return false;
 					}
 					memcpy(cont->rverts, &verts[0], sizeof(int)*cont->nrverts*4);
-					if (borderSize > 0)
+					if (borderSize.low > 0)				//@UE
 					{
 						// If the heightfield was build with bordersize, remove the offset.
 						for (int j = 0; j < cont->nrverts; ++j)
 						{
 							int* v = &cont->rverts[j*4];
-							v[0] -= borderSize;
-							v[2] -= borderSize;
+							v[0] -= borderSize.low;		//@UE
+							v[2] -= borderSize.low;		//@UE
 						}
 					}
 					
-/*					cont->cx = cont->cy = cont->cz = 0;
-					for (int i = 0; i < cont->nverts; ++i)
-					{
-						cont->cx += cont->verts[i*4+0];
-						cont->cy += cont->verts[i*4+1];
-						cont->cz += cont->verts[i*4+2];
-					}
-					cont->cx /= cont->nverts;
-					cont->cy /= cont->nverts;
-					cont->cz /= cont->nverts;*/
-					
 					cont->reg = reg;
 					cont->area = area;
 				}
diff --git a/Engine/Source/Runtime/Navmesh/Private/Recast/RecastLayers.cpp b/Engine/Source/Runtime/Navmesh/Private/Recast/RecastLayers.cpp
--- a/Engine/Source/Runtime/Navmesh/Private/Recast/RecastLayers.cpp
+++ b/Engine/Source/Runtime/Navmesh/Private/Recast/RecastLayers.cpp
@@ -100,7 +100,7 @@ struct rcLayerSweepSpan
 	unsigned short nei;	// neighbour id
 };
 
-static bool CollectLayerRegionsMonotone(rcContext* ctx, rcCompactHeightfield& chf, const int borderSize,
+static bool CollectLayerRegionsMonotone(rcContext* ctx, rcCompactHeightfield& chf, const rcBorderSize borderSize,	//@UE
 	unsigned short* srcReg, rcLayerRegionMonotone*& regs, int& nregs)
 {
 	const int w = chf.width;
@@ -120,14 +120,14 @@ static bool CollectLayerRegionsMonotone(rcContext* ctx, rcCompactHeightfield& ch
 	rcIntArray prev(256);
 	unsigned short regId = 0;
 
-	for (int y = borderSize; y < h - borderSize; ++y)
+	for (int y = borderSize.low; y < h - borderSize.high; ++y)		//@UE
 	{
 		prev.resize(regId + 1);
 		memset(&prev[0], 0, sizeof(int)*regId);
 		unsigned short sweepId = 0;
 		unsigned int MaxSpanCount = 0;
 
-		for (int x = borderSize; x < w - borderSize; ++x)
+		for (int x = borderSize.low; x < w - borderSize.high; ++x)		//@UE
 		{
 			const rcCompactCell& c = chf.cells[x + y*w];
 			MaxSpanCount = rcMax(MaxSpanCount, c.count);
@@ -214,7 +214,7 @@ static bool CollectLayerRegionsMonotone(rcContext* ctx, rcCompactHeightfield& ch
 		}
 
 		// Remap local sweep ids to region ids.
-		for (int x = borderSize; x < w - borderSize; ++x)
+		for (int x = borderSize.low; x < w - borderSize.high; ++x)	//@UE
 		{
 			const rcCompactCell& c = chf.cells[x + y*w];
 			for (int i = (int)c.index, ni = (int)(c.index + c.count); i < ni; ++i)
@@ -313,7 +313,7 @@ static bool CollectLayerRegionsMonotone(rcContext* ctx, rcCompactHeightfield& ch
 }
 
 static bool CollectLayerRegionsChunky(rcContext* ctx, rcCompactHeightfield& chf,
-									  const int borderSize, const int chunkSize,
+									  const rcBorderSize borderSize, const int chunkSize,		//@UE
 									  unsigned short* srcReg, rcLayerRegionMonotone*& regs, int& nregs)
 {
 	const int w = chf.width;
@@ -330,12 +330,12 @@ static bool CollectLayerRegionsChunky(rcContext* ctx, rcCompactHeightfield& chf,
 	rcIntArray prev(256);
 	unsigned short regId = 0;
 
-	for (int chunkx = borderSize; chunkx < w-borderSize; chunkx += chunkSize)
+	for (int chunkx = borderSize.low; chunkx < w-borderSize.high; chunkx += chunkSize)		//@UE
 	{
-		for (int chunky = borderSize; chunky < h-borderSize; chunky += chunkSize)
+		for (int chunky = borderSize.low; chunky < h-borderSize.high; chunky += chunkSize)		//@UE
 		{
-			const int maxx = rcMin(chunkx + chunkSize, w-borderSize);
-			const int maxy = rcMin(chunky + chunkSize, h-borderSize);
+			const int maxx = rcMin(chunkx + chunkSize, w-borderSize.high);		//@UE
+			const int maxy = rcMin(chunky + chunkSize, h-borderSize.high);		//@UE
 
 			for (int y = chunky; y < maxy; ++y)
 			{
@@ -529,7 +529,7 @@ static bool CollectLayerRegionsChunky(rcContext* ctx, rcCompactHeightfield& chf,
 }
 
 static bool SplitAndStoreLayerRegions(rcContext* ctx, rcCompactHeightfield& chf,
-									  const int borderSize, const int walkableHeight,
+									  const rcBorderSize borderSize, const int walkableHeight,		//@UE
 									  unsigned short* srcReg, rcLayerRegionMonotone* regs, const int nregs,
 									  rcHeightfieldLayerSet& lset)
 {
@@ -720,17 +720,17 @@ static bool SplitAndStoreLayerRegions(rcContext* ctx, rcCompactHeightfield& chf,
 
 	const int w = chf.width;
 	const int h = chf.height;
-	const int lw = w - borderSize*2;
-	const int lh = h - borderSize*2;
+	const int lw = w - (borderSize.low + borderSize.high);		//@UE
+	const int lh = h - (borderSize.low + borderSize.high);		//@UE
 
 	// Build contracted bbox for layers.
 	rcReal bmin[3], bmax[3];
 	rcVcopy(bmin, chf.bmin);
 	rcVcopy(bmax, chf.bmax);
-	bmin[0] += borderSize*chf.cs;
-	bmin[2] += borderSize*chf.cs;
-	bmax[0] -= borderSize*chf.cs;
-	bmax[2] -= borderSize*chf.cs;
+	bmin[0] += borderSize.low*chf.cs;		//@UE
+	bmin[2] += borderSize.low*chf.cs;		//@UE
+	bmax[0] -= borderSize.high*chf.cs;		//@UE
+	bmax[2] -= borderSize.high*chf.cs;		//@UE
 
 	lset.nlayers = (int)layerId;
 
@@ -812,8 +812,8 @@ static bool SplitAndStoreLayerRegions(rcContext* ctx, rcCompactHeightfield& chf,
 		{
 			for (int x = 0; x < lw; ++x)
 			{
-				const int cx = borderSize+x;
-				const int cy = borderSize+y;
+				const int cx = borderSize.low+x;	//@UE
+				const int cy = borderSize.low+y;	//@UE
 				const rcCompactCell& c = chf.cells[cx+cy*w];
 				for (int j = (int)c.index, nj = (int)(c.index+c.count); j < nj; ++j)
 				{
@@ -860,8 +860,8 @@ static bool SplitAndStoreLayerRegions(rcContext* ctx, rcCompactHeightfield& chf,
 							// Valid connection mask
 							if (chf.areas[ai] != RC_NULL_AREA && lid == alid)
 							{
-								const int nx = ax - borderSize;
-								const int ny = ay - borderSize;
+								const int nx = ax - borderSize.low;		//@UE
+								const int ny = ay - borderSize.low;		//@UE
 								if (nx >= 0 && ny >= 0 && nx < lw && ny < lh)
 								{
 									con |= (unsigned char)(1 << dir);
@@ -893,7 +893,7 @@ static bool SplitAndStoreLayerRegions(rcContext* ctx, rcCompactHeightfield& chf,
 /// 
 /// @see rcAllocHeightfieldLayerSet, rcCompactHeightfield, rcHeightfieldLayerSet, rcConfig
 bool rcBuildHeightfieldLayersMonotone(rcContext* ctx, rcCompactHeightfield& chf,
-							  const int borderSize, const int walkableHeight,
+							  const rcBorderSize borderSize, const int walkableHeight,		//@UE
 							  rcHeightfieldLayerSet& lset)
 {
 	QUICK_SCOPE_CYCLE_COUNTER(STAT_Navigation_BuildHeightfieldLayersMonotone);
@@ -935,7 +935,7 @@ bool rcBuildHeightfieldLayersMonotone(rcContext* ctx, rcCompactHeightfield& chf,
 }
 
 bool rcBuildHeightfieldLayersChunky(rcContext* ctx, rcCompactHeightfield& chf,
-									const int borderSize, const int walkableHeight,
+									const rcBorderSize borderSize, const int walkableHeight,	//@UE
 									const int chunkSize,
 									rcHeightfieldLayerSet& lset)
 {
@@ -978,7 +978,7 @@ bool rcBuildHeightfieldLayersChunky(rcContext* ctx, rcCompactHeightfield& chf,
 }
 
 /// helper function from RecastRegion.cpp, requires distance data in compact height field
-bool rcGatherRegionsNoFilter(rcContext* ctx, rcCompactHeightfield& chf, const int borderSize, unsigned short* spanBuf4);
+bool rcGatherRegionsNoFilter(rcContext* ctx, rcCompactHeightfield& chf, const rcBorderSize borderSize, unsigned short* spanBuf4);	//@UE
 
 struct rcLayerRegion
 {
@@ -1107,7 +1107,7 @@ static void walkContour(int x, int y, int i, int dir, rcCompactHeightfield& chf,
 /// 
 /// @see rcAllocHeightfieldLayerSet, rcCompactHeightfield, rcHeightfieldLayerSet, rcConfig
 bool rcBuildHeightfieldLayers(rcContext* ctx, rcCompactHeightfield& chf,
-	const int borderSize, const int walkableHeight,
+	const rcBorderSize borderSize, const int walkableHeight,		//@UE
 	rcHeightfieldLayerSet& lset)
 {
 	QUICK_SCOPE_CYCLE_COUNTER(STAT_Navigation_BuildHeightfieldLayers);
@@ -1374,17 +1374,17 @@ bool rcBuildHeightfieldLayers(rcContext* ctx, rcCompactHeightfield& chf,
 	// Create layers.
 	rcAssert(lset.layers == 0);
 
-	const int lw = w - borderSize*2;
-	const int lh = h - borderSize*2;
+	const int lw = w - (borderSize.low+borderSize.high);	//@UE
+	const int lh = h - (borderSize.low+borderSize.high);	//@UE
 
 	// Build contracted bbox for layers.
 	rcReal bmin[3], bmax[3];
 	rcVcopy(bmin, chf.bmin);
 	rcVcopy(bmax, chf.bmax);
-	bmin[0] += borderSize*chf.cs;
-	bmin[2] += borderSize*chf.cs;
-	bmax[0] -= borderSize*chf.cs;
-	bmax[2] -= borderSize*chf.cs;
+	bmin[0] += borderSize.low*chf.cs;		//@UE
+	bmin[2] += borderSize.low*chf.cs;		//@UE
+	bmax[0] -= borderSize.high*chf.cs;		//@UE
+	bmax[2] -= borderSize.high*chf.cs;		//@UE
 
 	lset.nlayers = (int)maxLayerId;
 
@@ -1466,8 +1466,8 @@ bool rcBuildHeightfieldLayers(rcContext* ctx, rcCompactHeightfield& chf,
 		{
 			for (int x = 0; x < lw; ++x)
 			{
-				const int cx = borderSize+x;
-				const int cy = borderSize+y;
+				const int cx = borderSize.low+x;		//@UE
+				const int cy = borderSize.low+y;		//@UE
 				const rcCompactCell& c = chf.cells[cx+cy*w];
 				for (int j = (int)c.index, nj = (int)(c.index+c.count); j < nj; ++j)
 				{
@@ -1514,8 +1514,8 @@ bool rcBuildHeightfieldLayers(rcContext* ctx, rcCompactHeightfield& chf,
 							// Valid connection mask
 							if (chf.areas[ai] != RC_NULL_AREA && lid == alid)
 							{
-								const int nx = ax - borderSize;
-								const int ny = ay - borderSize;
+								const int nx = ax - borderSize.low;		//@UE
+								const int ny = ay - borderSize.low;		//@UE
 								if (nx >= 0 && ny >= 0 && nx < lw && ny < lh)
 								{
 									con |= (unsigned char)(1 << dir);
diff --git a/Engine/Source/Runtime/Navmesh/Private/Recast/RecastMesh.cpp b/Engine/Source/Runtime/Navmesh/Private/Recast/RecastMesh.cpp
--- a/Engine/Source/Runtime/Navmesh/Private/Recast/RecastMesh.cpp
+++ b/Engine/Source/Runtime/Navmesh/Private/Recast/RecastMesh.cpp
@@ -1165,7 +1165,7 @@ bool rcBuildPolyMesh(rcContext* ctx, rcContourSet& cset, const int nvp, rcPolyMe
 	}
 	
 	// Find portal edges
-	if (mesh.borderSize > 0)
+	if (mesh.borderSize.low > 0 || mesh.borderSize.high > 0)	//@UE
 	{
 		const int w = cset.width;
 		const int h = cset.height;
diff --git a/Engine/Source/Runtime/Navmesh/Private/Recast/RecastMeshDetail.cpp b/Engine/Source/Runtime/Navmesh/Private/Recast/RecastMeshDetail.cpp
--- a/Engine/Source/Runtime/Navmesh/Private/Recast/RecastMeshDetail.cpp
+++ b/Engine/Source/Runtime/Navmesh/Private/Recast/RecastMeshDetail.cpp
@@ -749,7 +749,7 @@ static bool buildPolyDetail(rcContext* ctx, const rcReal* in, const int nin,
 
 static void getHeightData(const rcCompactHeightfield& chf,
 						  const unsigned short* poly, const int npoly,
-						  const unsigned short* verts, const int bs,
+						  const unsigned short* verts, const rcBorderSize bs,	//@UE
 						  rcHeightPatch& hp, rcIntArray& stack)
 {
 	// Floodfill the heightfield to get 2D height data,
@@ -781,7 +781,7 @@ static void getHeightData(const rcCompactHeightfield& chf,
 				az < hp.ymin || az >= hp.ymin+hp.height)
 				continue;
 			
-			const rcCompactCell& c = chf.cells[(ax+bs)+(az+bs)*chf.width];
+			const rcCompactCell& c = chf.cells[(ax+bs.low)+(az+bs.low)*chf.width];	//@UE
 			for (int i = (int)c.index, ni = (int)(c.index+c.count); i < ni; ++i)
 			{
 				const rcCompactSpan& s = chf.spans[i];
@@ -853,7 +853,7 @@ static void getHeightData(const rcCompactHeightfield& chf,
 			if (hp.data[ax-hp.xmin+(ay-hp.ymin)*hp.width] != 0)
 				continue;
 			
-			const int ai = (int)chf.cells[(ax+bs)+(ay+bs)*chf.width].index + rcGetCon(cs, dir);
+			const int ai = (int)chf.cells[(ax+bs.low)+(ay+bs.low)*chf.width].index + rcGetCon(cs, dir);	//@UE
 
 			int idx = ax-hp.xmin+(ay-hp.ymin)*hp.width;
 			hp.data[idx] = 1;
@@ -909,7 +909,7 @@ static void getHeightData(const rcCompactHeightfield& chf,
 			if (hp.data[ax-hp.xmin+(ay-hp.ymin)*hp.width] != RC_UNSET_HEIGHT)
 				continue;
 			
-			const int ai = (int)chf.cells[(ax+bs)+(ay+bs)*chf.width].index + rcGetCon(cs, dir);
+			const int ai = (int)chf.cells[(ax+bs.low)+(ay+bs.low)*chf.width].index + rcGetCon(cs, dir);	//@UE
 			
 			const rcCompactSpan& as = chf.spans[ai];
 			int idx = ax-hp.xmin+(ay-hp.ymin)*hp.width;
@@ -967,7 +967,7 @@ bool rcBuildPolyMeshDetail(rcContext* ctx, const rcPolyMesh& mesh, const rcCompa
 	const rcReal cs = mesh.cs;
 	const rcReal ch = mesh.ch;
 	const rcReal* orig = mesh.bmin;
-	const int borderSize = mesh.borderSize;
+	const rcBorderSize borderSize = mesh.borderSize;	//@UE
 	
 	rcIntArray edges(64);
 	rcIntArray tris(512);
diff --git a/Engine/Source/Runtime/Navmesh/Private/Recast/RecastRegion.cpp b/Engine/Source/Runtime/Navmesh/Private/Recast/RecastRegion.cpp
--- a/Engine/Source/Runtime/Navmesh/Private/Recast/RecastRegion.cpp
+++ b/Engine/Source/Runtime/Navmesh/Private/Recast/RecastRegion.cpp
@@ -1009,7 +1009,7 @@ bool rcBuildDistanceField(rcContext* ctx, rcCompactHeightfield& chf)
 }
 
 static void paintRectRegion(int minx, int maxx, int miny, int maxy, unsigned short regId,
-							rcCompactHeightfield& chf, unsigned short* srcReg)
+							const rcCompactHeightfield& chf, unsigned short* srcReg)
 {
 	const int w = chf.width;	
 	for (int y = miny; y < maxy; ++y)
@@ -1026,6 +1026,30 @@ static void paintRectRegion(int minx, int maxx, int miny, int maxy, unsigned sho
 	}
 }
 
+// @UE BEGIN
+namespace UE::Recast::Private
+{
+	void rcMarkBorderRegions(const rcBorderSize borderSize, const int w, const int h,
+		rcCompactHeightfield& chf, unsigned short* srcReg, unsigned short& id)
+	{
+		if (borderSize.low > 0 || borderSize.high > 0)
+		{
+			const int lowW = rcMin(w, borderSize.low);
+			const int lowH = rcMin(h, borderSize.low);
+			const int highW = rcMin(w, borderSize.high);
+			const int highH = rcMin(h, borderSize.high);
+			
+			// Paint regions
+			paintRectRegion(0, lowW, 0, h, id|RC_BORDER_REG, chf, srcReg); id++;
+			paintRectRegion(w-highW, w, 0, h, id|RC_BORDER_REG, chf, srcReg); id++;
+			paintRectRegion(0, w, 0, lowH, id|RC_BORDER_REG, chf, srcReg); id++;
+			paintRectRegion(0, w, h-highH, h, id|RC_BORDER_REG, chf, srcReg); id++;
+
+			chf.borderSize = borderSize;
+		}
+	}
+} //namespace UE::Recast::Private
+// @UE END
 
 static const unsigned short RC_NULL_NEI = 0xffff;
 
@@ -1057,7 +1081,7 @@ struct rcSweepSpan
 /// 
 /// @see rcCompactHeightfield, rcCompactSpan, rcBuildDistanceField, rcBuildRegionsMonotone, rcConfig
 bool rcBuildRegionsMonotone(rcContext* ctx, rcCompactHeightfield& chf,
-							const int borderSize, const int minRegionArea, const int mergeRegionArea)
+							const rcBorderSize borderSize, const int minRegionArea, const int mergeRegionArea)	//@UE
 {
 	rcAssert(ctx);
 	
@@ -1082,34 +1106,20 @@ bool rcBuildRegionsMonotone(rcContext* ctx, rcCompactHeightfield& chf,
 		ctx->log(RC_LOG_ERROR, "rcBuildRegionsMonotone: Out of memory 'sweeps' (%d).", nsweeps);
 		return false;
 	}
-	
-	
-	// Mark border regions.
-	if (borderSize > 0)
-	{
-		// Make sure border will not overflow.
-		const int bw = rcMin(w, borderSize);
-		const int bh = rcMin(h, borderSize);
-		// Paint regions
-		paintRectRegion(0, bw, 0, h, id|RC_BORDER_REG, chf, srcReg); id++;
-		paintRectRegion(w-bw, w, 0, h, id|RC_BORDER_REG, chf, srcReg); id++;
-		paintRectRegion(0, w, 0, bh, id|RC_BORDER_REG, chf, srcReg); id++;
-		paintRectRegion(0, w, h-bh, h, id|RC_BORDER_REG, chf, srcReg); id++;
-		
-		chf.borderSize = borderSize;
-	}
+
+	UE::Recast::Private::rcMarkBorderRegions(borderSize, w, h, chf, srcReg, id);		//@UE
 	
 	rcIntArray prev(256);
 
 	// Sweep one line at a time.
-	for (int y = borderSize; y < h-borderSize; ++y)
+	for (int y = borderSize.low; y < h-borderSize.high; ++y)		//@UE
 	{
 		// Collect spans from this row.
 		prev.resize(id+1);
 		memset(&prev[0],0,sizeof(int)*id);
 		unsigned short rid = 1;
 		
-		for (int x = borderSize; x < w-borderSize; ++x)
+		for (int x = borderSize.low; x < w-borderSize.high; ++x)	//@UE
 		{
 			const rcCompactCell& c = chf.cells[x+y*w];
 			
@@ -1178,7 +1188,7 @@ bool rcBuildRegionsMonotone(rcContext* ctx, rcCompactHeightfield& chf,
 		}
 		
 		// Remap IDs
-		for (int x = borderSize; x < w-borderSize; ++x)
+		for (int x = borderSize.low; x < w-borderSize.high; ++x)	//@UE
 		{
 			const rcCompactCell& c = chf.cells[x+y*w];
 			
@@ -1209,7 +1219,7 @@ bool rcBuildRegionsMonotone(rcContext* ctx, rcCompactHeightfield& chf,
 }
 
 bool rcBuildRegionsChunky(rcContext* ctx, rcCompactHeightfield& chf,
-						  const int borderSize, const int minRegionArea, const int mergeRegionArea,
+						  const rcBorderSize borderSize, const int minRegionArea, const int mergeRegionArea,	//@UE
 						  const int regionChunkSize)
 {
 	rcAssert(ctx);
@@ -1236,29 +1246,15 @@ bool rcBuildRegionsChunky(rcContext* ctx, rcCompactHeightfield& chf,
 		return false;
 	}
 
-
-	// Mark border regions.
-	if (borderSize > 0)
-	{
-		// Make sure border will not overflow.
-		const int bw = rcMin(w, borderSize);
-		const int bh = rcMin(h, borderSize);
-		// Paint regions
-		paintRectRegion(0, bw, 0, h, id|RC_BORDER_REG, chf, srcReg); id++;
-		paintRectRegion(w-bw, w, 0, h, id|RC_BORDER_REG, chf, srcReg); id++;
-		paintRectRegion(0, w, 0, bh, id|RC_BORDER_REG, chf, srcReg); id++;
-		paintRectRegion(0, w, h-bh, h, id|RC_BORDER_REG, chf, srcReg); id++;
-
-		chf.borderSize = borderSize;
-	}
+	UE::Recast::Private::rcMarkBorderRegions(borderSize, w, h, chf, srcReg, id);		//@UE
 
 	rcIntArray prev(256);
-	for (int chunkx = borderSize; chunkx < w-borderSize; chunkx += regionChunkSize)
+	for (int chunkx = borderSize.low; chunkx < w-borderSize.high; chunkx += regionChunkSize)		//@UE
 	{
-		for (int chunky = borderSize; chunky < h-borderSize; chunky += regionChunkSize)
+		for (int chunky = borderSize.low; chunky < h-borderSize.high; chunky += regionChunkSize)		//@UE
 		{
-			const int maxx = rcMin(chunkx + regionChunkSize, w-borderSize);
-			const int maxy = rcMin(chunky + regionChunkSize, h-borderSize);
+			const int maxx = rcMin(chunkx + regionChunkSize, w-borderSize.high);		//@UE
+			const int maxy = rcMin(chunky + regionChunkSize, h-borderSize.high);		//@UE
 
 			// Sweep one line at a time.
 			for (int y = chunky; y < maxy; ++y)
@@ -1374,7 +1370,7 @@ bool rcBuildRegionsChunky(rcContext* ctx, rcCompactHeightfield& chf,
 ///		  
 ///		  spanBuf4 is temporary buffer, allocated and freed by caller (size = chf.spanCount * 4)
 ///
-bool rcGatherRegionsNoFilter(rcContext* ctx, rcCompactHeightfield& chf, const int borderSize, unsigned short* spanBuf4)
+bool rcGatherRegionsNoFilter(rcContext* ctx, rcCompactHeightfield& chf, const rcBorderSize borderSize, unsigned short* spanBuf4)		//@UE
 {
 	const int w = chf.width;
 	const int h = chf.height;
@@ -1399,19 +1395,7 @@ bool rcGatherRegionsNoFilter(rcContext* ctx, rcCompactHeightfield& chf, const in
 	//	const int expandIters = 4 + walkableRadius * 2;
 	const int expandIters = 8;
 
-	if (borderSize > 0)
-	{
-		// Make sure border will not overflow.
-		const int bw = rcMin(w, borderSize);
-		const int bh = rcMin(h, borderSize);
-		// Paint regions
-		paintRectRegion(0, bw, 0, h, regionId|RC_BORDER_REG, chf, srcReg); regionId++;
-		paintRectRegion(w-bw, w, 0, h, regionId|RC_BORDER_REG, chf, srcReg); regionId++;
-		paintRectRegion(0, w, 0, bh, regionId|RC_BORDER_REG, chf, srcReg); regionId++;
-		paintRectRegion(0, w, h-bh, h, regionId|RC_BORDER_REG, chf, srcReg); regionId++;
-
-		chf.borderSize = borderSize;
-	}
+	UE::Recast::Private::rcMarkBorderRegions(borderSize, w, h, chf, srcReg, regionId);		//@UE
 
 	while (level > 0)
 	{
@@ -1480,7 +1464,7 @@ bool rcGatherRegionsNoFilter(rcContext* ctx, rcCompactHeightfield& chf, const in
 /// 
 /// @see rcCompactHeightfield, rcCompactSpan, rcBuildDistanceField, rcBuildRegionsMonotone, rcConfig
 bool rcBuildRegions(rcContext* ctx, rcCompactHeightfield& chf,
-					const int borderSize, const int minRegionArea, const int mergeRegionArea)
+					const rcBorderSize borderSize, const int minRegionArea, const int mergeRegionArea)	//@UE
 {
 	rcAssert(ctx);
 	
diff --git a/Engine/Source/Runtime/Navmesh/Public/Recast/Recast.h b/Engine/Source/Runtime/Navmesh/Public/Recast/Recast.h
--- a/Engine/Source/Runtime/Navmesh/Public/Recast/Recast.h
+++ b/Engine/Source/Runtime/Navmesh/Public/Recast/Recast.h
@@ -245,6 +245,15 @@ enum rcRegionPartitioning
 	RC_REGION_CHUNKY,		///< monotone partitioning on small chunks
 };
 
+//@UE BEGIN
+/// Specifies the size of borders around the heightfield. 
+struct rcBorderSize
+{
+	int low;	///< Size of the border in the negative direction of the axis [Limit: >= 0] [Units: vx]
+	int high;	///< Size of the border in the positive direction of the axis [Limit: >= 0] [Units: vx]
+};
+//@UE END
+
 /// Specifies a configuration to use when performing Recast builds.
 /// @ingroup recast
 struct rcConfig
@@ -258,8 +267,8 @@ struct rcConfig
 	/// The width/height size of tile's on the xz-plane. [Limit: >= 0] [Units: vx]
 	int tileSize;
 	
-	/// The size of the non-navigable border around the heightfield. [Limit: >=0] [Units: vx]
-	int borderSize;
+	/// The size of the non-navigable border around the heightfield.
+	rcBorderSize borderSize;													//@UE														
 
 	/// The xz-plane cell size to use for fields. [Limit: > 0] [Units: wu] 
 	rcReal cs;
@@ -439,7 +448,7 @@ struct rcCompactHeightfield
 	int spanCount;				///< The number of spans in the heightfield.
 	int walkableHeight;			///< The walkable height used during the build of the field.  (See: rcConfig::walkableHeight)
 	int walkableClimb;			///< The walkable climb used during the build of the field. (See: rcConfig::walkableClimb)
-	int borderSize;				///< The AABB border size used during the build of the field. (See: rcConfig::borderSize)
+	rcBorderSize borderSize;	///< The AABB border size used during the build of the field. (See: rcConfig::borderSize)		//@UE
 	unsigned short maxDistance;	///< The maximum distance value of any span within the field. 
 	unsigned short maxRegions;	///< The maximum region id of any span within the field. 
 	rcReal bmin[3];				///< The minimum bounds in world space. [(x, y, z)]
@@ -497,15 +506,15 @@ struct rcContour
 /// @ingroup recast
 struct rcContourSet
 {
-	rcContour* conts;	///< An array of the contours in the set. [Size: #nconts]
-	int nconts;			///< The number of contours in the set.
-	rcReal bmin[3];  	///< The minimum bounds in world space. [(x, y, z)]
-	rcReal bmax[3];		///< The maximum bounds in world space. [(x, y, z)]
-	rcReal cs;			///< The size of each cell. (On the xz-plane.)
-	rcReal ch;			///< The height of each cell. (The minimum increment along the y-axis.)
-	int width;			///< The width of the set. (Along the x-axis in cell units.) 
-	int height;			///< The height of the set. (Along the z-axis in cell units.) 
-	int borderSize;		///< The AABB border size used to generate the source data from which the contours were derived.
+	rcContour* conts;			///< An array of the contours in the set. [Size: #nconts]
+	int nconts;					///< The number of contours in the set.
+	rcReal bmin[3];  			///< The minimum bounds in world space. [(x, y, z)]
+	rcReal bmax[3];				///< The maximum bounds in world space. [(x, y, z)]
+	rcReal cs;					///< The size of each cell. (On the xz-plane.)
+	rcReal ch;					///< The height of each cell. (The minimum increment along the y-axis.)
+	int width;					///< The width of the set. (Along the x-axis in cell units.) 
+	int height;					///< The height of the set. (Along the z-axis in cell units.) 
+	rcBorderSize borderSize;	///< The AABB border size used to generate the source data from which the contours were derived.		//@UE
 };
 
 // @UE BEGIN
@@ -539,7 +548,7 @@ struct rcPolyMesh
 	rcReal bmax[3];			///< The maximum bounds in world space. [(x, y, z)]
 	rcReal cs;				///< The size of each cell. (On the xz-plane.)
 	rcReal ch;				///< The height of each cell. (The minimum increment along the y-axis.)
-	int borderSize;			///< The AABB border size used to generate the source data from which the mesh was derived.
+	rcBorderSize borderSize;///< The AABB border size used to generate the source data from which the mesh was derived.		//@UE
 };
 
 /// Contains triangle meshes that represent detailed height data associated 
@@ -1305,7 +1314,7 @@ NAVMESH_API bool rcBuildRegions(rcContext* ctx, rcCompactHeightfield& chf,
 ///  								be merged with larger regions. [Limit: >=0] [Units: vx] 
 ///  @returns True if the operation completed successfully.
 NAVMESH_API bool rcBuildRegionsMonotone(rcContext* ctx, rcCompactHeightfield& chf,
-							const int borderSize, const int minRegionArea, const int mergeRegionArea);
+							const rcBorderSize borderSize, const int minRegionArea, const int mergeRegionArea);		//@UE
 
 /// Builds region data for the heightfield using simple monotone partitioning.
 ///  @ingroup recast 
@@ -1320,7 +1329,7 @@ NAVMESH_API bool rcBuildRegionsMonotone(rcContext* ctx, rcCompactHeightfield& ch
 ///	 @param[in]		chunkSize		Size of subregion [Units: vx]
 ///  @returns True if the operation completed successfully.
 NAVMESH_API bool rcBuildRegionsChunky(rcContext* ctx, rcCompactHeightfield& chf,
-						  const int borderSize, const int minRegionArea, const int mergeRegionArea,
+						  const rcBorderSize borderSize, const int minRegionArea, const int mergeRegionArea,		//@UE
 						  const int chunkSize);
 
 /// Sets the neighbor connection data for the specified direction.
@@ -1383,7 +1392,7 @@ inline int rcGetDirOffsetY(int dir)
 ///  @param[out]	lset		The resulting layer set. (Must be pre-allocated.)
 ///  @returns True if the operation completed successfully.
 NAVMESH_API bool rcBuildHeightfieldLayers(rcContext* ctx, rcCompactHeightfield& chf,
-							  const int borderSize, const int walkableHeight,
+							  const rcBorderSize borderSize, const int walkableHeight,		//@UE
 							  rcHeightfieldLayerSet& lset);
 
 /// Builds a layer set from the specified compact heightfield.
@@ -1397,7 +1406,7 @@ NAVMESH_API bool rcBuildHeightfieldLayers(rcContext* ctx, rcCompactHeightfield&
 ///  @param[out]	lset		The resulting layer set. (Must be pre-allocated.)
 ///  @returns True if the operation completed successfully.
 NAVMESH_API bool rcBuildHeightfieldLayersMonotone(rcContext* ctx, rcCompactHeightfield& chf,
-									  const int borderSize, const int walkableHeight,
+									  const rcBorderSize borderSize, const int walkableHeight,		//@UE
 									  rcHeightfieldLayerSet& lset);
 
 /// Builds a layer set from the specified compact heightfield.
@@ -1412,7 +1421,7 @@ NAVMESH_API bool rcBuildHeightfieldLayersMonotone(rcContext* ctx, rcCompactHeigh
 ///  @param[out]	lset		The resulting layer set. (Must be pre-allocated.)
 ///  @returns True if the operation completed successfully.
 NAVMESH_API bool rcBuildHeightfieldLayersChunky(rcContext* ctx, rcCompactHeightfield& chf,
-									const int borderSize, const int walkableHeight,
+									const rcBorderSize borderSize, const int walkableHeight,		//@UE
 									const int chunkSize,
 									rcHeightfieldLayerSet& lset);
 
