[NavMesh] Fix crash when changing TileSizeUU value

diff --git a/Engine/Source/Runtime/NavigationSystem/Private/NavMesh/RecastNavMesh.cpp b/Engine/Source/Runtime/NavigationSystem/Private/NavMesh/RecastNavMesh.cpp
--- a/Engine/Source/Runtime/NavigationSystem/Private/NavMesh/RecastNavMesh.cpp
+++ b/Engine/Source/Runtime/NavigationSystem/Private/NavMesh/RecastNavMesh.cpp
@@ -3251,8 +3251,8 @@ void ARecastNavMesh::PostEditChangeChainProperty(FPropertyChangedChainEvent& Pro
 						if (Index != ChangedIndex)
 						{
 							float& ResolutionCellSize = NavMeshResolutionParams[Index].CellSize;
-							const float ResolutionCellCount = FMath::TruncToFloat(TileSizeUU / ResolutionCellSize);
-							ResolutionCellSize = UE::NavMesh::Private::GetClampedCellSize(TileSizeUU / ResolutionCellCount);
+							const int32 ResolutionCellCount = FMath::RoundToInt(TileSizeUU / ResolutionCellSize);
+							ResolutionCellSize = UE::NavMesh::Private::GetClampedCellSize(TileSizeUU / (float)ResolutionCellCount);
 						}
 					}
 
@@ -3327,7 +3327,9 @@ void ARecastNavMesh::PostEditChangeProperty(FPropertyChangedEvent& PropertyChang
 				// Match cell sizes to tile size.
 				for (uint8 Index = 0; Index < (uint8)ENavigationDataResolution::MAX; Index++)
 				{
-					SetCellSize((ENavigationDataResolution)Index, TileSizeUU / FMath::TruncToFloat(TileSizeUU / GetCellSize((ENavigationDataResolution)Index)));	
+					const int32 CellCount = FMath::RoundToInt(TileSizeUU / GetCellSize((ENavigationDataResolution)Index));
+					const float NewCellSize = UE::NavMesh::Private::GetClampedCellSize(TileSizeUU / (float)CellCount);
+					SetCellSize((ENavigationDataResolution)Index, NewCellSize);
 				}
 
 				PRAGMA_DISABLE_DEPRECATION_WARNINGS
