Expose some RecastNavmesh  methods for use outside of the engine

diff --git a/Engine/Source/Runtime/NavigationSystem/Public/NavMesh/RecastInternalDebugData.h b/Engine/Source/Runtime/NavigationSystem/Public/NavMesh/RecastInternalDebugData.h
--- a/Engine/Source/Runtime/NavigationSystem/Public/NavMesh/RecastInternalDebugData.h
+++ b/Engine/Source/Runtime/NavigationSystem/Public/NavMesh/RecastInternalDebugData.h
@@ -65,10 +65,10 @@ struct FRecastInternalDebugData : public duDebugDraw
 		vertex(pos[0], pos[1], pos[2], color, uv[0], uv[1]);
 	}
 
-	virtual void vertex(const FVector::FReal x, const FVector::FReal y, const FVector::FReal z, unsigned int color, const FVector::FReal u, const FVector::FReal v) override;
+	NAVIGATIONSYSTEM_API virtual void vertex(const FVector::FReal x, const FVector::FReal y, const FVector::FReal z, unsigned int color, const FVector::FReal u, const FVector::FReal v) override;
 
-	virtual void text(const FVector::FReal x, const FVector::FReal y, const FVector::FReal z, const char* text) override;
+	NAVIGATIONSYSTEM_API virtual void text(const FVector::FReal x, const FVector::FReal y, const FVector::FReal z, const char* text) override;
 
-	virtual void end() override;
+	NAVIGATIONSYSTEM_API virtual void end() override;
 };
 #endif // WITH_RECAST
diff --git a/Engine/Source/Runtime/NavigationSystem/Public/NavMesh/RecastNavMeshGenerator.h b/Engine/Source/Runtime/NavigationSystem/Public/NavMesh/RecastNavMeshGenerator.h
--- a/Engine/Source/Runtime/NavigationSystem/Public/NavMesh/RecastNavMeshGenerator.h
+++ b/Engine/Source/Runtime/NavigationSystem/Public/NavMesh/RecastNavMeshGenerator.h
@@ -176,7 +176,7 @@ struct FRecastGeometryCache
 	int32* Indices;
 
 	FRecastGeometryCache() {}
-	FRecastGeometryCache(const uint8* Memory);
+	NAVIGATIONSYSTEM_API FRecastGeometryCache(const uint8* Memory);
 
 	static bool IsValid(const uint8* Memory, int32 MemorySize);
 };
