When in the editor with a World Partitioned map, only build the nav mesh inside of loaded areas

diff --git a/Engine/Source/Runtime/NavigationSystem/Private/NavigationSystem.cpp b/Engine/Source/Runtime/NavigationSystem/Private/NavigationSystem.cpp
--- a/Engine/Source/Runtime/NavigationSystem/Private/NavigationSystem.cpp
+++ b/Engine/Source/Runtime/NavigationSystem/Private/NavigationSystem.cpp
@@ -42,6 +42,7 @@
 #include "EditorModes.h"
 #include "LevelEditor.h"
 #include "Misc/MessageDialog.h"
+#include "WorldPartition/WorldPartition.h"
 #endif
 
 #include "NavAreas/NavArea_Null.h"
@@ -3716,6 +3717,7 @@ void UNavigationSystemV1::OnNavigationBoundsUpdated(ANavMeshBoundsVolume* NavVol
 		UpdateRequest.UpdateRequest = FNavigationBoundsUpdateRequest::Removed;
 	}
 
+	CheckToLimitNavigationBoundsToLoadedRegions(UpdateRequest.NavBounds);
 	AddNavigationBoundsUpdateRequest(UpdateRequest);
 }
 
@@ -3733,6 +3735,8 @@ void UNavigationSystemV1::OnNavigationBoundsAdded(ANavMeshBoundsVolume* NavVolum
 	UpdateRequest.NavBounds.SupportedAgents = NavVolume->SupportedAgents;
 
 	UpdateRequest.UpdateRequest = FNavigationBoundsUpdateRequest::Added;
+
+	CheckToLimitNavigationBoundsToLoadedRegions(UpdateRequest.NavBounds);
 	AddNavigationBoundsUpdateRequest(UpdateRequest);
 }
 
@@ -3750,12 +3754,74 @@ void UNavigationSystemV1::OnNavigationBoundsRemoved(ANavMeshBoundsVolume* NavVol
 	UpdateRequest.NavBounds.SupportedAgents = NavVolume->SupportedAgents;
 
 	UpdateRequest.UpdateRequest = FNavigationBoundsUpdateRequest::Removed;
+
+	CheckToLimitNavigationBoundsToLoadedRegions(UpdateRequest.NavBounds);
 	AddNavigationBoundsUpdateRequest(UpdateRequest);
 }
 
+void UNavigationSystemV1::CheckToLimitNavigationBoundsToLoadedRegions(FNavigationBounds& OutBounds) const
+{
+#if WITH_EDITOR && WITH_RECAST
+	// Find out if at least one of the nav meshes is world partitioned
+	bool bAnyWorldPartitionedNavMeshes = false;
+	for (const TObjectPtr<ANavigationData> NavData : NavDataSet)
+	{
+		const ARecastNavMesh* RecastNavMesh = Cast<ARecastNavMesh>(NavData);
+		if (RecastNavMesh && RecastNavMesh->bIsWorldPartitioned)
+		{
+			bAnyWorldPartitionedNavMeshes = true;
+			break;
+		}
+	}
+
+	// Don't limit nav bounds if none of the nav meshes are world partitioned
+	if (!bAnyWorldPartitionedNavMeshes)
+	{
+		return;
+	}
+
+	// Don't limit nav bounds at runtime
+	const UWorld* World = MainNavData->GetWorld();
+	if (!World || World->WorldType != EWorldType::Editor)
+	{
+		return;
+	}
+	
+	// Don't limit nav bounds if not in a world partitioned world
+	const UWorldPartition* WorldPartition = World->GetWorldPartition();
+	if (!WorldPartition)
+	{
+		return;
+	}
+	
+	// Get all loaded regions from the world partition
+	const TArray<FBox> LoadedWorldPartitionRegions = WorldPartition->GetUserLoadedEditorRegions();
+
+	// Store all overlaps between loaded world partition regions and UpdateRequest's nav bounds
+	TArray<FBox> OverlapRegions;
+	for (const FBox& LoadedWorldPartitionRegion : LoadedWorldPartitionRegions)
+	{
+		if (OutBounds.AreaBox.Intersect(LoadedWorldPartitionRegion))
+		{
+			OverlapRegions.Add(OutBounds.AreaBox.Overlap(LoadedWorldPartitionRegion));
+		}
+	}
+
+	// Merge all regions which overlap UpdateRequest's nav bounds
+	if (!OverlapRegions.IsEmpty())
+	{
+		OutBounds.AreaBox = FBox(ForceInitToZero);
+		for (const FBox& OverlapRegion : OverlapRegions)
+		{
+			OutBounds.AreaBox += OverlapRegion;
+		}
+	}
+#endif // WITH_EDITOR && WITH_RECAST
+}
+
 void UNavigationSystemV1::AddNavigationBoundsUpdateRequest(const FNavigationBoundsUpdateRequest& UpdateRequest)
 {
-	int32 ExistingIdx = PendingNavBoundsUpdates.IndexOfByPredicate([&](const FNavigationBoundsUpdateRequest& Element) {
+	const int32 ExistingIdx = PendingNavBoundsUpdates.IndexOfByPredicate([&](const FNavigationBoundsUpdateRequest& Element) {
 		return UpdateRequest.NavBounds.UniqueID == Element.NavBounds.UniqueID;
 	});
 
@@ -5543,15 +5609,15 @@ int UNavigationSystemV1::GetNavigationBoundsForNavData(const ANavigationData& Na
 
 	if (AgentIndex != INDEX_NONE)
 	{
-	for (const FNavigationBounds& NavigationBounds : RegisteredNavBounds)
-	{
-		if ((InLevel == nullptr || NavigationBounds.Level == InLevel)
-			&& NavigationBounds.SupportedAgents.Contains(AgentIndex))
+		for (const FNavigationBounds& NavigationBounds : RegisteredNavBounds)
 		{
-			OutBounds.Add(NavigationBounds.AreaBox);
+			if ((InLevel == nullptr || NavigationBounds.Level == InLevel)
+				&& NavigationBounds.SupportedAgents.Contains(AgentIndex))
+			{
+				OutBounds.Add(NavigationBounds.AreaBox);
+			}
 		}
 	}
-	}
 
 	return OutBounds.Num() - InitialBoundsCount;
 }
diff --git a/Engine/Source/Runtime/NavigationSystem/Public/NavigationSystem.h b/Engine/Source/Runtime/NavigationSystem/Public/NavigationSystem.h
--- a/Engine/Source/Runtime/NavigationSystem/Public/NavigationSystem.h
+++ b/Engine/Source/Runtime/NavigationSystem/Public/NavigationSystem.h
@@ -1344,9 +1344,13 @@ protected:
 public:
 	NAVIGATIONSYSTEM_API void DemandLazyDataGathering(FNavigationRelevantData& ElementData);
 
+private:
+	// Reduces a nav bounds only to regions overlapping loaded world partition regions (if any exist)
+	void CheckToLimitNavigationBoundsToLoadedRegions(FNavigationBounds& OutBounds) const;
+
 protected:
 	NAVIGATIONSYSTEM_API virtual void RebuildDirtyAreas(float DeltaSeconds);
-
+	
 	// adds navigation bounds update request to a pending list
 	NAVIGATIONSYSTEM_API void AddNavigationBoundsUpdateRequest(const FNavigationBoundsUpdateRequest& UpdateRequest);
 
