Fixed a bug with a loop variable

diff --git a/Engine/Source/Runtime/NavigationSystem/Private/NavigationSystem.cpp b/Engine/Source/Runtime/NavigationSystem/Private/NavigationSystem.cpp
--- a/Engine/Source/Runtime/NavigationSystem/Private/NavigationSystem.cpp
+++ b/Engine/Source/Runtime/NavigationSystem/Private/NavigationSystem.cpp
@@ -3764,7 +3764,7 @@ void UNavigationSystemV1::CheckToLimitNavigationBoundsToLoadedRegions(FNavigatio
 #if WITH_EDITOR && WITH_RECAST
 	// Find out if at least one of the nav meshes is world partitioned
 	bool bAnyWorldPartitionedNavMeshes = false;
-	for (const TObjectPtr<ANavigationData> NavData : NavDataSet)
+	for (const TObjectPtr<ANavigationData>& NavData : NavDataSet)
 	{
 		const ARecastNavMesh* RecastNavMesh = Cast<ARecastNavMesh>(NavData);
 		if (RecastNavMesh && RecastNavMesh->bIsWorldPartitioned)
