Templatize FRunningTileElement so that it can be used by other NavigationDataGenerator

diff --git a/Engine/Source/Runtime/NavigationSystem/Private/NavMesh/RecastNavMeshGenerator.cpp b/Engine/Source/Runtime/NavigationSystem/Private/NavMesh/RecastNavMeshGenerator.cpp
--- a/Engine/Source/Runtime/NavigationSystem/Private/NavMesh/RecastNavMeshGenerator.cpp
+++ b/Engine/Source/Runtime/NavigationSystem/Private/NavMesh/RecastNavMeshGenerator.cpp
@@ -5367,7 +5367,7 @@ void FRecastNavMeshGenerator::EnsureBuildCompletion()
 		ProcessTileTasksAndGetUpdatedTiles(NumTasksToProcess);
 		
 		// Block until tasks are finished
-		for (FRunningTileElement& Element : RunningDirtyTiles)
+		for (TRunningTileElement<FRecastTileGeneratorWrapper>& Element : RunningDirtyTiles)
 		{
 			Element.AsyncTask->EnsureCompletion();
 		}
@@ -6119,7 +6119,7 @@ void FRecastNavMeshGenerator::DiscardCurrentBuildingTasks()
 {
 	PendingDirtyTiles.Empty();
 	
-	for (FRunningTileElement& Element : RunningDirtyTiles)
+	for (TRunningTileElement<FRecastTileGeneratorWrapper>& Element : RunningDirtyTiles)
 	{
 		if (Element.AsyncTask)
 		{
@@ -6235,7 +6235,7 @@ int32 FRecastNavMeshGenerator::GetDirtyTilesCount(const FBox& AreaBounds) const
 	}
 
 	int32 RunningCount = 0;
-	for (const FRunningTileElement& RunningElement : RunningDirtyTiles)
+	for (const TRunningTileElement<FRecastTileGeneratorWrapper>& RunningElement : RunningDirtyTiles)
 	{
 		RunningCount += TileBox.Contains(RunningElement.Coord) ? 1 : 0;
 	}
@@ -6656,13 +6656,13 @@ TArray<FNavTileRef> FRecastNavMeshGenerator::ProcessTileTasksAsyncAndGetUpdatedT
 		QUICK_SCOPE_CYCLE_COUNTER(STAT_RecastNavMeshGenerator_ProcessTileTasks_NewTasks);
 
 		FPendingTileElement& PendingElement = PendingDirtyTiles[ElementIdx];
-		FRunningTileElement RunningElement(PendingElement.Coord);
+		TRunningTileElement<FRecastTileGeneratorWrapper> RunningElement(PendingElement.Coord);
 		
 		// Make sure that we are not submitting generator for grid cell that is currently being regenerated
 		if (!RunningDirtyTiles.Contains(RunningElement))
 		{
 			// Spawn async task
-			TUniquePtr<FRecastTileGeneratorTask> TileTask = MakeUnique<FRecastTileGeneratorTask>(CreateTileGenerator(PendingElement.Coord, PendingElement.DirtyAreas, PendingElement.CreationTime));
+			TUniquePtr<FAsyncTask<FRecastTileGeneratorWrapper>> TileTask = MakeUnique<FAsyncTask<FRecastTileGeneratorWrapper>>(CreateTileGenerator(PendingElement.Coord, PendingElement.DirtyAreas, PendingElement.CreationTime));
 
 			// Start it in background in case it has something to build
 			if (TileTask->GetTask().TileGenerator->HasDataToBuild())
@@ -6706,7 +6706,7 @@ TArray<FNavTileRef> FRecastNavMeshGenerator::ProcessTileTasksAsyncAndGetUpdatedT
 	{
 		QUICK_SCOPE_CYCLE_COUNTER(STAT_RecastNavMeshGenerator_ProcessTileTasks_FinishedTasks);
 
-		FRunningTileElement& Element = RunningDirtyTiles[Idx];
+		TRunningTileElement<FRecastTileGeneratorWrapper>& Element = RunningDirtyTiles[Idx];
 		check(Element.AsyncTask);
 
 		if (Element.AsyncTask->IsDone())
@@ -7437,7 +7437,7 @@ uint32 FRecastNavMeshGenerator::LogMemUsed() const
 	UE_LOG(LogNavigation, Display, TEXT("    FRecastNavMeshGenerator: self %d"), sizeof(FRecastNavMeshGenerator));
 	
 	uint32 GeneratorsMem = 0;
-	for (const FRunningTileElement& Element : RunningDirtyTiles)
+	for (const TRunningTileElement<FRecastTileGeneratorWrapper>& Element : RunningDirtyTiles)
 	{
 		GeneratorsMem += Element.AsyncTask->GetTask().TileGenerator->UsedMemoryOnStartup;
 		if (SyncTimeSlicedData.TileGeneratorSync.IsValid())
diff --git a/Engine/Source/Runtime/NavigationSystem/Public/NavMesh/RecastNavMeshGenerator.h b/Engine/Source/Runtime/NavigationSystem/Public/NavMesh/RecastNavMeshGenerator.h
--- a/Engine/Source/Runtime/NavigationSystem/Public/NavMesh/RecastNavMeshGenerator.h
+++ b/Engine/Source/Runtime/NavigationSystem/Public/NavMesh/RecastNavMeshGenerator.h
@@ -702,23 +702,24 @@ struct FPendingTileElement
 	}
 };
 
-struct FRunningTileElement
+template<typename TTileGeneratorTask>
+struct TRunningTileElement
 {
-	FRunningTileElement()
+	TRunningTileElement()
 		: Coord(FIntPoint::NoneValue)
 		, bShouldDiscard(false)
 		, AsyncTask(nullptr)
 	{
 	}
 	
-	FRunningTileElement(FIntPoint InCoord)
+	TRunningTileElement(FIntPoint InCoord)
 		: Coord(InCoord)
 		, bShouldDiscard(false)
 		, AsyncTask(nullptr)
 	{
 	}
 
-	bool operator == (const FRunningTileElement& Other) const
+	bool operator == (const TRunningTileElement& Other) const
 	{
 		return Coord == Other.Coord;
 	}
@@ -727,9 +728,12 @@ struct FRunningTileElement
 	FIntPoint					Coord;
 	/** whether generated results should be discarded */
 	bool						bShouldDiscard; 
-	FRecastTileGeneratorTask*	AsyncTask;
+	FAsyncTask<TTileGeneratorTask>* AsyncTask;
 };
 
+UE_DEPRECATED(5.5, "FRunningTileElement is deprecated. Please use TRunningTileElement<FRecastTileGeneratorWrapper> instead.")
+typedef TRunningTileElement<FRecastTileGeneratorWrapper> FRunningTileElement;
+
 struct FTileTimestamp
 {
 	FNavTileRef NavTileRef;
@@ -1102,7 +1106,7 @@ protected:
 	TNavStatArray<FPendingTileElement> PendingDirtyTiles;			
 	
 	/** List of dirty tiles currently being regenerated */
-	TNavStatArray<FRunningTileElement> RunningDirtyTiles;
+	TNavStatArray<TRunningTileElement<FRecastTileGeneratorWrapper>> RunningDirtyTiles;
 
 #if WITH_EDITOR
 	/** List of tiles that were recently regenerated */
