RecastNavmeshGenerator cleanup undefined methods and unused macros

diff --git a/Engine/Source/Runtime/NavigationSystem/Private/NavMesh/RecastNavMeshGenerator.cpp b/Engine/Source/Runtime/NavigationSystem/Private/NavMesh/RecastNavMeshGenerator.cpp
--- a/Engine/Source/Runtime/NavigationSystem/Private/NavMesh/RecastNavMeshGenerator.cpp
+++ b/Engine/Source/Runtime/NavigationSystem/Private/NavMesh/RecastNavMeshGenerator.cpp
@@ -55,8 +55,6 @@
 
 #define SHOW_NAV_EXPORT_PREVIEW 0
 
-#define TEXT_WEAKOBJ_NAME(obj) (obj.IsValid(false) ? *obj->GetName() : (obj.IsValid(false, true)) ? TEXT("MT-Unreachable") : TEXT("INVALID"))
-
 CSV_DEFINE_CATEGORY(NAVREGEN, false);
 
 struct dtTileCacheAlloc;
diff --git a/Engine/Source/Runtime/NavigationSystem/Public/NavMesh/RecastNavMeshGenerator.h b/Engine/Source/Runtime/NavigationSystem/Public/NavMesh/RecastNavMeshGenerator.h
--- a/Engine/Source/Runtime/NavigationSystem/Public/NavMesh/RecastNavMeshGenerator.h
+++ b/Engine/Source/Runtime/NavigationSystem/Public/NavMesh/RecastNavMeshGenerator.h
@@ -177,8 +177,6 @@ struct FRecastGeometryCache
 
 	FRecastGeometryCache() {}
 	NAVIGATIONSYSTEM_API FRecastGeometryCache(const uint8* Memory);
-
-	static bool IsValid(const uint8* Memory, int32 MemorySize);
 };
 
 struct FRecastRawGeometryElement
@@ -517,7 +515,6 @@ protected:
 	/** Appends specified geometry to tile's geometry */
 	NAVIGATIONSYSTEM_API void ValidateAndAppendGeometry(const TSharedRef<FNavigationRelevantData, ESPMode::ThreadSafe>& ElementData, const FCompositeNavModifier& InModifier);
 	NAVIGATIONSYSTEM_API void AppendGeometry(const FNavigationRelevantData& DataRef, const FCompositeNavModifier& InModifier, const FNavDataPerInstanceTransformDelegate& InTransformsDelegate);
-	NAVIGATIONSYSTEM_API void AppendVoxels(rcSpanCache* SpanData, int32 NumSpans);
 	
 	/** prepare voxel cache from collision data */
 	NAVIGATIONSYSTEM_API void PrepareVoxelCache(const TNavStatArray<uint8>& RawCollisionCache, const FCompositeNavModifier& InModifier, TNavStatArray<rcSpanCache>& SpanData);
