- Deprecating a handful of navigation functions that used raw tile indices. Instead they are replaced with versions that use safe FNavTileRefs
- Converting over current uses of the freshly-deprecated functions to work with FNavTileRefs
- Adding a new function to ARecastNavMesh in order to retrieve FNavTileRefs for all currently-active tiles in the mesh
- Providing a GetTypeHash implementation for FNavTileRefs so they can be used in certain container types

diff --git a/Engine/Source/Runtime/AIModule/Private/GameplayDebugger/GameplayDebuggerCategory_Navmesh.cpp b/Engine/Source/Runtime/AIModule/Private/GameplayDebugger/GameplayDebuggerCategory_Navmesh.cpp
--- a/Engine/Source/Runtime/AIModule/Private/GameplayDebugger/GameplayDebuggerCategory_Navmesh.cpp
+++ b/Engine/Source/Runtime/AIModule/Private/GameplayDebugger/GameplayDebuggerCategory_Navmesh.cpp
@@ -245,7 +245,7 @@ void FGameplayDebuggerCategory_Navmesh::CollectData(APlayerController* OwnerPC,
 		int32 TargetTileY = 0;
 		RecastNavMesh->GetNavMeshTileXY(TargetLocation, TargetTileX, TargetTileY);
 
-		TArray<int32> TileSet;
+		TArray<FNavTileRef> TileSet;
 		for (int32 Idx = 0; Idx < NumTilesToDisplay; Idx++)
 		{
 			const int32 NeiX = TargetTileX + DeltaX[Idx];
diff --git a/Engine/Source/Runtime/NavigationSystem/Private/NavMesh/NavMeshRenderingComponent.cpp b/Engine/Source/Runtime/NavigationSystem/Private/NavMesh/NavMeshRenderingComponent.cpp
--- a/Engine/Source/Runtime/NavigationSystem/Private/NavMesh/NavMeshRenderingComponent.cpp
+++ b/Engine/Source/Runtime/NavigationSystem/Private/NavMesh/NavMeshRenderingComponent.cpp
@@ -488,7 +488,12 @@ uint32 FNavMeshSceneProxyData::GetAllocatedSize() const
 
 #if WITH_RECAST
 
+// Deprecated
 void FNavMeshSceneProxyData::GatherData(const ARecastNavMesh* NavMesh, int32 InNavDetailFlags, const TArray<int32>& TileSet)
+{
+}
+
+void FNavMeshSceneProxyData::GatherData(const ARecastNavMesh* NavMesh, int32 InNavDetailFlags, const TArray<FNavTileRef>& TileSet)
 {
 	QUICK_SCOPE_CYCLE_COUNTER(STAT_NavMesh_GatherDebugDrawingGeometry);
 	Reset();
@@ -682,7 +687,7 @@ void FNavMeshSceneProxyData::GatherData(const ARecastNavMesh* NavMesh, int32 InN
 		}
 		else
 		{
-			NavMesh->GetDebugGeometryForTile(NavMeshGeometry, INDEX_NONE);
+			NavMesh->GetDebugGeometryForTile(NavMeshGeometry, FNavTileRef());
 		}
 
 		const TArray<FVector>& MeshVerts = NavMeshGeometry.MeshVerts;
@@ -800,31 +805,27 @@ void FNavMeshSceneProxyData::GatherData(const ARecastNavMesh* NavMesh, int32 InN
 		{
 			QUICK_SCOPE_CYCLE_COUNTER(STAT_NavMesh_GatherDebugDrawing_TileIterations);
 			
-			TArray<int32> UseTileIndices;
+			TArray<FNavTileRef> UseTileRefs;
 			if (TileSet.Num() > 0)
 			{
-				UseTileIndices = TileSet;
+				UseTileRefs = TileSet;
 			}
 			else
 			{
-				const int32 TilesCount = NavMesh->GetNavMeshTilesCount();
-				for (int32 Idx = 0; Idx < TilesCount; Idx++)
-				{
-					UseTileIndices.Add(Idx);
-				}
+				NavMesh->GetAllNavMeshTiles(UseTileRefs);
 			}
 
 			TMap<FIntPoint, FVector> TileBuildTimeLabelLocations;
 			
 			// calculate appropriate points for displaying debug labels
-			DebugLabels.Reserve(UseTileIndices.Num());
-			for (int32 TileSetIdx = 0; TileSetIdx < UseTileIndices.Num(); TileSetIdx++)
+			DebugLabels.Reserve(UseTileRefs.Num());
+			for (int32 TileSetIdx = 0; TileSetIdx < UseTileRefs.Num(); TileSetIdx++)
 			{
-				const int32 TileIndex = UseTileIndices[TileSetIdx];
+				const FNavTileRef TileRef = UseTileRefs[TileSetIdx];
 				int32 X, Y, Layer;
-				if (NavMesh->GetNavMeshTileXY(TileIndex, X, Y, Layer))
+				if (NavMesh->GetNavMeshTileXY(TileRef, X, Y, Layer))
 				{
-					const FBox TileBoundingBox = NavMesh->GetNavMeshTileBounds(TileIndex);
+					const FBox TileBoundingBox = NavMesh->GetNavMeshTileBounds(TileRef);
 					FVector TileLabelLocation = TileBoundingBox.GetCenter();
 					TileLabelLocation.Z = TileBoundingBox.Max.Z;
 
@@ -857,7 +858,7 @@ void FNavMeshSceneProxyData::GatherData(const ARecastNavMesh* NavMesh, int32 InN
 					if (bGatherPolygonLabels || bGatherPolygonCost || bGatherPolygonFlags)
 					{
 						TArray<FNavPoly> Polys;
-						NavMesh->GetPolysInTile(TileIndex, Polys);
+						NavMesh->GetPolysInTile(TileRef, Polys);
 
 						float DefaultCosts[RECAST_MAX_AREAS];
 						float FixedCosts[RECAST_MAX_AREAS];
@@ -914,7 +915,7 @@ void FNavMeshSceneProxyData::GatherData(const ARecastNavMesh* NavMesh, int32 InN
 
 					if (bGatherTileBounds)
 					{
-						const FBox TileBox = NavMesh->GetNavMeshTileBounds(TileIndex);
+						const FBox TileBox = NavMesh->GetNavMeshTileBounds(TileRef);
 						const FVector::FReal DrawZ = (TileBox.Min.Z + TileBox.Max.Z) * 0.5;
 						const FVector LL(TileBox.Min.X, TileBox.Min.Y, DrawZ);
 						const FVector UR(TileBox.Max.X, TileBox.Max.Y, DrawZ);
@@ -929,7 +930,7 @@ void FNavMeshSceneProxyData::GatherData(const ARecastNavMesh* NavMesh, int32 InN
 
 					if (bGatherTileResolutions)
 					{
-						const FBox TileBox = NavMesh->GetNavMeshTileBounds(TileIndex);
+						const FBox TileBox = NavMesh->GetNavMeshTileBounds(TileRef);
 						const FVector::FReal DrawZ = TileBox.Max.Z + NavMeshDrawOffset.Z;
 						constexpr FVector::FReal InsideOffset = 10.f;
 						const FVector LowerLeft(TileBox.Min.X + InsideOffset, TileBox.Min.Y + InsideOffset, DrawZ);
@@ -940,7 +941,7 @@ void FNavMeshSceneProxyData::GatherData(const ARecastNavMesh* NavMesh, int32 InN
 						FColor TileBoundsColor = FColor::Silver;
 						ENavigationDataResolution Resolution = ENavigationDataResolution::Invalid;
 						
-						if (NavMesh->GetNavmeshTileResolution(TileIndex, Resolution))
+						if (NavMesh->GetNavmeshTileResolution(TileRef, Resolution))
 						{
 							switch (Resolution)
 							{
@@ -1843,7 +1844,7 @@ void UNavMeshRenderingComponent::GatherData(const ARecastNavMesh& NavMesh, FNavM
 {
 #if WITH_RECAST
 	const int32 DetailFlags = FNavMeshRenderingHelpers::GetDetailFlags(&NavMesh);
-	const TArray<int32> EmptyTileSet;
+	const TArray<FNavTileRef> EmptyTileSet;
 	OutProxyData.GatherData(&NavMesh, DetailFlags, EmptyTileSet);
 #endif // WITH_RECAST
 }
diff --git a/Engine/Source/Runtime/NavigationSystem/Private/NavMesh/PImplRecastNavMesh.cpp b/Engine/Source/Runtime/NavigationSystem/Private/NavMesh/PImplRecastNavMesh.cpp
--- a/Engine/Source/Runtime/NavigationSystem/Private/NavMesh/PImplRecastNavMesh.cpp
+++ b/Engine/Source/Runtime/NavigationSystem/Private/NavMesh/PImplRecastNavMesh.cpp
@@ -510,7 +510,7 @@ void FPImplRecastNavMesh::Serialize( FArchive& Ar, int32 NavMeshVersion )
 	}
 
 	int32 NumTiles = 0;
-	TArray<int32> TilesToSave;
+	TArray<FNavTileRef> TilesToSave;
 
 	if (Ar.IsSaving())
 	{
@@ -544,7 +544,8 @@ void FPImplRecastNavMesh::Serialize( FArchive& Ar, int32 NavMeshVersion )
 					const dtMeshTile* Tile = ConstNavMesh->getTile(i);
 					if (Tile != NULL && Tile->header != NULL && Tile->dataSize > 0)
 					{
-						TilesToSave.Add(i);
+						FNavTileRef TileRef(ConstNavMesh->getTileRef(Tile));
+						TilesToSave.Add(TileRef);
 					}
 				}
 			}
@@ -700,9 +701,9 @@ void FPImplRecastNavMesh::Serialize( FArchive& Ar, int32 NavMeshVersion )
 		const bool bSupportsRuntimeGeneration = NavMeshOwner->SupportsRuntimeGeneration();
 		dtNavMesh const* ConstNavMesh = DetourNavMesh;
 		
-		for (int TileIndex : TilesToSave)
+		for (FNavTileRef TileRefToSave : TilesToSave)
 		{
-			const dtMeshTile* Tile = ConstNavMesh->getTile(TileIndex);
+			const dtMeshTile* Tile = ConstNavMesh->getTileByRef(static_cast<dtTileRef>(TileRefToSave));
 			dtTileRef TileRef = ConstNavMesh->getTileRef(Tile);
 			int32 TileDataSize = Tile->dataSize;
 			Ar << TileRef << TileDataSize;
@@ -2554,6 +2555,7 @@ bool FPImplRecastNavMesh::FilterPolys(TArray<NavNodeRef>& PolyRefs, const FRecas
 	return true;
 }
 
+// Deprecated
 bool FPImplRecastNavMesh::GetPolysInTile(int32 TileIndex, TArray<FNavPoly>& Polys) const
 {
 	if (DetourNavMesh == NULL || TileIndex < 0 || TileIndex >= DetourNavMesh->getMaxTiles())
@@ -2735,12 +2737,25 @@ uint8 GetValidEnds(const dtNavMesh& NavMesh, const dtMeshTile& Tile, const dtPol
 	return ValidEnds;
 }
 
+// Deprecated
 bool FPImplRecastNavMesh::GetDebugGeometryForTile(FRecastDebugGeometry& OutGeometry, int32 TileIndex) const
+{
+	FNavTileRef TileRef;
+	if (TileIndex != INDEX_NONE)
+	{
+		TArray<FNavTileRef> TileRefs;
+		FNavTileRef::DeprecatedMakeTileRefsFromTileIds(this, { static_cast<uint32>(TileIndex) }, TileRefs);
+		TileRef = TileRefs[0];
+	}
+	return GetDebugGeometryForTile(OutGeometry, TileRef);
+}
+
+bool FPImplRecastNavMesh::GetDebugGeometryForTile(FRecastDebugGeometry& OutGeometry, FNavTileRef TileRef) const
 {
 	QUICK_SCOPE_CYCLE_COUNTER(STAT_FPImplRecastNavMesh_GetDebugGeometryForTile);
 	
 	bool bDone = false;
-	if (DetourNavMesh == nullptr || TileIndex >= DetourNavMesh->getMaxTiles())
+	if (DetourNavMesh == nullptr)
 	{
 		bDone = true;
 		return bDone;
@@ -2786,9 +2801,9 @@ bool FPImplRecastNavMesh::GetDebugGeometryForTile(FRecastDebugGeometry& OutGeome
 		}
 	};
 
-	if (TileIndex != INDEX_NONE)
+	if (TileRef.IsValid())
 	{
-		dtMeshTile const* const Tile = ConstNavMesh->getTile(TileIndex);
+		dtMeshTile const* const Tile = ConstNavMesh->getTileByRef(static_cast<dtTileRef>(TileRef));
 		if (Tile != nullptr && Tile->header != nullptr)
 		{
 			const FIntPoint TileCoord = FIntPoint(Tile->header->x, Tile->header->y);
@@ -2799,7 +2814,7 @@ bool FPImplRecastNavMesh::GetDebugGeometryForTile(FRecastDebugGeometry& OutGeome
 				ReserveGeometryArrays(OutGeometry, NumVertsToReserve, NumIndicesToReserve);
 
 				const uint32 VertBase = OutGeometry.MeshVerts.Num();
-				GetTilesDebugGeometry(Generator, *Tile, VertBase, OutGeometry, TileIndex, ForbiddenFlags);
+				GetTilesDebugGeometry(Generator, *Tile, VertBase, OutGeometry, ConstNavMesh->getTileIndex(Tile), ForbiddenFlags);
 			}
 		}
 	}
@@ -3120,6 +3135,7 @@ FBox FPImplRecastNavMesh::GetNavMeshBounds() const
 	return Bbox;
 }
 
+// Deprecated
 FBox FPImplRecastNavMesh::GetNavMeshTileBounds(int32 TileIndex) const
 {
 	FBox Bbox(ForceInit);
@@ -3143,6 +3159,7 @@ FBox FPImplRecastNavMesh::GetNavMeshTileBounds(int32 TileIndex) const
 	return Bbox;
 }
 
+// Deprecated
 /** Retrieves XY coordinates of tile specified by index */
 bool FPImplRecastNavMesh::GetNavMeshTileXY(int32 TileIndex, int32& OutX, int32& OutY, int32& OutLayer) const
 {
@@ -3188,6 +3205,7 @@ bool FPImplRecastNavMesh::GetNavMeshTileXY(const FVector& Point, int32& OutX, in
 	return false;
 }
 
+// Deprecated
 void FPImplRecastNavMesh::GetNavMeshTilesAt(int32 TileX, int32 TileY, TArray<int32>& Indices) const
 {
 	if (DetourNavMesh)
@@ -3212,7 +3230,18 @@ void FPImplRecastNavMesh::GetNavMeshTilesAt(int32 TileX, int32 TileY, TArray<int
 	}
 }
 
+// Deprecated
 void FPImplRecastNavMesh::GetNavMeshTilesIn(const TArray<FBox>& InclusionBounds, TArray<int32>& Indices) const
+{
+	TArray<FNavTileRef> Refs;
+	GetNavMeshTilesIn(InclusionBounds, Refs);
+
+	TArray<uint32> UnsignedIndices;
+	FNavTileRef::DeprecatedGetTileIdsFromNavTileRefs(this, Refs, UnsignedIndices);
+	Indices.Append(UnsignedIndices);
+}
+
+void FPImplRecastNavMesh::GetNavMeshTilesIn(const TArray<FBox>& InclusionBounds, TArray<FNavTileRef>& OutRefs) const
 {
 	if (DetourNavMesh)
 	{
@@ -3239,7 +3268,7 @@ void FPImplRecastNavMesh::GetNavMeshTilesIn(const TArray<FBox>& InclusionBounds,
 		}
 
 		// We guess that each tile has 3 layers in average
-		Indices.Reserve(TileCoords.Num()*3);
+		OutRefs.Reserve(TileCoords.Num()*3);
 
 		TArray<const dtMeshTile*> MeshTiles;
 		MeshTiles.Reserve(3);
@@ -3264,8 +3293,7 @@ void FPImplRecastNavMesh::GetNavMeshTilesIn(const TArray<FBox>& InclusionBounds,
 						{
 							if (TileBounds.Intersect(RequestedBounds))
 							{
-								int32 TileIndex = (int32)DetourNavMesh->decodePolyIdTile(TileRef);
-								Indices.Add(TileIndex);
+								OutRefs.Add(FNavTileRef(TileRef));
 								break;
 							}
 						}
diff --git a/Engine/Source/Runtime/NavigationSystem/Private/NavMesh/RecastNavMesh.cpp b/Engine/Source/Runtime/NavigationSystem/Private/NavMesh/RecastNavMesh.cpp
--- a/Engine/Source/Runtime/NavigationSystem/Private/NavMesh/RecastNavMesh.cpp
+++ b/Engine/Source/Runtime/NavigationSystem/Private/NavMesh/RecastNavMesh.cpp
@@ -1260,20 +1260,76 @@ FBox ARecastNavMesh::GetNavMeshBounds() const
 	return Bounds;
 }
 
+// Deprecated
 FBox ARecastNavMesh::GetNavMeshTileBounds(int32 TileIndex) const
+{
+	if (RecastNavMeshImpl && (TileIndex != INDEX_NONE))
+	{
+		TArray<FNavTileRef> TileRefs;
+		FNavTileRef::DeprecatedMakeTileRefsFromTileIds(RecastNavMeshImpl, { static_cast<uint32>(TileIndex) }, TileRefs);
+		return GetNavMeshTileBounds(TileRefs[0]);
+	}
+
+	return FBox(ForceInit);
+}
+
+FBox ARecastNavMesh::GetNavMeshTileBounds(FNavTileRef TileRef) const
 {
 	FBox Bounds(ForceInit);
-	if (RecastNavMeshImpl)
+
+	if (HasValidNavmesh() && TileRef.IsValid())
 	{
-		Bounds = RecastNavMeshImpl->GetNavMeshTileBounds(TileIndex);
+		// workaround for privacy issue in the recast API
+		dtNavMesh const* const ConstRecastNavMesh = RecastNavMeshImpl->DetourNavMesh;
+
+		dtMeshTile const* const Tile = ConstRecastNavMesh->getTileByRef(static_cast<dtTileRef>(TileRef));
+		if (Tile)
+		{
+			dtMeshHeader const* const Header = Tile->header;
+			if (Header)
+			{
+				Bounds = Recast2UnrealBox(Header->bmin, Header->bmax);
+			}
+		}
 	}
 
 	return Bounds;
 }
 
+// Deprecated
 bool ARecastNavMesh::GetNavMeshTileXY(int32 TileIndex, int32& OutX, int32& OutY, int32& OutLayer) const
 {
-	return RecastNavMeshImpl && RecastNavMeshImpl->GetNavMeshTileXY(TileIndex, OutX, OutY, OutLayer);
+	if (RecastNavMeshImpl && (TileIndex != INDEX_NONE))
+	{
+		TArray<FNavTileRef> TileRefs;
+		FNavTileRef::DeprecatedMakeTileRefsFromTileIds(RecastNavMeshImpl, { static_cast<uint32>(TileIndex) }, TileRefs);
+		return GetNavMeshTileXY(TileRefs[0], OutX, OutY, OutLayer);
+	}
+	return false;
+}
+
+bool ARecastNavMesh::GetNavMeshTileXY(FNavTileRef TileRef, int32& OutX, int32& OutY, int32& OutLayer) const
+{
+	if (HasValidNavmesh() && TileRef.IsValid())
+	{
+		// workaround for privacy issue in the recast API
+		dtNavMesh const* const ConstRecastNavMesh = RecastNavMeshImpl->DetourNavMesh;
+
+		dtMeshTile const* const Tile = ConstRecastNavMesh->getTileByRef(static_cast<dtTileRef>(TileRef));
+		if (Tile)
+		{
+			dtMeshHeader const* const Header = Tile->header;
+			if (Header)
+			{
+				OutX = Header->x;
+				OutY = Header->y;
+				OutLayer = Header->layer;
+				return true;
+			}
+		}
+	}
+
+	return false;
 }
 
 bool ARecastNavMesh::GetNavMeshTileXY(const FVector& Point, int32& OutX, int32& OutY) const
@@ -1281,13 +1337,26 @@ bool ARecastNavMesh::GetNavMeshTileXY(const FVector& Point, int32& OutX, int32&
 	return RecastNavMeshImpl && RecastNavMeshImpl->GetNavMeshTileXY(Point, OutX, OutY);
 }
 
+// Deprecated
 bool ARecastNavMesh::GetNavmeshTileResolution(int32 TileIndex, ENavigationDataResolution& OutResolution) const
+{
+	if (RecastNavMeshImpl && (TileIndex != INDEX_NONE))
+	{
+		TArray<FNavTileRef> TileRefs;
+		FNavTileRef::DeprecatedMakeTileRefsFromTileIds(RecastNavMeshImpl, { static_cast<uint32>(TileIndex) }, TileRefs);
+		return GetNavmeshTileResolution(TileRefs[0], OutResolution);
+	}
+
+	return false;
+}
+
+bool ARecastNavMesh::GetNavmeshTileResolution(FNavTileRef TileRef, ENavigationDataResolution& OutResolution) const
 {
 	if (RecastNavMeshImpl)
 	{
 		if (const dtNavMesh* const RecastNavMesh = RecastNavMeshImpl->DetourNavMesh)
 		{
-			if (const dtMeshTile* const Tile = RecastNavMesh->getTile(TileIndex))
+			if (const dtMeshTile* const Tile = RecastNavMesh->getTileByRef(static_cast<dtTileRef>(TileRef)))
 			{
 				if (const dtMeshHeader* const Header = Tile->header)
 				{
@@ -1315,17 +1384,115 @@ bool ARecastNavMesh::CheckTileIndicesInValidRange(const FVector& Point, bool& bO
 	return bValidMesh;
 }
 
+// Deprecated
 void ARecastNavMesh::GetNavMeshTilesAt(int32 TileX, int32 TileY, TArray<int32>& Indices) const
 {
 	if (RecastNavMeshImpl)
 	{
-		RecastNavMeshImpl->GetNavMeshTilesAt(TileX, TileY, Indices);
+		TArray<FNavTileRef> Refs;
+		GetNavMeshTilesAt(TileX, TileY, Refs);
+
+		TArray<uint32> UnsignedIndices;
+		FNavTileRef::DeprecatedGetTileIdsFromNavTileRefs(RecastNavMeshImpl, Refs, UnsignedIndices);
+
+		Indices.Append(UnsignedIndices);
+	}
+}
+
+void ARecastNavMesh::GetNavMeshTilesAt(int32 TileX, int32 TileY, TArray<FNavTileRef>& OutRefs) const
+{	
+	if (HasValidNavmesh())
+	{
+		// workaround for privacy issue in the recast API
+		dtNavMesh const* const ConstRecastNavMesh = RecastNavMeshImpl->DetourNavMesh;
+
+		const int32 MaxTiles = ConstRecastNavMesh->getTileCountAt(TileX, TileY);
+		TArray<const dtMeshTile*> Tiles;
+		Tiles.AddZeroed(MaxTiles);
+
+		const int32 NumTiles = ConstRecastNavMesh->getTilesAt(TileX, TileY, Tiles.GetData(), MaxTiles);
+		for (int32 i = 0; i < NumTiles; i++)
+		{
+			dtTileRef TileRef = ConstRecastNavMesh->getTileRef(Tiles[i]);
+			if (TileRef)
+			{
+				OutRefs.Add(FNavTileRef(TileRef));
+			}
+		}
+	}
+}
+
+void ARecastNavMesh::GetAllNavMeshTiles(TArray<FNavTileRef>& OutRefs) const
+{
+	if (HasValidNavmesh())
+	{
+		// workaround for privacy issue in the recast API
+		const dtNavMesh* const ConstDetourNavMesh = RecastNavMeshImpl->DetourNavMesh;
+
+		const int32 TileCount = ConstDetourNavMesh->getMaxTiles();
+
+		OutRefs.Reserve(OutRefs.Num() + TileCount);
+		for (int32 TileIndex = 0; TileIndex < TileCount; ++TileIndex)
+		{
+			if (const dtMeshTile* const Tile = ConstDetourNavMesh->getTile(TileIndex))
+			{
+				if (const dtTileRef TileRef = ConstDetourNavMesh->getTileRef(Tile))
+				{
+					OutRefs.Add(FNavTileRef(TileRef));
+				}
+			}
+		}
 	}
 }
 
+// Deprecated
 bool ARecastNavMesh::GetPolysInTile(int32 TileIndex, TArray<FNavPoly>& Polys) const
 {
-	return RecastNavMeshImpl && RecastNavMeshImpl->GetPolysInTile(TileIndex, Polys);
+	if (RecastNavMeshImpl && (TileIndex != INDEX_NONE))
+	{
+		TArray<FNavTileRef> TileRefs;
+		FNavTileRef::DeprecatedMakeTileRefsFromTileIds(RecastNavMeshImpl, { static_cast<uint32>(TileIndex) }, TileRefs);
+		return GetPolysInTile(TileRefs[0], Polys);
+	}
+
+	return false;
+}
+
+bool ARecastNavMesh::GetPolysInTile(FNavTileRef TileRef, TArray<FNavPoly>& Polys) const
+{
+	if (HasValidNavmesh() && TileRef.IsValid())
+	{
+		// workaround for privacy issue in the recast API
+		const dtNavMesh* const ConstDetourNavMesh = RecastNavMeshImpl->DetourNavMesh;
+
+		const dtMeshTile* Tile = ConstDetourNavMesh->getTileByRef(static_cast<dtTileRef>(TileRef));
+		const int32 MaxPolys = Tile && Tile->header ? Tile->header->offMeshBase : 0;
+		if (MaxPolys > 0)
+		{
+			// only ground type polys
+			int32 BaseIdx = Polys.Num();
+			Polys.AddZeroed(MaxPolys);
+
+			dtPoly* Poly = Tile->polys;
+			for (int32 i = 0; i < MaxPolys; i++, Poly++)
+			{
+				FVector PolyCenter(0);
+				for (int k = 0; k < Poly->vertCount; ++k)
+				{
+					PolyCenter += Recast2UnrealPoint(&Tile->verts[Poly->verts[k] * 3]);
+				}
+				PolyCenter /= Poly->vertCount;
+
+				FNavPoly& OutPoly = Polys[BaseIdx + i];
+				OutPoly.Ref = ConstDetourNavMesh->encodePolyId(Tile->salt, ConstDetourNavMesh->getTileIndex(Tile), i);
+				OutPoly.Center = PolyCenter;
+			}
+		}
+
+		return (MaxPolys > 0);
+	}
+
+	return false;
 }
 
 bool ARecastNavMesh::GetNavLinksInTile(const int32 TileIndex, TArray<FNavPoly>& Polys, const bool bIncludeLinksFromNeighborTiles) const
@@ -2425,14 +2592,41 @@ bool ARecastNavMesh::GetPolysWithinPathingDistance(FVector const& StartLoc, cons
 // Deprecated
 void ARecastNavMesh::GetDebugGeometry(FRecastDebugGeometry& OutGeometry, int32 TileIndex) const
 {
-	GetDebugGeometryForTile(OutGeometry, TileIndex);
+	if (RecastNavMeshImpl)
+	{
+		FNavTileRef TileRef;
+		if (TileIndex != INDEX_NONE)
+		{
+			TArray<FNavTileRef> TileRefs;
+			FNavTileRef::DeprecatedMakeTileRefsFromTileIds(RecastNavMeshImpl, { static_cast<uint32>(TileIndex) }, TileRefs);
+			TileRef = TileRefs[0];
+		}
+		GetDebugGeometryForTile(OutGeometry, TileRef);
+	}
 }
 
+// Deprecated
 bool ARecastNavMesh::GetDebugGeometryForTile(FRecastDebugGeometry& OutGeometry, int32 TileIndex) const
 {
 	if (RecastNavMeshImpl)
 	{
-		return RecastNavMeshImpl->GetDebugGeometryForTile(OutGeometry, TileIndex);
+		FNavTileRef TileRef;
+		if (TileIndex != INDEX_NONE)
+		{
+			TArray<FNavTileRef> TileRefs;
+			FNavTileRef::DeprecatedMakeTileRefsFromTileIds(RecastNavMeshImpl, { static_cast<uint32>(TileIndex) }, TileRefs);
+			TileRef = TileRefs[0];
+		}
+		return GetDebugGeometryForTile(OutGeometry, TileRef);
+	}
+	return true;
+}
+
+bool ARecastNavMesh::GetDebugGeometryForTile(FRecastDebugGeometry& OutGeometry, FNavTileRef TileRef) const
+{
+	if (RecastNavMeshImpl)
+	{
+		return RecastNavMeshImpl->GetDebugGeometryForTile(OutGeometry, TileRef);
 	}
 	return true;
 }
@@ -2651,7 +2845,7 @@ void ARecastNavMesh::OnNavMeshGenerationFinished()
 				if (SupportsStreaming())
 				{
 					// We use navigation volumes that belongs to this streaming level to find tiles we want to save
-					TArray<int32> LevelTiles;
+					TArray<FNavTileRef> LevelTiles;
 					TArray<FBox> LevelNavBounds = GetNavigableBoundsInLevel(Level);
 					RecastNavMeshImpl->GetNavMeshTilesIn(LevelNavBounds, LevelTiles);
 
@@ -2787,17 +2981,17 @@ void ARecastNavMesh::FillNavigationDataChunkActor(const FBox& QueryBounds, ANavi
 	UE_LOG(LogNavigation, Verbose, TEXT("%s Bounds pos: (%s)  size: (%s)."), ANSI_TO_TCHAR(__FUNCTION__), *QueryBounds.GetCenter().ToString(), *QueryBounds.GetSize().ToString());
 
 	const TArray<FBox> Boxes({ QueryBounds });
-	TArray<int32> TileIndices;
-	RecastNavMeshImpl->GetNavMeshTilesIn(Boxes, TileIndices);
-	if (!TileIndices.IsEmpty())
+	TArray<FNavTileRef> TileRefs;
+	RecastNavMeshImpl->GetNavMeshTilesIn(Boxes, TileRefs);
+	if (!TileRefs.IsEmpty())
 	{
 		// Add a data chunk for this navmesh
 		URecastNavMeshDataChunk* DataChunk = NewObject<URecastNavMeshDataChunk>(&DataChunkActor);
 		DataChunk->NavigationDataName = GetFName();
 		DataChunkActor.GetMutableNavDataChunk().Add(DataChunk);
 
-		DataChunk->GetTiles(RecastNavMeshImpl, TileIndices, SupportsRuntimeGeneration() ? EGatherTilesCopyMode::CopyDataAndCacheData : EGatherTilesCopyMode::CopyData);
-		DataChunk->GetTilesBounds(*RecastNavMeshImpl, TileIndices, OutTilesBounds);
+		DataChunk->GetTiles(RecastNavMeshImpl, TileRefs, SupportsRuntimeGeneration() ? EGatherTilesCopyMode::CopyDataAndCacheData : EGatherTilesCopyMode::CopyData);
+		DataChunk->GetTilesBounds(*RecastNavMeshImpl, TileRefs, OutTilesBounds);
 	}
 }
 
diff --git a/Engine/Source/Runtime/NavigationSystem/Private/NavMesh/RecastNavMeshDataChunk.cpp b/Engine/Source/Runtime/NavigationSystem/Private/NavMesh/RecastNavMeshDataChunk.cpp
--- a/Engine/Source/Runtime/NavigationSystem/Private/NavMesh/RecastNavMeshDataChunk.cpp
+++ b/Engine/Source/Runtime/NavigationSystem/Private/NavMesh/RecastNavMeshDataChunk.cpp
@@ -551,16 +551,30 @@ void URecastNavMeshDataChunk::ReleaseTiles()
 	Tiles.Reset();
 }
 
+// Deprecated
 void URecastNavMeshDataChunk::GetTiles(const FPImplRecastNavMesh* NavMeshImpl, const TArray<int32>& TileIndices, const EGatherTilesCopyMode CopyMode, const bool bMarkAsAttached /*= true*/)
 {
-	Tiles.Empty(TileIndices.Num());
+	if (NavMeshImpl)
+	{
+		TArray<uint32> TileUnsignedIndices;
+		TileUnsignedIndices.Append(TileIndices);
+
+		TArray<FNavTileRef> TileRefs;
+		FNavTileRef::DeprecatedMakeTileRefsFromTileIds(NavMeshImpl, TileUnsignedIndices, TileRefs);
+		GetTiles(NavMeshImpl, TileRefs, CopyMode, bMarkAsAttached);
+	}
+}
+
+void URecastNavMeshDataChunk::GetTiles(const FPImplRecastNavMesh* NavMeshImpl, const TArray<FNavTileRef>& TileRefs, const EGatherTilesCopyMode CopyMode, const bool bMarkAsAttached /*= true*/)
+{
+	Tiles.Empty(TileRefs.Num());
 
 #if WITH_RECAST
 	const dtNavMesh* NavMesh = NavMeshImpl->DetourNavMesh;
 	
-	for (int32 TileIdx : TileIndices)
+	for (const FNavTileRef TileRef : TileRefs)
 	{
-		const dtMeshTile* Tile = NavMesh->getTile(TileIdx);
+		const dtMeshTile* Tile = NavMesh->getTileByRef(static_cast<dtTileRef>(TileRef));
 		if (Tile && Tile->header)
 		{
 			// Make our own copy of tile data
@@ -597,15 +611,26 @@ void URecastNavMeshDataChunk::GetTiles(const FPImplRecastNavMesh* NavMeshImpl, c
 #endif // WITH_RECAST
 }
 
+// Deprecated
 void URecastNavMeshDataChunk::GetTilesBounds(const FPImplRecastNavMesh& NavMeshImpl, const TArray<int32>& TileIndices, FBox& OutBounds) const
+{
+	TArray<uint32> TileUnsignedIndices;
+	TileUnsignedIndices.Append(TileIndices);
+
+	TArray<FNavTileRef> TileRefs;
+	FNavTileRef::DeprecatedMakeTileRefsFromTileIds(&NavMeshImpl, TileUnsignedIndices, TileRefs);
+	GetTilesBounds(NavMeshImpl, TileRefs, OutBounds);
+}
+
+void URecastNavMeshDataChunk::GetTilesBounds(const FPImplRecastNavMesh& NavMeshImpl, const TArray<FNavTileRef>& TileRefs, FBox& OutBounds) const
 {
 	OutBounds.Init();
 #if WITH_RECAST
 	const dtNavMesh* NavMesh = NavMeshImpl.DetourNavMesh;
 
-	for (const int32 TileIdx : TileIndices)
+	for (const FNavTileRef TileRef : TileRefs)
 	{
-		const dtMeshTile* Tile = NavMesh->getTile(TileIdx);
+		const dtMeshTile* Tile = NavMesh->getTileByRef(static_cast<dtTileRef>(TileRef));
 		if (Tile && Tile->header)
 		{
 			OutBounds += Recast2UnrealBox(Tile->header->bmin, Tile->header->bmax);
diff --git a/Engine/Source/Runtime/NavigationSystem/Public/NavMesh/NavMeshRenderingComponent.h b/Engine/Source/Runtime/NavigationSystem/Public/NavMesh/NavMeshRenderingComponent.h
--- a/Engine/Source/Runtime/NavigationSystem/Public/NavMesh/NavMeshRenderingComponent.h
+++ b/Engine/Source/Runtime/NavigationSystem/Public/NavMesh/NavMeshRenderingComponent.h
@@ -24,6 +24,8 @@ class FPrimitiveDrawInterface;
 class UCanvas;
 class UNavMeshRenderingComponent;
 
+struct FNavTileRef;
+
 enum class ENavMeshDetailFlags : uint8
 {
 	TriangleEdges,
@@ -104,7 +106,9 @@ struct FNavMeshSceneProxyData : public TSharedFromThis<FNavMeshSceneProxyData, E
 	NAVIGATIONSYSTEM_API uint32 GetAllocatedSize() const;
 
 #if WITH_RECAST
+	UE_DEPRECATED(5.5, "Use the version of this function that takes an array of FNavTileRefs instead")
 	NAVIGATIONSYSTEM_API void GatherData(const ARecastNavMesh* NavMesh, int32 InNavDetailFlags, const TArray<int32>& TileSet);
+	NAVIGATIONSYSTEM_API void GatherData(const ARecastNavMesh* NavMesh, int32 InNavDetailFlags, const TArray<FNavTileRef>& TileSet);
 
 #if RECAST_INTERNAL_DEBUG_DATA
 	NAVIGATIONSYSTEM_API void AddMeshForInternalData(const struct FRecastInternalDebugData& InInternalData);
diff --git a/Engine/Source/Runtime/NavigationSystem/Public/NavMesh/PImplRecastNavMesh.h b/Engine/Source/Runtime/NavigationSystem/Public/NavMesh/PImplRecastNavMesh.h
--- a/Engine/Source/Runtime/NavigationSystem/Public/NavMesh/PImplRecastNavMesh.h
+++ b/Engine/Source/Runtime/NavigationSystem/Public/NavMesh/PImplRecastNavMesh.h
@@ -60,26 +60,41 @@ public:
 	 * @params TileIndex Used to collect geometry for a specific tile, INDEX_NONE will gather all tiles.
 	 * @return True if done collecting.
 	 */
+	UE_DEPRECATED(5.5, "Use the version of this function that takes a FNavTileRef instead.")
 	NAVIGATIONSYSTEM_API bool GetDebugGeometryForTile(FRecastDebugGeometry& OutGeometry, int32 TileIndex) const;
+
+	/* Gather debug geometry.
+	 * @params OutGeometry Output geometry.
+	 * @params TileRef Used to collect geometry for a specific tile, an invalid FNavTileRef will gather all tiles.
+	 * @return True if done collecting.
+	 */
+	NAVIGATIONSYSTEM_API bool GetDebugGeometryForTile(FRecastDebugGeometry& OutGeometry, FNavTileRef TileRef) const;
 	
 	/** Returns bounding box for the whole navmesh. */
 	NAVIGATIONSYSTEM_API FBox GetNavMeshBounds() const;
 
 	/** Returns bounding box for a given navmesh tile. */
+	UE_DEPRECATED(5.5, "Use the version of this function on ARecastNavMesh that takes a FNavTileRef instead.")
 	NAVIGATIONSYSTEM_API FBox GetNavMeshTileBounds(int32 TileIndex) const;
 
 	/** Retrieves XY and layer coordinates of tile specified by index */
+	UE_DEPRECATED(5.5, "Use the version of this function on ARecastNavMesh that takes a FNavTileRef instead.")
 	NAVIGATIONSYSTEM_API bool GetNavMeshTileXY(int32 TileIndex, int32& OutX, int32& OutY, int32& OutLayer) const;
 
 	/** Retrieves XY coordinates of tile specified by position */
 	NAVIGATIONSYSTEM_API bool GetNavMeshTileXY(const FVector& Point, int32& OutX, int32& OutY) const;
 
 	/** Retrieves all tile indices at matching XY coordinates */
+	UE_DEPRECATED(5.5, "Use the version of this function on ARecastNavMesh that takes an array of FNavTileRefs instead.")
 	NAVIGATIONSYSTEM_API void GetNavMeshTilesAt(int32 TileX, int32 TileY, TArray<int32>& Indices) const;
 
 	/** Retrieves list of tiles that intersect specified bounds */
+	UE_DEPRECATED(5.5, "Use the version of the function that takes an array of FNavTileRefs instead")
 	NAVIGATIONSYSTEM_API void GetNavMeshTilesIn(const TArray<FBox>& InclusionBounds, TArray<int32>& Indices) const;
 
+	/** Retrieves list of tiles that intersect specified bounds */
+	NAVIGATIONSYSTEM_API void GetNavMeshTilesIn(const TArray<FBox>& InclusionBounds, TArray<FNavTileRef>& OutRefs) const;
+
 	/** Retrieves number of tiles in this navmesh */
 	FORCEINLINE int32 GetNavMeshTilesCount() const { return DetourNavMesh ? DetourNavMesh->getMaxTiles() : 0; }
 
@@ -164,6 +179,7 @@ public:
 	NAVIGATIONSYSTEM_API bool FilterPolys(TArray<NavNodeRef>& PolyRefs, const FRecastQueryFilter* Filter, const UObject* Owner) const;
 
 	/** Get all polys from tile */
+	UE_DEPRECATED(5.5, "Use the version of this function in ARecastNavMesh that takes a FNavTileRef instead.")
 	NAVIGATIONSYSTEM_API bool GetPolysInTile(int32 TileIndex, TArray<FNavPoly>& Polys) const;
 
 	UE_DEPRECATED(5.3, "Please use the version of this function that takes a FNavLinkId. This function has no effect.")
diff --git a/Engine/Source/Runtime/NavigationSystem/Public/NavMesh/RecastNavMesh.h b/Engine/Source/Runtime/NavigationSystem/Public/NavMesh/RecastNavMesh.h
--- a/Engine/Source/Runtime/NavigationSystem/Public/NavMesh/RecastNavMesh.h
+++ b/Engine/Source/Runtime/NavigationSystem/Public/NavMesh/RecastNavMesh.h
@@ -284,6 +284,8 @@ struct FNavTileRef
 
 	bool IsValid() const { return TileRef != (uint64)FNavTileRef(); }
 
+	FORCEINLINE friend uint32 GetTypeHash(const FNavTileRef& NavTileRef) { return GetTypeHash(NavTileRef.TileRef); }
+
 	/** Those 2 functions are used for backward compatibility of the following deprecated functions in FRecastNavMeshGenerator and ARecastNavMesh:
 	*	  RemoveTileLayers
 	*     AddGeneratedTilesTimeSliced
@@ -1179,26 +1181,45 @@ public:
 	NAVIGATIONSYSTEM_API FBox GetNavMeshBounds() const;
 
 	/** Returns bounding box for a given navmesh tile. */
+	UE_DEPRECATED(5.5, "Use the version of this function that takes a FNavTileRef instead.")
 	NAVIGATIONSYSTEM_API FBox GetNavMeshTileBounds(int32 TileIndex) const;
 
+	/** Returns bounding box for a given navmesh tile. */
+	NAVIGATIONSYSTEM_API FBox GetNavMeshTileBounds(FNavTileRef TileRef) const;
+
 	/** Retrieves XY coordinates of tile specified by index */
+	UE_DEPRECATED(5.5, "Use the version of this function that takes a FNavTileRef instead.")
 	NAVIGATIONSYSTEM_API bool GetNavMeshTileXY(int32 TileIndex, int32& OutX, int32& OutY, int32& Layer) const;
 
+	/** Retrieves XY coordinates of tile */
+	NAVIGATIONSYSTEM_API bool GetNavMeshTileXY(FNavTileRef TileRef, int32& OutX, int32& OutY, int32& Layer) const;
+
 	/** Retrieves XY coordinates of tile specified by position */
 	NAVIGATIONSYSTEM_API bool GetNavMeshTileXY(const FVector& Point, int32& OutX, int32& OutY) const;
 
 	/** Retrieves the tile resolution */
+	UE_DEPRECATED(5.5, "Use the version of this function that takes a FNavTileRef instead.")
 	NAVIGATIONSYSTEM_API bool GetNavmeshTileResolution(int32 TileIndex, ENavigationDataResolution& OutResolution) const;
 
+	/** Retrieves the tile resolution */
+	NAVIGATIONSYSTEM_API bool GetNavmeshTileResolution(FNavTileRef TileRef, ENavigationDataResolution& OutResolution) const;
+
 	/** Checks the supplied Points tile indicies can fit in the range of an int32 */
 	NAVIGATIONSYSTEM_API bool CheckTileIndicesInValidRange(const FVector& Point, bool& bOutInRange) const;
 
 	/** Retrieves all tile indices at matching XY coordinates */
+	UE_DEPRECATED(5.5, "Use the version of this function that takes an array of FNavTileRefs instead.")
 	NAVIGATIONSYSTEM_API void GetNavMeshTilesAt(int32 TileX, int32 TileY, TArray<int32>& Indices) const;
 
+	/** Retrieves all tiles at matching XY coordinates */
+	NAVIGATIONSYSTEM_API void GetNavMeshTilesAt(int32 TileX, int32 TileY, TArray<FNavTileRef>& OutRefs) const;
+
 	/** Retrieves number of tiles in this navmesh */
 	NAVIGATIONSYSTEM_API int32 GetNavMeshTilesCount() const;
 
+	/** Retrieves all tiles in this navmesh */
+	NAVIGATIONSYSTEM_API void GetAllNavMeshTiles(TArray<FNavTileRef>& OutRefs) const;
+
 	/** Removes compressed tile data at given tile coord */
 	NAVIGATIONSYSTEM_API void RemoveTileCacheLayers(int32 TileX, int32 TileY);
 	
@@ -1265,8 +1286,16 @@ public:
 	 * @params TileIndex Used to collect geometry for a specific tile, INDEX_NONE will gather all tiles
 	 * @return True if done collecting.
 	 */
+	UE_DEPRECATED(5.5, "Use the version of the function that takes a FNavTileRef instead.")
 	NAVIGATIONSYSTEM_API bool GetDebugGeometryForTile(FRecastDebugGeometry& OutGeometry, int32 TileIndex) const;
 
+	/* Gather debug geometry.
+	 * @params OutGeometry Output geometry.
+	 * @params TileRef Used to collect geometry for a specific tile, an invalid FNavTileRef will gather all tiles
+	 * @return True if done collecting.
+	 */
+	NAVIGATIONSYSTEM_API bool GetDebugGeometryForTile(FRecastDebugGeometry& OutGeometry, FNavTileRef TileRef) const;
+
 	// @todo docuement
 	NAVIGATIONSYSTEM_API void DrawDebugPathCorridor(NavNodeRef const* PathPolys, int32 NumPathPolys, bool bPersistent=true) const;
 
@@ -1430,8 +1459,12 @@ public:
 	NAVIGATIONSYSTEM_API bool FilterPolys(TArray<NavNodeRef>& PolyRefs, const FRecastQueryFilter* Filter, const UObject* Querier = NULL) const;
 
 	/** Get all polys from tile */
+	UE_DEPRECATED(5.5, "Use the version of this function that takes a FNavTileRef instead.")
 	NAVIGATIONSYSTEM_API bool GetPolysInTile(int32 TileIndex, TArray<FNavPoly>& Polys) const;
 
+	/** Get all polys from tile */
+	NAVIGATIONSYSTEM_API bool GetPolysInTile(FNavTileRef TileRef, TArray<FNavPoly>& Polys) const;
+
 	/** Get up to 256 polys that overlap the specified box */
 	NAVIGATIONSYSTEM_API bool GetPolysInBox(const FBox& Box, TArray<FNavPoly>& Polys, FSharedConstNavQueryFilter Filter = nullptr, const UObject* Owner = nullptr) const;
 
diff --git a/Engine/Source/Runtime/NavigationSystem/Public/NavMesh/RecastNavMeshDataChunk.h b/Engine/Source/Runtime/NavigationSystem/Public/NavMesh/RecastNavMeshDataChunk.h
--- a/Engine/Source/Runtime/NavigationSystem/Public/NavMesh/RecastNavMeshDataChunk.h
+++ b/Engine/Source/Runtime/NavigationSystem/Public/NavMesh/RecastNavMeshDataChunk.h
@@ -116,8 +116,12 @@ class URecastNavMeshDataChunk : public UNavigationDataChunk
 	const TArray<FRecastTileData>& GetTiles() const { return Tiles; }
 
 	/** Returns the AABB for the given tiles. */
+	UE_DEPRECATED(5.5, "Use the version of this function that takes an array of FNavTileRefs instead.")
 	NAVIGATIONSYSTEM_API void GetTilesBounds(const FPImplRecastNavMesh& NavMeshImpl, const TArray<int32>& TileIndices, FBox& OutBounds) const;
 
+	/** Returns the AABB for the given tiles. */
+	NAVIGATIONSYSTEM_API void GetTilesBounds(const FPImplRecastNavMesh& NavMeshImpl, const TArray<FNavTileRef>& TileRefs, FBox& OutBounds) const;
+
 	/** Mutable accessor to the list of tiles in the data chunk. */
 	TArray<FRecastTileData>& GetMutableTiles() { return Tiles; }
 
@@ -125,8 +129,12 @@ class URecastNavMeshDataChunk : public UNavigationDataChunk
 	NAVIGATIONSYSTEM_API void ReleaseTiles();
 
 	/** Collect tiles with data and/or cache data from the provided TileIndices. */
+	UE_DEPRECATED(5.5, "Use the version of this function that takes an array of FNavTileRefs instead.")
 	NAVIGATIONSYSTEM_API void GetTiles(const FPImplRecastNavMesh* NavMeshImpl, const TArray<int32>& TileIndices, const EGatherTilesCopyMode CopyMode, const bool bMarkAsAttached = true);
 
+	/** Collect tiles with data and/or cache data from the provided tile references. */
+	NAVIGATIONSYSTEM_API void GetTiles(const FPImplRecastNavMesh* NavMeshImpl, const TArray<FNavTileRef>& TileRefs, const EGatherTilesCopyMode CopyMode, const bool bMarkAsAttached = true);
+
 private:
 #if WITH_RECAST
 	NAVIGATIONSYSTEM_API void SerializeRecastData(FArchive& Ar, int32 NavMeshVersion);
