Add way to test the NavigationData's Raycast method in NavigationTestingActor

diff --git a/Engine/Source/Runtime/NavigationSystem/Private/NavMesh/NavTestRenderingComponent.cpp b/Engine/Source/Runtime/NavigationSystem/Private/NavMesh/NavTestRenderingComponent.cpp
--- a/Engine/Source/Runtime/NavigationSystem/Private/NavMesh/NavTestRenderingComponent.cpp
+++ b/Engine/Source/Runtime/NavigationSystem/Private/NavMesh/NavTestRenderingComponent.cpp
@@ -122,6 +122,17 @@ void FNavTestSceneProxy::GetDynamicMeshElements(const TArray<const FSceneView*>&
 						NavTestActor->RadiusUsedToValidateNavData, 0.5f*Distance + NavTestActor->RadiusUsedToValidateNavData,  
 						NumSides, SDPG_World);
 				}
+
+				if (NavTestActor->bDrawRaycastToQueryTargetActor && NavTestActor->QueryTargetActor)
+				{
+					const FVector QueryTargetLocation = NavTestActor->QueryTargetActor->GetActorLocation();
+					PDI->DrawLine(ActorLocation, QueryTargetLocation, NavTestActor->bRaycastToQueryTargetActorResult ? FColor::Red : FColor::Green, SDPG_World, 2.5);
+
+					if (FNavigationSystem::IsValidLocation(NavTestActor->RaycastHitLocation))
+					{
+						DrawWireSphere(PDI, NavTestActor->RaycastHitLocation, FColor::Red, 25.f, 6, SDPG_World);
+					}
+				}
 			}
 
 			// draw path
diff --git a/Engine/Source/Runtime/NavigationSystem/Private/NavigationTestingActor.cpp b/Engine/Source/Runtime/NavigationSystem/Private/NavigationTestingActor.cpp
--- a/Engine/Source/Runtime/NavigationSystem/Private/NavigationTestingActor.cpp
+++ b/Engine/Source/Runtime/NavigationSystem/Private/NavigationTestingActor.cpp
@@ -55,8 +55,10 @@ ANavigationTestingActor::ANavigationTestingActor(const FObjectInitializer& Objec
 	bGatherDetailedInfo = true;
 	bDrawDistanceToWall = false;
 	ClosestWallLocation = FNavigationSystem::InvalidLocation;
+	RaycastHitLocation = FNavigationSystem::InvalidLocation;
 	bNavDataIsReadyInRadius = false;
 	bNavDataIsReadyToQueryTargetActor = false;
+	bRaycastToQueryTargetActorResult = false;
 	OffsetFromCornersDistance = 0.f;
 
 	QueryingExtent = FVector(DEFAULT_NAV_QUERY_EXTENT_HORIZONTAL, DEFAULT_NAV_QUERY_EXTENT_HORIZONTAL, DEFAULT_NAV_QUERY_EXTENT_VERTICAL);
@@ -140,6 +142,7 @@ void ANavigationTestingActor::PostEditChangeProperty(FPropertyChangedEvent& Prop
 	static const FName NAME_QueryTargetActor = GET_MEMBER_NAME_CHECKED(ANavigationTestingActor, QueryTargetActor);
 	static const FName NAME_IsSearchStart = GET_MEMBER_NAME_CHECKED(ANavigationTestingActor, bSearchStart);
 	static const FName NAME_InvokerComponent = GET_MEMBER_NAME_CHECKED(ANavigationTestingActor, InvokerComponent);
+	static const FName NAME_FilterClass = GET_MEMBER_NAME_CHECKED(ANavigationTestingActor, FilterClass);
 
 	Super::PostEditChangeProperty(PropertyChangedEvent);
 
@@ -189,17 +192,12 @@ void ANavigationTestingActor::PostEditChangeProperty(FPropertyChangedEvent& Prop
 				}
 			}
 
-			if (bDrawDistanceToWall)
-			{
-				ClosestWallLocation = FindClosestWallLocation();
-			}
-			if (bDrawIfNavDataIsReadyInRadius)
-			{
-				bNavDataIsReadyInRadius = CheckIfNavDataIsReadyInRadius();
-			}
-			if (bDrawIfNavDataIsReadyToQueryTargetActor)
+			UpdateLocalQueries();
+			UpdateTargetActorQueries();
+
+			if (ChangedPropName == NAME_FilterClass)
 			{
-				bNavDataIsReadyToQueryTargetActor = CheckIfNavDataIsReadyToActor(QueryTargetActor);
+				UpdatePathfinding();
 			}
 		}
 		else if (ChangedCategory == TEXT("Pathfinding"))
@@ -279,20 +277,8 @@ void ANavigationTestingActor::PostEditMove(bool bFinished)
 			UpdatePathfinding();
 		}
 
-		if (bDrawDistanceToWall)
-		{
-			ClosestWallLocation = FindClosestWallLocation();
-		}
-
-		if (bDrawIfNavDataIsReadyInRadius)
-		{
-			bNavDataIsReadyInRadius = CheckIfNavDataIsReadyInRadius();
-		}
-
-		if (bDrawIfNavDataIsReadyToQueryTargetActor)
-		{
-			bNavDataIsReadyToQueryTargetActor = CheckIfNavDataIsReadyToActor(QueryTargetActor);
-		}
+		UpdateLocalQueries();
+		UpdateTargetActorQueries();
 	}
 }
 
@@ -422,6 +408,30 @@ void ANavigationTestingActor::UpdatePathfinding()
 	}
 }
 
+void ANavigationTestingActor::UpdateLocalQueries()
+{
+	if (bDrawDistanceToWall)
+	{
+		ClosestWallLocation = FindClosestWallLocation();
+	}
+	if (bDrawIfNavDataIsReadyInRadius)
+	{
+		bNavDataIsReadyInRadius = CheckIfNavDataIsReadyInRadius();
+	}
+}
+
+void ANavigationTestingActor::UpdateTargetActorQueries()
+{
+	if (bDrawIfNavDataIsReadyToQueryTargetActor)
+	{
+		bNavDataIsReadyToQueryTargetActor = CheckIfNavDataIsReadyToActor(QueryTargetActor);
+	}
+	if (bDrawRaycastToQueryTargetActor)
+	{
+		bRaycastToQueryTargetActorResult = CheckRaycastToActor(QueryTargetActor, RaycastHitLocation);
+	}
+}
+
 FVector ANavigationTestingActor::FindClosestWallLocation() const
 {
 #if WITH_EDITORONLY_DATA
@@ -486,12 +496,29 @@ bool ANavigationTestingActor::CheckIfNavDataIsReadyToActor(const AActor* TargetA
 	return false;
 }
 
-void ANavigationTestingActor::OnQueryTargetActorTransformUpdated(USceneComponent* InRootComponent, EUpdateTransformFlags UpdateTransformFlags, ETeleportType Teleport)
+bool ANavigationTestingActor::CheckRaycastToActor(const AActor* TargetActor, FVector& OutHitLocation)
 {
-	if (bDrawIfNavDataIsReadyToQueryTargetActor)
+	OutHitLocation = FNavigationSystem::InvalidLocation;
+#if WITH_EDITORONLY_DATA
+	if (EdRenderComp)
 	{
-		bNavDataIsReadyToQueryTargetActor = CheckIfNavDataIsReadyToActor(QueryTargetActor);
+		EdRenderComp->MarkRenderStateDirty();
 	}
+#endif // WITH_EDITORONLY_DATA
+
+	UpdateNavData();
+	if (MyNavData && TargetActor)
+	{
+		FSharedConstNavQueryFilter Filter = UNavigationQueryFilter::GetQueryFilter(*MyNavData, this, FilterClass);
+		return MyNavData->Raycast(GetActorLocation(), TargetActor->GetActorLocation(), OutHitLocation, Filter, this);
+	}
+
+	return false;
+}
+
+void ANavigationTestingActor::OnQueryTargetActorTransformUpdated(USceneComponent* InRootComponent, EUpdateTransformFlags UpdateTransformFlags, ETeleportType Teleport)
+{
+	UpdateTargetActorQueries();
 }
 
 void ANavigationTestingActor::SearchPathTo(ANavigationTestingActor* Goal)
diff --git a/Engine/Source/Runtime/NavigationSystem/Public/NavigationTestingActor.h b/Engine/Source/Runtime/NavigationSystem/Public/NavigationTestingActor.h
--- a/Engine/Source/Runtime/NavigationSystem/Public/NavigationTestingActor.h
+++ b/Engine/Source/Runtime/NavigationSystem/Public/NavigationTestingActor.h
@@ -124,6 +124,10 @@ public:
 	UPROPERTY(EditAnywhere, Category=Query)
 	uint32 bDrawIfNavDataIsReadyToQueryTargetActor : 1;
 
+	/** If set, a line is drawn to indicate to result of a raycast on the navigation data between the current actor and the QueryTargetActor location. */
+	UPROPERTY(EditAnywhere, Category=Query)
+	uint32 bDrawRaycastToQueryTargetActor : 1;
+
 	/** Actor to use as a target for navigation data queries */
 	UPROPERTY(EditAnywhere, Category=Query)
 	TObjectPtr<AActor> QueryTargetActor;
@@ -177,8 +181,8 @@ public:
 	UPROPERTY(EditAnywhere, Category=Pathfinding)
 	TObjectPtr<ANavigationTestingActor> OtherActor;
 
-	/** "None" will result in default filter being used */
-	UPROPERTY(EditAnywhere, Category=Pathfinding)
+	/** "None" will result in default filter being used. This filter is used by the PathFind and Raycast queries. */
+	UPROPERTY(EditAnywhere, Category=Query)
 	TSubclassOf<class UNavigationQueryFilter> FilterClass;
 
 	/** Show debug steps up to this index. Use -1 to disable. */
@@ -189,9 +193,11 @@ public:
 	float OffsetFromCornersDistance;
 
 	FVector ClosestWallLocation;
-	
+	FVector RaycastHitLocation;
+
 	bool bNavDataIsReadyInRadius;
 	bool bNavDataIsReadyToQueryTargetActor;
+	bool bRaycastToQueryTargetActorResult;
 
 #if WITH_RECAST && WITH_EDITORONLY_DATA
 	/** detail data gathered from each step of regular A* algorithm */
@@ -251,8 +257,12 @@ public:
 #endif
 
 protected:
+	void UpdateLocalQueries();
+	void UpdateTargetActorQueries();
+
 	NAVIGATIONSYSTEM_API FVector FindClosestWallLocation() const;
 	bool CheckIfNavDataIsReadyInRadius();
 	bool CheckIfNavDataIsReadyToActor(const AActor* TargetActor);
+	bool CheckRaycastToActor(const AActor* TargetActor, FVector& OutHitLocation);
 	void OnQueryTargetActorTransformUpdated(USceneComponent* InRootComponent, EUpdateTransformFlags UpdateTransformFlags, ETeleportType Teleport);
 };
