Adding a public function to ARecastNavMesh to allow users to directly grab all exterior navmesh edges from a provided tile. Includes renaming/deprecating the PImpl version of this function.

diff --git a/Engine/Source/Runtime/NavigationSystem/Private/NavMesh/PImplRecastNavMesh.cpp b/Engine/Source/Runtime/NavigationSystem/Private/NavMesh/PImplRecastNavMesh.cpp
--- a/Engine/Source/Runtime/NavigationSystem/Private/NavMesh/PImplRecastNavMesh.cpp
+++ b/Engine/Source/Runtime/NavigationSystem/Private/NavMesh/PImplRecastNavMesh.cpp
@@ -2605,16 +2605,18 @@ static FORCEINLINE FVector::FReal PointDistToSegment2DSquared(const FVector::FRe
 	return dx*dx + dz*dz;
 }
 
-/** 
- * Traverses given tile's edges and detects the ones that are either poly (i.e. not triangle, but whole navmesh polygon) 
- * or navmesh edge. Returns a pair of verts for each edge found.
- */
+// Deprecated
 void FPImplRecastNavMesh::GetDebugPolyEdges(const dtMeshTile& Tile, bool bInternalEdges, bool bNavMeshEdges, TArray<FVector>& InternalEdgeVerts, TArray<FVector>& NavMeshEdgeVerts) const
+{
+	GetTilePolyEdges(Tile, bInternalEdges, bNavMeshEdges, InternalEdgeVerts, NavMeshEdgeVerts);
+}
+
+void FPImplRecastNavMesh::GetTilePolyEdges(const dtMeshTile& Tile, bool bGatherInteriorPolyEdges, bool bGatherExteriorNavMeshEdges, TArray<FVector>& OutInteriorPolyEdgeVerts, TArray<FVector>& OutExteriorNavMeshEdgeVerts) const
 {
 	static const FVector::FReal thr = FMath::Square(0.01f);
 
-	ensure(bInternalEdges || bNavMeshEdges);
-	const bool bExportAllEdges = bInternalEdges && !bNavMeshEdges;
+	ensure(bGatherInteriorPolyEdges || bGatherExteriorNavMeshEdges);
+	const bool bExportAllEdges = bGatherInteriorPolyEdges && !bGatherExteriorNavMeshEdges;
 	
 	for (int i = 0; i < Tile.header->polyCount; ++i)
 	{
@@ -2662,8 +2664,8 @@ void FPImplRecastNavMesh::GetDebugPolyEdges(const dtMeshTile& Tile, bool bIntern
 				}
 			}
 
-			TArray<FVector>* EdgeVerts = bInternalEdges && bIsConnected ? &InternalEdgeVerts 
-				: (bNavMeshEdges && bIsExternal && !bIsConnected ? &NavMeshEdgeVerts : NULL);
+			TArray<FVector>* EdgeVerts = bGatherInteriorPolyEdges && bIsConnected ? &OutInteriorPolyEdgeVerts 
+				: (bGatherExteriorNavMeshEdges && bIsExternal && !bIsConnected ? &OutExteriorNavMeshEdgeVerts : NULL);
 			if (EdgeVerts == NULL)
 			{
 				continue;
@@ -2672,7 +2674,7 @@ void FPImplRecastNavMesh::GetDebugPolyEdges(const dtMeshTile& Tile, bool bIntern
 			const FVector::FReal* V0 = &Tile.verts[Poly->verts[j] * 3];
 			const FVector::FReal* V1 = &Tile.verts[Poly->verts[(j + 1) % nj] * 3];
 
-			// Draw detail mesh edges which align with the actual poly edge.
+			// Gather detail mesh edges which align with the actual poly edge.
 			// This is really slow.
 			for (int32 k = 0; k < pd->triCount; ++k)
 			{
@@ -3099,7 +3101,7 @@ int32 FPImplRecastNavMesh::GetTilesDebugGeometry(const FRecastNavMeshGenerator*
 	// Get tile edges and navmesh edges
 	if (OutGeometry.bGatherPolyEdges || OutGeometry.bGatherNavMeshEdges)
 	{
-		GetDebugPolyEdges(Tile, !!OutGeometry.bGatherPolyEdges, !!OutGeometry.bGatherNavMeshEdges
+		GetTilePolyEdges(Tile, !!OutGeometry.bGatherPolyEdges, !!OutGeometry.bGatherNavMeshEdges
 			, OutGeometry.PolyEdges, OutGeometry.NavMeshEdges);
 	}
 
diff --git a/Engine/Source/Runtime/NavigationSystem/Private/NavMesh/RecastNavMesh.cpp b/Engine/Source/Runtime/NavigationSystem/Private/NavMesh/RecastNavMesh.cpp
--- a/Engine/Source/Runtime/NavigationSystem/Private/NavMesh/RecastNavMesh.cpp
+++ b/Engine/Source/Runtime/NavigationSystem/Private/NavMesh/RecastNavMesh.cpp
@@ -2176,6 +2176,40 @@ bool ARecastNavMesh::FindEdges(const NavNodeRef CenterNodeRef, const FVector Cen
 	return false;
 }
 
+bool ARecastNavMesh::GetEdgesInTile(FNavTileRef TileRef, TArray<FNavigationWallEdge>& OutEdges) const
+{
+	if (HasValidNavmesh())
+	{
+		const dtNavMesh* const DetourNavMesh = RecastNavMeshImpl->GetRecastMesh();
+		if (const dtMeshTile* const Tile = DetourNavMesh->getTileByRef(static_cast<dtTileRef>(TileRef)))
+		{
+			const bool bGatherInternalEdges = false;
+			TArray<FVector> UNUSED_InternalEdgeVerts;
+
+			const bool bGatherExternalEdges = true;
+			TArray<FVector> ExternalEdgeVerts;
+			// Guessing that each detail tri in the tile has one external edge
+			ExternalEdgeVerts.Reserve(Tile->header->detailTriCount);
+
+			RecastNavMeshImpl->GetTilePolyEdges(*Tile, bGatherInternalEdges, bGatherExternalEdges, UNUSED_InternalEdgeVerts, ExternalEdgeVerts);
+
+			const int32 ExternalEdgeVertCount = ExternalEdgeVerts.Num();
+			if (ExternalEdgeVertCount > 0)
+			{
+				OutEdges.Reset(ExternalEdgeVertCount / 2);
+				for (int32 ExternalEdgeVertIndex = 0; ExternalEdgeVertIndex < ExternalEdgeVertCount; ExternalEdgeVertIndex += 2)
+				{
+					FNavigationWallEdge& Edge = OutEdges.AddDefaulted_GetRef();
+					Edge.Start = ExternalEdgeVerts[ExternalEdgeVertIndex];
+					Edge.End = ExternalEdgeVerts[ExternalEdgeVertIndex + 1];
+				}
+			}
+		}
+	}
+
+	return false;
+}
+
 bool ARecastNavMesh::ProjectPointMulti(const FVector& Point, TArray<FNavLocation>& OutLocations, const FVector& Extent,
 	FVector::FReal MinZ, FVector::FReal MaxZ, FSharedConstNavQueryFilter Filter, const UObject* QueryOwner) const
 {
diff --git a/Engine/Source/Runtime/NavigationSystem/Public/NavMesh/PImplRecastNavMesh.h b/Engine/Source/Runtime/NavigationSystem/Public/NavMesh/PImplRecastNavMesh.h
--- a/Engine/Source/Runtime/NavigationSystem/Public/NavMesh/PImplRecastNavMesh.h
+++ b/Engine/Source/Runtime/NavigationSystem/Public/NavMesh/PImplRecastNavMesh.h
@@ -314,8 +314,21 @@ public:
 		FVector RecastStart, FVector RecastEnd,
 		dtQueryResult& PathResult) const;
 
+	UE_DEPRECATED(5.5, "Use GetTilePolyEdges instead.")
 	NAVIGATIONSYSTEM_API void GetDebugPolyEdges(const dtMeshTile& Tile, bool bInternalEdges, bool bNavMeshEdges, TArray<FVector>& InternalEdgeVerts, TArray<FVector>& NavMeshEdgeVerts) const;
 
+	/**
+	 * Traverses given tile's edges and detects the ones that are either internal poly (i.e. not triangle, but whole navmesh polygon)
+	 * or external navmesh edge. Returns a pair of verts for each edge found.
+	 * @param Tile The tile whose edges to traverse.
+	 * @param bGatherInteriorPolyEdges If true, populates OutPolyEdgeVerts with the tile's internal poly edges.
+	 * @param bGatherExternalNavMeshEdges If true, populates OutNavMeshEdgeVerts with the tile's external navmesh edges.
+	 * @param OutInteriorPolyEdgeVerts Output poly edge vertex array. Contains a pair of verts for each edge.
+	 * @param OutExteriorNavMeshEdgeVerts Output navmesh edge vertex array. Contains a pair of verts for each edge.
+	 * @note This is really slow.
+	 */
+	NAVIGATIONSYSTEM_API void GetTilePolyEdges(const dtMeshTile& Tile, bool bGatherInteriorPolyEdges, bool bGatherExteriorNavMeshEdges, TArray<FVector>& OutInteriorPolyEdgeVerts, TArray<FVector>& OutExteriorNavMeshEdgeVerts) const;
+
 	/** workhorse function finding portal edges between corridor polys */
 	NAVIGATIONSYSTEM_API void GetEdgesForPathCorridorImpl(const TArray<NavNodeRef>* PathCorridor, TArray<FNavigationPortalEdge>* PathCorridorEdges, const dtNavMeshQuery& NavQuery) const;
 
diff --git a/Engine/Source/Runtime/NavigationSystem/Public/NavMesh/RecastNavMesh.h b/Engine/Source/Runtime/NavigationSystem/Public/NavMesh/RecastNavMesh.h
--- a/Engine/Source/Runtime/NavigationSystem/Public/NavMesh/RecastNavMesh.h
+++ b/Engine/Source/Runtime/NavigationSystem/Public/NavMesh/RecastNavMesh.h
@@ -1474,6 +1474,9 @@ public:
 	/** Find up to 64 navmesh eges in up to 64 polys around the center */
 	NAVIGATIONSYSTEM_API bool FindEdges(const NavNodeRef CenterNodeRef, const FVector Center, const FVector::FReal Radius, const FSharedConstNavQueryFilter Filter, TArray<FNavigationWallEdge>& OutEdges) const;
 
+	/** Get all exterior nav mesh edges from tile */
+	NAVIGATIONSYSTEM_API bool GetEdgesInTile(FNavTileRef TileRef, TArray<FNavigationWallEdge>& OutEdges) const;
+
 	/** Get all polys from tile */
 	NAVIGATIONSYSTEM_API bool GetNavLinksInTile(const int32 TileIndex, TArray<FNavPoly>& Polys, const bool bIncludeLinksFromNeighborTiles) const;
 
