Wrapping recently-introduced FNavTileRef usage with WITH_RECAST to fix build errors when it is not defined

diff --git a/Engine/Source/Runtime/NavigationSystem/Private/NavMesh/RecastNavMeshDataChunk.cpp b/Engine/Source/Runtime/NavigationSystem/Private/NavMesh/RecastNavMeshDataChunk.cpp
--- a/Engine/Source/Runtime/NavigationSystem/Private/NavMesh/RecastNavMeshDataChunk.cpp
+++ b/Engine/Source/Runtime/NavigationSystem/Private/NavMesh/RecastNavMeshDataChunk.cpp
@@ -554,6 +554,9 @@ void URecastNavMeshDataChunk::ReleaseTiles()
 // Deprecated
 void URecastNavMeshDataChunk::GetTiles(const FPImplRecastNavMesh* NavMeshImpl, const TArray<int32>& TileIndices, const EGatherTilesCopyMode CopyMode, const bool bMarkAsAttached /*= true*/)
 {
+	Tiles.Empty(TileIndices.Num());
+
+#if WITH_RECAST
 	if (NavMeshImpl)
 	{
 		TArray<uint32> TileUnsignedIndices;
@@ -563,13 +566,14 @@ void URecastNavMeshDataChunk::GetTiles(const FPImplRecastNavMesh* NavMeshImpl, c
 		FNavTileRef::DeprecatedMakeTileRefsFromTileIds(NavMeshImpl, TileUnsignedIndices, TileRefs);
 		GetTiles(NavMeshImpl, TileRefs, CopyMode, bMarkAsAttached);
 	}
+#endif // WITH_RECAST
 }
 
+#if WITH_RECAST
 void URecastNavMeshDataChunk::GetTiles(const FPImplRecastNavMesh* NavMeshImpl, const TArray<FNavTileRef>& TileRefs, const EGatherTilesCopyMode CopyMode, const bool bMarkAsAttached /*= true*/)
 {
 	Tiles.Empty(TileRefs.Num());
 
-#if WITH_RECAST
 	const dtNavMesh* NavMesh = NavMeshImpl->DetourNavMesh;
 	
 	for (const FNavTileRef TileRef : TileRefs)
@@ -608,24 +612,28 @@ void URecastNavMeshDataChunk::GetTiles(const FPImplRecastNavMesh* NavMeshImpl, c
 			Tiles.Add(RecastTileData);
 		}
 	}
-#endif // WITH_RECAST
 }
+#endif // WITH_RECAST
 
 // Deprecated
 void URecastNavMeshDataChunk::GetTilesBounds(const FPImplRecastNavMesh& NavMeshImpl, const TArray<int32>& TileIndices, FBox& OutBounds) const
 {
+	OutBounds.Init();
+#if WITH_RECAST
 	TArray<uint32> TileUnsignedIndices;
 	TileUnsignedIndices.Append(TileIndices);
 
 	TArray<FNavTileRef> TileRefs;
 	FNavTileRef::DeprecatedMakeTileRefsFromTileIds(&NavMeshImpl, TileUnsignedIndices, TileRefs);
 	GetTilesBounds(NavMeshImpl, TileRefs, OutBounds);
+#endif // WITH_RECAST
 }
 
+#if WITH_RECAST
 void URecastNavMeshDataChunk::GetTilesBounds(const FPImplRecastNavMesh& NavMeshImpl, const TArray<FNavTileRef>& TileRefs, FBox& OutBounds) const
 {
 	OutBounds.Init();
-#if WITH_RECAST
+
 	const dtNavMesh* NavMesh = NavMeshImpl.DetourNavMesh;
 
 	for (const FNavTileRef TileRef : TileRefs)
@@ -636,5 +644,5 @@ void URecastNavMeshDataChunk::GetTilesBounds(const FPImplRecastNavMesh& NavMeshI
 			OutBounds += Recast2UnrealBox(Tile->header->bmin, Tile->header->bmax);
 		}
 	}
+}
 #endif // WITH_RECAST
-}
\ No newline at end of file
diff --git a/Engine/Source/Runtime/NavigationSystem/Public/NavMesh/NavMeshRenderingComponent.h b/Engine/Source/Runtime/NavigationSystem/Public/NavMesh/NavMeshRenderingComponent.h
--- a/Engine/Source/Runtime/NavigationSystem/Public/NavMesh/NavMeshRenderingComponent.h
+++ b/Engine/Source/Runtime/NavigationSystem/Public/NavMesh/NavMeshRenderingComponent.h
@@ -24,8 +24,6 @@ class FPrimitiveDrawInterface;
 class UCanvas;
 class UNavMeshRenderingComponent;
 
-struct FNavTileRef;
-
 enum class ENavMeshDetailFlags : uint8
 {
 	TriangleEdges,
@@ -108,7 +106,7 @@ struct FNavMeshSceneProxyData : public TSharedFromThis<FNavMeshSceneProxyData, E
 #if WITH_RECAST
 	UE_DEPRECATED(5.5, "Use the version of this function that takes an array of FNavTileRefs instead")
 	NAVIGATIONSYSTEM_API void GatherData(const ARecastNavMesh* NavMesh, int32 InNavDetailFlags, const TArray<int32>& TileSet);
-	NAVIGATIONSYSTEM_API void GatherData(const ARecastNavMesh* NavMesh, int32 InNavDetailFlags, const TArray<FNavTileRef>& TileSet);
+	NAVIGATIONSYSTEM_API void GatherData(const ARecastNavMesh* NavMesh, int32 InNavDetailFlags, const TArray<struct FNavTileRef>& TileSet);
 
 #if RECAST_INTERNAL_DEBUG_DATA
 	NAVIGATIONSYSTEM_API void AddMeshForInternalData(const struct FRecastInternalDebugData& InInternalData);
diff --git a/Engine/Source/Runtime/NavigationSystem/Public/NavMesh/RecastNavMeshDataChunk.h b/Engine/Source/Runtime/NavigationSystem/Public/NavMesh/RecastNavMeshDataChunk.h
--- a/Engine/Source/Runtime/NavigationSystem/Public/NavMesh/RecastNavMeshDataChunk.h
+++ b/Engine/Source/Runtime/NavigationSystem/Public/NavMesh/RecastNavMeshDataChunk.h
@@ -119,8 +119,10 @@ class URecastNavMeshDataChunk : public UNavigationDataChunk
 	UE_DEPRECATED(5.5, "Use the version of this function that takes an array of FNavTileRefs instead.")
 	NAVIGATIONSYSTEM_API void GetTilesBounds(const FPImplRecastNavMesh& NavMeshImpl, const TArray<int32>& TileIndices, FBox& OutBounds) const;
 
+#if WITH_RECAST
 	/** Returns the AABB for the given tiles. */
 	NAVIGATIONSYSTEM_API void GetTilesBounds(const FPImplRecastNavMesh& NavMeshImpl, const TArray<FNavTileRef>& TileRefs, FBox& OutBounds) const;
+#endif // WITH_RECAST
 
 	/** Mutable accessor to the list of tiles in the data chunk. */
 	TArray<FRecastTileData>& GetMutableTiles() { return Tiles; }
@@ -132,8 +134,10 @@ class URecastNavMeshDataChunk : public UNavigationDataChunk
 	UE_DEPRECATED(5.5, "Use the version of this function that takes an array of FNavTileRefs instead.")
 	NAVIGATIONSYSTEM_API void GetTiles(const FPImplRecastNavMesh* NavMeshImpl, const TArray<int32>& TileIndices, const EGatherTilesCopyMode CopyMode, const bool bMarkAsAttached = true);
 
+#if WITH_RECAST
 	/** Collect tiles with data and/or cache data from the provided tile references. */
 	NAVIGATIONSYSTEM_API void GetTiles(const FPImplRecastNavMesh* NavMeshImpl, const TArray<FNavTileRef>& TileRefs, const EGatherTilesCopyMode CopyMode, const bool bMarkAsAttached = true);
+#endif // WITH_RECAST
 
 private:
 #if WITH_RECAST
