Fix crash when generating navmesh tiles if DirtyLayers has not been initialized.

diff --git a/Engine/Source/Runtime/NavigationSystem/Private/NavMesh/RecastNavMeshGenerator.cpp b/Engine/Source/Runtime/NavigationSystem/Private/NavMesh/RecastNavMeshGenerator.cpp
--- a/Engine/Source/Runtime/NavigationSystem/Private/NavMesh/RecastNavMeshGenerator.cpp
+++ b/Engine/Source/Runtime/NavigationSystem/Private/NavMesh/RecastNavMeshGenerator.cpp
@@ -1769,6 +1769,7 @@ void FRecastTileGenerator::Setup(const FRecastNavMeshGenerator& ParentGenerator,
 	bRegenerateCompressedLayers = (bGeometryChanged || TileConfig.bGenerateLinks || CompressedLayers.Num() == 0);
 	
 	// Gather geometry for tile if it's inside navigable bounds
+	// DirtyLayers might not be initialized if InclusionBounds are empty
 	if (InclusionBounds.Num())
 	{
 		if (!bRegenerateCompressedLayers)
@@ -1805,7 +1806,6 @@ void FRecastTileGenerator::Setup(const FRecastNavMeshGenerator& ParentGenerator,
 		}
 	}
 	
-	//
 	UsedMemoryOnStartup = GetUsedMemCount() + sizeof(FRecastTileGenerator);
 }
 
@@ -4206,7 +4206,7 @@ ETimeSliceWorkResult FRecastTileGenerator::GenerateNavigationDataTimeSliced(FNav
 	}//fall through to next state
 	case EGenerateNavDataTimeSlicedState::GenerateLayers:
 	{
-		for (; GenNavDataLayerTimeSlicedIdx < CompressedLayers.Num(); GenNavDataLayerTimeSlicedIdx++)
+		for (; GenNavDataLayerTimeSlicedIdx < CompressedLayers.Num() && GenNavDataLayerTimeSlicedIdx < DirtyLayers.Num(); GenNavDataLayerTimeSlicedIdx++)
 		{
 			if (DirtyLayers[GenNavDataLayerTimeSlicedIdx] == false || !CompressedLayers[GenNavDataLayerTimeSlicedIdx].IsValid())
 			{
