Add a flag to indicate that a link has been generated by the navigation system and display generated links with thicker lines.

diff --git a/Engine/Source/Runtime/Engine/Classes/AI/Navigation/NavLinkDefinition.h b/Engine/Source/Runtime/Engine/Classes/AI/Navigation/NavLinkDefinition.h
--- a/Engine/Source/Runtime/Engine/Classes/AI/Navigation/NavLinkDefinition.h
+++ b/Engine/Source/Runtime/Engine/Classes/AI/Navigation/NavLinkDefinition.h
@@ -55,7 +55,7 @@ struct FNavigationLinkBase
 	UPROPERTY(EditAnywhere, Category=Default, meta=(ClampMin = "1.0", EditCondition="bUseSnapHeight"))
 	float SnapHeight;
 
-	/** restrict area only to specified agents */
+	/** restrict link only to specified agents */
 	UPROPERTY(EditAnywhere, Category=Default)
 	FNavAgentSelector SupportedAgents;
 
@@ -116,7 +116,7 @@ struct FNavigationLinkBase
 	UPROPERTY(EditAnywhere, Category=Default, meta=(InlineEditConditionToggle))
 	uint8 bUseSnapHeight : 1;
 
-	/** If set, link will try to snap to cheapest area in given radius */
+	/** If set, link will try to snap to the cheapest area in given radius */
 	UPROPERTY(EditAnywhere, Category = Default)
 	uint8 bSnapToCheapestArea : 1;
 	
@@ -152,6 +152,9 @@ struct FNavigationLinkBase
 	UPROPERTY()
 	uint8 bCustomFlag7 : 1;
 
+	/** Was this link generated by the navigation system */
+	uint8 bIsGenerated : 1;
+
 	ENGINE_API FNavigationLinkBase();
 
 	PRAGMA_DISABLE_DEPRECATION_WARNINGS
diff --git a/Engine/Source/Runtime/Engine/Private/AI/Navigation/NavigationModifier.cpp b/Engine/Source/Runtime/Engine/Private/AI/Navigation/NavigationModifier.cpp
--- a/Engine/Source/Runtime/Engine/Private/AI/Navigation/NavigationModifier.cpp
+++ b/Engine/Source/Runtime/Engine/Private/AI/Navigation/NavigationModifier.cpp
@@ -32,7 +32,8 @@ FNavigationLinkBase::FNavigationLinkBase()
 	: LeftProjectHeight(0.0f), MaxFallDownLength(1000.0f), SnapRadius(30.f), SnapHeight(50.0f),
 	  Direction(ENavLinkDirection::BothWays), bUseSnapHeight(false), bSnapToCheapestArea(true),
 	  bCustomFlag0(false), bCustomFlag1(false), bCustomFlag2(false), bCustomFlag3(false), bCustomFlag4(false),
-	  bCustomFlag5(false), bCustomFlag6(false), bCustomFlag7(false)
+	  bCustomFlag5(false), bCustomFlag6(false), bCustomFlag7(false),
+	  bIsGenerated(false)
 {
 	PRAGMA_DISABLE_DEPRECATION_WARNINGS
 	UserId = InvalidUserId;
diff --git a/Engine/Source/Runtime/NavigationSystem/Private/NavMesh/NavMeshRenderingComponent.cpp b/Engine/Source/Runtime/NavigationSystem/Private/NavMesh/NavMeshRenderingComponent.cpp
--- a/Engine/Source/Runtime/NavigationSystem/Private/NavMesh/NavMeshRenderingComponent.cpp
+++ b/Engine/Source/Runtime/NavigationSystem/Private/NavMesh/NavMeshRenderingComponent.cpp
@@ -44,6 +44,7 @@ static constexpr float TileResolution_LineThickness = 5.f;
 static constexpr float PolyEdges_LineThickness = 1.1f;
 static constexpr float NavMeshEdges_LineThickness = 4.f;
 static constexpr float LinkLines_LineThickness = 2.0f;
+static constexpr float GeneratedLinkLines_LineThickness = 4.0f;
 static constexpr float ClusterLinkLines_LineThickness = 2.0f;
 
 namespace FNavMeshRenderingHelpers
@@ -141,7 +142,7 @@ namespace FNavMeshRenderingHelpers
 			const FVector::FReal u = (FVector::FReal)i * ArcPtsScale;
 			const FVector Pt = EvalArc(Start, Dir, Length*Height, u);
 
-			DebugLines.Add(FDebugRenderSceneProxy::FDebugLine(Prev, Pt, Color.ToFColor(true)));
+			DebugLines.Add(FDebugRenderSceneProxy::FDebugLine(Prev, Pt, Color.ToFColor(true), LineThickness));
 			Prev = Pt;
 		}
 	}
@@ -152,19 +153,21 @@ namespace FNavMeshRenderingHelpers
 		const FVector Ay = (Origin - Tip).GetSafeNormal();
 		const FVector Ax = FVector::CrossProduct(Az, Ay);
 
-		DebugLines.Add(FDebugRenderSceneProxy::FDebugLine(Tip, FVector(Tip.X + Ay.X*Size + Ax.X*Size / 3, Tip.Y + Ay.Y*Size + Ax.Y*Size / 3, Tip.Z + Ay.Z*Size + Ax.Z*Size / 3), Color.ToFColor(true)));
-		DebugLines.Add(FDebugRenderSceneProxy::FDebugLine(Tip, FVector(Tip.X + Ay.X*Size - Ax.X*Size / 3, Tip.Y + Ay.Y*Size - Ax.Y*Size / 3, Tip.Z + Ay.Z*Size - Ax.Z*Size / 3), Color.ToFColor(true)));
+		DebugLines.Add(FDebugRenderSceneProxy::FDebugLine(Tip, FVector(Tip.X + Ay.X*Size + Ax.X*Size / 3, Tip.Y + Ay.Y*Size + Ax.Y*Size / 3, Tip.Z + Ay.Z*Size + Ax.Z*Size / 3), Color.ToFColor(true), LineThickness));
+		DebugLines.Add(FDebugRenderSceneProxy::FDebugLine(Tip, FVector(Tip.X + Ay.X*Size - Ax.X*Size / 3, Tip.Y + Ay.Y*Size - Ax.Y*Size / 3, Tip.Z + Ay.Z*Size - Ax.Z*Size / 3), Color.ToFColor(true), LineThickness));
 	}
 
-	void CacheLink(TArray<FDebugRenderSceneProxy::FDebugLine>& DebugLines, const FVector V0, const FVector V1, const FColor LinkColor, const uint8 LinkDirection)
+	void CacheLink(TArray<FDebugRenderSceneProxy::FDebugLine>& DebugLines, const FVector V0, const FVector V1, const FColor LinkColor, const uint8 LinkDirection, const bool bIsGenerated)
 	{
-		FNavMeshRenderingHelpers::CacheArc(DebugLines, V0, V1, 0.4f, 4, LinkColor, LinkLines_LineThickness);
+		const float LineThickness = bIsGenerated ? GeneratedLinkLines_LineThickness : LinkLines_LineThickness;
+
+		FNavMeshRenderingHelpers::CacheArc(DebugLines, V0, V1, 0.4f, 4, LinkColor, LineThickness);
 
 		const FVector VOffset(0, 0, FVector::Dist(V0, V1) * 1.333f);
-		FNavMeshRenderingHelpers::CacheArrowHead(DebugLines, V1, V0 + VOffset, 30.f, LinkColor, LinkLines_LineThickness);
+		FNavMeshRenderingHelpers::CacheArrowHead(DebugLines, V1, V0 + VOffset, 30.f, LinkColor, LineThickness);
 		if (LinkDirection)
 		{
-			FNavMeshRenderingHelpers::CacheArrowHead(DebugLines, V0, V1 + VOffset, 30.f, LinkColor, LinkLines_LineThickness);
+			FNavMeshRenderingHelpers::CacheArrowHead(DebugLines, V0, V1 + VOffset, 30.f, LinkColor, LineThickness);
 		}
 	}
 
@@ -759,7 +762,7 @@ void FNavMeshSceneProxyData::GatherData(const ARecastNavMesh* NavMesh, int32 InN
 					const FVector V1 = Link.Right + NavMeshDrawOffset;
 					const FColor LinkColor = ((Link.Direction && Link.ValidEnds) || (Link.ValidEnds & FRecastDebugGeometry::OMLE_Left)) ? FNavMeshRenderingHelpers::SemiDarkenColor(NavMeshColors[Link.AreaID]) : NavMeshRenderColor_OffMeshConnectionInvalid;
 
-					FNavMeshRenderingHelpers::CacheLink(NavLinkLines, V0, V1, LinkColor, Link.Direction);
+					FNavMeshRenderingHelpers::CacheLink(NavLinkLines, V0, V1, LinkColor, Link.Direction, Link.bIsGenerated);
 
 					// if the connection as a whole is valid check if there are any of ends is invalid
 					if (LinkColor != NavMeshRenderColor_OffMeshConnectionInvalid)
@@ -789,7 +792,7 @@ void FNavMeshSceneProxyData::GatherData(const ARecastNavMesh* NavMesh, int32 InN
 					const FVector V1 = Link.Right + NavMeshDrawOffset;
 					const FColor LinkColor = NavMeshRenderColor_PolyForbidden;
 
-					FNavMeshRenderingHelpers::CacheLink(NavLinkLines, V0, V1, LinkColor, Link.Direction);
+					FNavMeshRenderingHelpers::CacheLink(NavLinkLines, V0, V1, LinkColor, Link.Direction, Link.bIsGenerated);
 				}
 			}
 		}
@@ -1586,7 +1589,7 @@ void FNavMeshSceneProxy::GetDynamicMeshElements(const TArray<const FSceneView*>&
 				{
 					if (FNavMeshRenderingHelpers::LineInCorrectDistance(Line.Start, Line.End, View))
 					{
-						PDI->DrawLine(Line.Start, Line.End, Line.Color, SDPG_World, LinkLines_LineThickness, 0, true);
+						PDI->DrawLine(Line.Start, Line.End, Line.Color, SDPG_World, Line.Thickness, 0, true);
 					}
 					else if (bUseThickLines)
 					{
diff --git a/Engine/Source/Runtime/NavigationSystem/Private/NavMesh/PImplRecastNavMesh.cpp b/Engine/Source/Runtime/NavigationSystem/Private/NavMesh/PImplRecastNavMesh.cpp
--- a/Engine/Source/Runtime/NavigationSystem/Private/NavMesh/PImplRecastNavMesh.cpp
+++ b/Engine/Source/Runtime/NavigationSystem/Private/NavMesh/PImplRecastNavMesh.cpp
@@ -3026,6 +3026,7 @@ int32 FPImplRecastNavMesh::GetTilesDebugGeometry(const FRecastNavMeshGenerator*
 				, LinkPoly->getArea()
 				, (uint8)OffMeshConnection->getBiDirectional()
 				, GetValidEnds(*DetourNavMesh, Tile, *LinkPoly)
+				, OffMeshConnection->getIsGenerated()
 				, UE_REAL_TO_FLOAT_CLAMPED_MAX(OffMeshConnection->rad)
 			};
 
diff --git a/Engine/Source/Runtime/NavigationSystem/Private/NavMesh/RecastNavMeshGenerator.cpp b/Engine/Source/Runtime/NavigationSystem/Private/NavMesh/RecastNavMeshGenerator.cpp
--- a/Engine/Source/Runtime/NavigationSystem/Private/NavMesh/RecastNavMeshGenerator.cpp
+++ b/Engine/Source/Runtime/NavigationSystem/Private/NavMesh/RecastNavMeshGenerator.cpp
@@ -1354,7 +1354,8 @@ struct FOffMeshData
 
 			NewInfo.type = DT_OFFMESH_CON_POINT | 
 				(Link.Direction == ENavLinkDirection::BothWays ? DT_OFFMESH_CON_BIDIR : 0) |
-				(Link.bSnapToCheapestArea ? DT_OFFMESH_CON_CHEAPAREA : 0);
+				(Link.bSnapToCheapestArea ? DT_OFFMESH_CON_CHEAPAREA : 0) |
+				(Link.bIsGenerated ? DT_OFFMESH_CON_GENERATED : 0);
 
 			NewInfo.snapRadius = Link.SnapRadius;
 			NewInfo.snapHeight = Link.bUseSnapHeight ? Link.SnapHeight : DefaultSnapHeight;
@@ -3856,7 +3857,11 @@ dtStatus FRecastTileGenerator::BuildTileCacheLinks(FNavMeshBuildContext& BuildCo
 		// Since trajectory validation starts at agentClimb height to ignore small bumps, remove the offset for the actual link height.
 		midA[1] -= linkBuilderConfig.agentClimb;
 		midB[1] -= linkBuilderConfig.agentClimb;
-		OutGeneratedLinks.Add(FNavigationLink(Recast2UnrealPoint(midA), Recast2UnrealPoint(midB)));
+
+		FNavigationLink& NewLink = OutGeneratedLinks.Emplace_GetRef();
+		NewLink.bIsGenerated = true;
+		NewLink.Left = Recast2UnrealPoint(midA);
+		NewLink.Right = Recast2UnrealPoint(midB);
 	}
 	
 	LogOnExit();
diff --git a/Engine/Source/Runtime/NavigationSystem/Public/NavMesh/RecastNavMesh.h b/Engine/Source/Runtime/NavigationSystem/Public/NavMesh/RecastNavMesh.h
--- a/Engine/Source/Runtime/NavigationSystem/Public/NavMesh/RecastNavMesh.h
+++ b/Engine/Source/Runtime/NavigationSystem/Public/NavMesh/RecastNavMesh.h
@@ -200,6 +200,7 @@ struct FRecastDebugGeometry
 		uint8	AreaID;
 		uint8	Direction;
 		uint8	ValidEnds;
+		bool	bIsGenerated;
 		float	Radius;
 		float	Height;
 		FColor	Color;
diff --git a/Engine/Source/Runtime/Navmesh/Public/Detour/DetourNavMesh.h b/Engine/Source/Runtime/Navmesh/Public/Detour/DetourNavMesh.h
--- a/Engine/Source/Runtime/Navmesh/Public/Detour/DetourNavMesh.h
+++ b/Engine/Source/Runtime/Navmesh/Public/Detour/DetourNavMesh.h
@@ -106,13 +106,11 @@ static const unsigned short DT_EXT_LINK = 0x8000;
 inline const unsigned int DT_NULL_LINK = 0xffffffff;
 
 /// A flag that indicates that an off-mesh connection can be traversed in both directions. (Is bidirectional.)
-static const unsigned char DT_OFFMESH_CON_BIDIR = 0x01;
-
-//@UE BEGIN
-static const unsigned char DT_OFFMESH_CON_POINT = 0x02;
-static const unsigned char DT_OFFMESH_CON_SEGMENT = 0x04;
-static const unsigned char DT_OFFMESH_CON_CHEAPAREA = 0x08;
-//@UE END
+static const unsigned char DT_OFFMESH_CON_BIDIR		= 0x01;
+static const unsigned char DT_OFFMESH_CON_POINT		= 0x02;		//UE
+static const unsigned char DT_OFFMESH_CON_SEGMENT	= 0x04;		//UE
+static const unsigned char DT_OFFMESH_CON_CHEAPAREA = 0x08;		//UE
+static const unsigned char DT_OFFMESH_CON_GENERATED = 0x10;		//UE
 
 /// The maximum number of user defined area ids.
 /// @ingroup detour
@@ -333,10 +331,11 @@ struct dtOffMeshConnection
 
 	//@UE BEGIN
 	/// Sets link flags
-	inline void setFlags(unsigned char conFlags)
+	inline void setFlags(unsigned char conTypeFlags)
 	{
-		flags = ((conFlags & DT_OFFMESH_CON_BIDIR) ? 0x80 : 0) |
-			((conFlags & DT_OFFMESH_CON_CHEAPAREA) ? 0x40 : 0);
+		flags = ((conTypeFlags & DT_OFFMESH_CON_BIDIR) ? 0x80 : 0) |
+				((conTypeFlags & DT_OFFMESH_CON_CHEAPAREA) ? 0x40 : 0) |
+				((conTypeFlags & DT_OFFMESH_CON_GENERATED) ? 0x20 : 0);
 	}
 
 	/// Gets the link direction
@@ -344,6 +343,9 @@ struct dtOffMeshConnection
 
 	/// Gets the link snap mode
 	inline bool getSnapToCheapestArea() const { return (flags & 0x40) != 0; }
+
+	/// Indicates if the link was automatically generated
+	inline bool getIsGenerated() const { return (flags & 0x20) != 0; }
 	//@UE END
 };
 
