[NavigationSystem] deprecated methods/properties cleanup

diff --git a/Engine/Source/Runtime/NavigationSystem/Private/NavMesh/PImplRecastNavMesh.cpp b/Engine/Source/Runtime/NavigationSystem/Private/NavMesh/PImplRecastNavMesh.cpp
--- a/Engine/Source/Runtime/NavigationSystem/Private/NavMesh/PImplRecastNavMesh.cpp
+++ b/Engine/Source/Runtime/NavigationSystem/Private/NavMesh/PImplRecastNavMesh.cpp
@@ -606,15 +606,6 @@ void FPImplRecastNavMesh::Serialize( FArchive& Ar, int32 NavMeshVersion )
 		// assumes tiles are rectangular
 		
 		float DefaultCellSize = NavMeshOwner->GetCellSize(ENavigationDataResolution::Default);
-		PRAGMA_DISABLE_DEPRECATION_WARNINGS
-		if (NavMeshVersion < NAVMESHVER_TILE_RESOLUTIONS)
-		{
-			// For backward compatibility, read original CellSize value.
-			// In ARecastNavMesh::PostLoad(), cell sizes for the different resolutions are set to the old CellSize value but it occurs later (PostLoad).
-			// This why we explicitly need to read CellSize for older versions here.
-			DefaultCellSize = NavMeshOwner->CellSize;
-		}
-		PRAGMA_ENABLE_DEPRECATION_WARNINGS
 		const FVector::FReal ActorsTileSize = static_cast<FVector::FReal>(FMath::TruncToInt(NavMeshOwner->TileSizeUU / DefaultCellSize)) * DefaultCellSize;
 
 		if (ActorsTileSize != Params.tileWidth)
@@ -1143,13 +1134,6 @@ void FPImplRecastNavMesh::Raycast(const FVector& StartLoc, const FVector& EndLoc
 	}
 }
 
-// DEPRECATED
-ENavigationQueryResult::Type FPImplRecastNavMesh::FindPath(const FVector& StartLoc, const FVector& EndLoc, const FVector::FReal CostLimit, FNavMeshPath& Path, const FNavigationQueryFilter& InQueryFilter, const UObject* Owner) const
-{
-	constexpr bool bRequireNavigableEndLocation = true; 
-	return FindPath(StartLoc, EndLoc, CostLimit, bRequireNavigableEndLocation, Path, InQueryFilter, Owner);
-}
-
 // @TODONAV
 ENavigationQueryResult::Type FPImplRecastNavMesh::FindPath(const FVector& StartLoc, const FVector& EndLoc, const FVector::FReal CostLimit, const bool bRequireNavigableEndLocation, FNavMeshPath& Path, const FNavigationQueryFilter& InQueryFilter, const UObject* Owner) const
 {
@@ -1249,13 +1233,6 @@ ENavigationQueryResult::Type FPImplRecastNavMesh::PostProcessPathInternal(dtStat
 	return DTStatusToNavQueryResult(FindPathStatus);
 }
 
-// DEPRECATED
-ENavigationQueryResult::Type FPImplRecastNavMesh::TestPath(const FVector& StartLoc, const FVector& EndLoc, const FNavigationQueryFilter& InQueryFilter, const UObject* Owner, int32* NumVisitedNodes) const
-{
-	constexpr bool bRequireNavigableEndLocation = true;
-	return TestPath(StartLoc, EndLoc, bRequireNavigableEndLocation, InQueryFilter, Owner, NumVisitedNodes);
-}
-
 ENavigationQueryResult::Type FPImplRecastNavMesh::TestPath(const FVector& StartLoc, const FVector& EndLoc, const bool bRequireNavigableEndLocation, const FNavigationQueryFilter& InQueryFilter, const UObject* Owner, int32* NumVisitedNodes) const
 {
 	// Same check as in FPImplRecastNavMesh::FindPath (ex: this can occur when tileWidth of loading navmesh mismatch).
@@ -1679,13 +1656,6 @@ static void StorePathfindingDebugStep(const dtNavMeshQuery& NavQuery, const dtNa
 	}
 }
 
-// DEPRECATED
-int32 FPImplRecastNavMesh::DebugPathfinding(const FVector& StartLoc, const FVector& EndLoc, const FVector::FReal CostLimit, const FNavigationQueryFilter& Filter, const UObject* Owner, TArray<FRecastDebugPathfindingData>& Steps)
-{
-	constexpr bool bRequireNavigableEndLocation = true;
-	return DebugPathfinding(StartLoc, EndLoc, CostLimit, bRequireNavigableEndLocation, Filter, Owner, Steps);
-}
-
 int32 FPImplRecastNavMesh::DebugPathfinding(const FVector& StartLoc, const FVector& EndLoc, const FVector::FReal CostLimit, const bool bRequireNavigableEndLocation, const FNavigationQueryFilter& Filter, const UObject* Owner, TArray<FRecastDebugPathfindingData>& Steps)
 {
 	int32 NumSteps = 0;
diff --git a/Engine/Source/Runtime/NavigationSystem/Private/NavMesh/RecastNavMesh.cpp b/Engine/Source/Runtime/NavigationSystem/Private/NavMesh/RecastNavMesh.cpp
--- a/Engine/Source/Runtime/NavigationSystem/Private/NavMesh/RecastNavMesh.cpp
+++ b/Engine/Source/Runtime/NavigationSystem/Private/NavMesh/RecastNavMesh.cpp
@@ -679,35 +679,6 @@ void ARecastNavMesh::PostLoad()
 			UNavigationSystemBase::OnNavigationInitStartStaticDelegate().AddUObject(this, &ARecastNavMesh::CheckToDiscardSubLevelNavData);
 		}
 	}
-
-	PRAGMA_DISABLE_DEPRECATION_WARNINGS
-	// If needed, initialize from deprecated value.
-	if (NavMeshVersion < NAVMESHVER_TILE_RESOLUTIONS)
-	{
-		for (int i = 0; i < (uint8)ENavigationDataResolution::MAX; ++i)
-		{
-			SetCellSize((ENavigationDataResolution)i, CellSize);
-		}
-	}
-
-	// If needed, initialize CellHeight from the deprecated value.
-	if (NavMeshVersion < NAVMESHVER_TILE_RESOLUTIONS_CELLHEIGHT)
-	{
-		for (int i = 0; i < (uint8)ENavigationDataResolution::MAX; ++i)
-		{
-			SetCellHeight((ENavigationDataResolution)i, CellHeight);
-		}
-	}
-
-	// If needed, initialize AgentMaxStepHeight from the deprecated value.
-	if (NavMeshVersion < NAVMESHVER_TILE_RESOLUTIONS_AGENTMAXSTEPHEIGHT)
-	{
-		for (int i = 0; i < (uint8)ENavigationDataResolution::MAX; ++i)
-		{
-			SetAgentMaxStepHeight((ENavigationDataResolution)i, AgentMaxStepHeight);
-		}
-	}
-	PRAGMA_ENABLE_DEPRECATION_WARNINGS
 	
 	for (uint8 Index = 0; Index < (uint8)ENavigationDataResolution::MAX; Index++)
 	{
@@ -2624,39 +2595,6 @@ bool ARecastNavMesh::GetPolysWithinPathingDistance(FVector const& StartLoc, cons
 	return RecastNavMeshImpl && RecastNavMeshImpl->GetPolysWithinPathingDistance(StartLoc, PathingDistance, GetRightFilterRef(Filter), QueryOwner, FoundPolys, DebugData);
 }
 
-// Deprecated
-void ARecastNavMesh::GetDebugGeometry(FRecastDebugGeometry& OutGeometry, int32 TileIndex) const
-{
-	if (RecastNavMeshImpl)
-	{
-		FNavTileRef TileRef;
-		if (TileIndex != INDEX_NONE)
-		{
-			TArray<FNavTileRef> TileRefs;
-			FNavTileRef::DeprecatedMakeTileRefsFromTileIds(RecastNavMeshImpl, { static_cast<uint32>(TileIndex) }, TileRefs);
-			TileRef = TileRefs[0];
-		}
-		GetDebugGeometryForTile(OutGeometry, TileRef);
-	}
-}
-
-// Deprecated
-bool ARecastNavMesh::GetDebugGeometryForTile(FRecastDebugGeometry& OutGeometry, int32 TileIndex) const
-{
-	if (RecastNavMeshImpl)
-	{
-		FNavTileRef TileRef;
-		if (TileIndex != INDEX_NONE)
-		{
-			TArray<FNavTileRef> TileRefs;
-			FNavTileRef::DeprecatedMakeTileRefsFromTileIds(RecastNavMeshImpl, { static_cast<uint32>(TileIndex) }, TileRefs);
-			TileRef = TileRefs[0];
-		}
-		return GetDebugGeometryForTile(OutGeometry, TileRef);
-	}
-	return true;
-}
-
 bool ARecastNavMesh::GetDebugGeometryForTile(FRecastDebugGeometry& OutGeometry, FNavTileRef TileRef) const
 {
 	if (RecastNavMeshImpl)
@@ -2737,27 +2675,11 @@ void ARecastNavMesh::DrawDebugPathCorridor(NavNodeRef const* PathPolys, int32 Nu
 #endif // ENABLE_DRAW_DEBUG
 }
 
-// Deprecated
-void ARecastNavMesh::OnNavMeshTilesUpdated(const TArray<uint32>& ChangedTiles)
-{
-	TArray<FNavTileRef> ChangedTileRefs;
-	FNavTileRef::DeprecatedMakeTileRefsFromTileIds(RecastNavMeshImpl, ChangedTiles, ChangedTileRefs);
-	OnNavMeshTilesUpdated(ChangedTileRefs);
-}
-
 void ARecastNavMesh::OnNavMeshTilesUpdated(const TArray<FNavTileRef>& ChangedTiles)
 {
 	InvalidateAffectedPaths(ChangedTiles);
 }
 
-// Deprecated
-void ARecastNavMesh::InvalidateAffectedPaths(const TArray<uint32>& ChangedTiles)
-{
-	TArray<FNavTileRef> ChangedTileRefs;
-	FNavTileRef::DeprecatedMakeTileRefsFromTileIds(RecastNavMeshImpl, ChangedTiles, ChangedTileRefs);
-	InvalidateAffectedPaths(ChangedTileRefs);
-}
-
 void ARecastNavMesh::InvalidateAffectedPaths(const TArray<FNavTileRef>& ChangedTiles)
 {
 	const int32 PathsCount = ActivePaths.Num();
@@ -3502,21 +3424,11 @@ void ARecastNavMesh::PostEditChangeChainProperty(FPropertyChangedChainEvent& Pro
 						}
 					}
 
-					PRAGMA_DISABLE_DEPRECATION_WARNINGS
-					// Update the deprecated CellSize to fit the default resolution CellSize
-					CellSize = NavMeshResolutionParams[(uint8)ENavigationDataResolution::Default].CellSize;
-					PRAGMA_ENABLE_DEPRECATION_WARNINGS
-
 					bRebuild = true;
 				}
 			}
 			else if (PropName == GET_MEMBER_NAME_CHECKED(FNavMeshResolutionParam, CellHeight))
 			{
-				PRAGMA_DISABLE_DEPRECATION_WARNINGS
-				// Update the deprecated CellHeight to fit the default resolution CellHeight
-				CellHeight = NavMeshResolutionParams[(uint8)ENavigationDataResolution::Default].CellHeight;
-				PRAGMA_ENABLE_DEPRECATION_WARNINGS
-
 				bRebuild = true;
 			}
 			else if (MemberProperty && (MemberProperty->GetFName() == NAME_NavLinkJumpDownConfig || MemberProperty->GetFName() == NAME_NavLinkJumpOverConfig))
@@ -3582,11 +3494,6 @@ void ARecastNavMesh::PostEditChangeProperty(FPropertyChangedEvent& PropertyChang
 					SetCellSize((ENavigationDataResolution)Index, NewCellSize);
 				}
 
-				PRAGMA_DISABLE_DEPRECATION_WARNINGS
-				// Set deprecated CellSize
-				CellSize = GetCellSize(ENavigationDataResolution::Default);
-				PRAGMA_ENABLE_DEPRECATION_WARNINGS
-
 				// update config
 				FillConfig(NavDataConfig);
 			}
@@ -3696,41 +3603,6 @@ void ARecastNavMesh::ConditionalConstructGenerator()
 	}
 }
 
-// Deprecated
-PRAGMA_DISABLE_DEPRECATION_WARNINGS
-void ARecastNavMesh::UpdateGenerationProperties(const FRecastNavMeshGenerationProperties& GenerationProps)
-{
-	TilePoolSize = GenerationProps.TilePoolSize;
-	TileSizeUU = GenerationProps.TileSizeUU;
-
-	CellSize = GenerationProps.CellSize;
-	CellHeight = GenerationProps.CellHeight;
-
-	AgentRadius = GenerationProps.AgentRadius;
-	AgentHeight = GenerationProps.AgentHeight;
-	AgentMaxSlope = GenerationProps.AgentMaxSlope;
-
-	AgentMaxStepHeight = GenerationProps.AgentMaxStepHeight;
-
-	MinRegionArea = GenerationProps.MinRegionArea;
-	MergeRegionSize = GenerationProps.MergeRegionSize;
-	MaxSimplificationError = GenerationProps.MaxSimplificationError;
-	TileNumberHardLimit = GenerationProps.TileNumberHardLimit;
-	RegionPartitioning = GenerationProps.RegionPartitioning;
-	LayerPartitioning = GenerationProps.LayerPartitioning;
-	RegionChunkSplits = GenerationProps.RegionChunkSplits;
-	LayerChunkSplits = GenerationProps.LayerChunkSplits;
-	bSortNavigationAreasByCost = GenerationProps.bSortNavigationAreasByCost;
-	bPerformVoxelFiltering = GenerationProps.bPerformVoxelFiltering;
-	bMarkLowHeightAreas = GenerationProps.bMarkLowHeightAreas;
-	bUseExtraTopCellWhenMarkingAreas = GenerationProps.bUseExtraTopCellWhenMarkingAreas;
-	bFilterLowSpanSequences = GenerationProps.bFilterLowSpanSequences;
-	bFilterLowSpanFromTileCache = GenerationProps.bFilterLowSpanFromTileCache;
-	bFixedTilePoolSize = GenerationProps.bFixedTilePoolSize;
-	bIsWorldPartitioned = GenerationProps.bIsWorldPartitioned;
-}
-PRAGMA_ENABLE_DEPRECATION_WARNINGS
-
 bool ARecastNavMesh::IsVoxelCacheEnabled()
 {
 #if RECAST_ASYNC_REBUILDING
@@ -4045,18 +3917,6 @@ void ARecastNavMesh::RemoveTiles(const TArray<FIntPoint>& Tiles)
 	}
 }
 
-// Deprecated
-void ARecastNavMesh::RebuildTile(const TArray<FIntPoint>& Tiles)
-{
-	TArray<FNavMeshDirtyTileElement> ActiveTiles;
-	ActiveTiles.Reserve(Tiles.Num());
-	for (const FIntPoint& Point : Tiles)
-	{
-		ActiveTiles.Add(FNavMeshDirtyTileElement{Point, TNumericLimits<FVector::FReal>::Max(), ENavigationInvokerPriority::Default});
-	}
-	RebuildTile(ActiveTiles);
-}
-
 void ARecastNavMesh::RebuildTile(const TArray<FNavMeshDirtyTileElement>& Tiles)
 {
 	TRACE_CPUPROFILER_EVENT_SCOPE(ARecastNavMesh::RebuildTile);
diff --git a/Engine/Source/Runtime/NavigationSystem/Private/NavMesh/RecastNavMeshDataChunk.cpp b/Engine/Source/Runtime/NavigationSystem/Private/NavMesh/RecastNavMeshDataChunk.cpp
--- a/Engine/Source/Runtime/NavigationSystem/Private/NavMesh/RecastNavMeshDataChunk.cpp
+++ b/Engine/Source/Runtime/NavigationSystem/Private/NavMesh/RecastNavMeshDataChunk.cpp
@@ -199,54 +199,6 @@ void URecastNavMeshDataChunk::SerializeRecastData(FArchive& Ar, int32 NavMeshVer
 }
 #endif// WITH_RECAST
 
-// Deprecated
-TArray<uint32> URecastNavMeshDataChunk::AttachTiles(FPImplRecastNavMesh& NavMeshImpl)
-{
-	TArray<uint32> TileIds;
-#if WITH_RECAST
-	check(NavMeshImpl.NavMeshOwner);
-	const TArray<FNavTileRef> TileRefs = AttachTiles(*NavMeshImpl.NavMeshOwner);
-	FNavTileRef::DeprecatedGetTileIdsFromNavTileRefs(&NavMeshImpl, TileRefs, TileIds);
-#endif // WITH_RECAST
-	return TileIds;
-}
-
-// Deprecated
-TArray<uint32> URecastNavMeshDataChunk::AttachTiles(FPImplRecastNavMesh& NavMeshImpl, const bool bKeepCopyOfData, const bool bKeepCopyOfCacheData)
-{
-	TArray<uint32> TileIds;
-#if WITH_RECAST
-	check(NavMeshImpl.NavMeshOwner);
-	const TArray<FNavTileRef> TileRefs = AttachTiles(*NavMeshImpl.NavMeshOwner, bKeepCopyOfData, bKeepCopyOfCacheData);
-	FNavTileRef::DeprecatedGetTileIdsFromNavTileRefs(&NavMeshImpl, TileRefs, TileIds);
-#endif // WITH_RECAST
-	return TileIds;
-}
-
-// Deprecated
-TArray<uint32> URecastNavMeshDataChunk::DetachTiles(FPImplRecastNavMesh& NavMeshImpl)
-{
-	TArray<uint32> TileIds;
-#if WITH_RECAST
-	check(NavMeshImpl.NavMeshOwner);
-	const TArray<FNavTileRef> TileRefs = DetachTiles(*NavMeshImpl.NavMeshOwner);
-	FNavTileRef::DeprecatedGetTileIdsFromNavTileRefs(&NavMeshImpl, TileRefs, TileIds);
-#endif // WITH_RECAST
-	return TileIds;
-}
-
-// Deprecated
-TArray<uint32> URecastNavMeshDataChunk::DetachTiles(FPImplRecastNavMesh& NavMeshImpl, const bool bTakeDataOwnership, const bool bTakeCacheDataOwnership)
-{
-	TArray<uint32> TileIds;
-#if WITH_RECAST
-	check(NavMeshImpl.NavMeshOwner);
-	const TArray<FNavTileRef> TileRefs = DetachTiles(*NavMeshImpl.NavMeshOwner, bTakeDataOwnership, bTakeCacheDataOwnership);
-	FNavTileRef::DeprecatedGetTileIdsFromNavTileRefs(&NavMeshImpl, TileRefs, TileIds);
-#endif // WITH_RECAST
-	return TileIds;
-}
-
 #if WITH_RECAST
 TArray<FNavTileRef> URecastNavMeshDataChunk::AttachTiles(ARecastNavMesh& NavMesh)
 {
diff --git a/Engine/Source/Runtime/NavigationSystem/Private/NavMesh/RecastNavMeshGenerator.cpp b/Engine/Source/Runtime/NavigationSystem/Private/NavMesh/RecastNavMeshGenerator.cpp
--- a/Engine/Source/Runtime/NavigationSystem/Private/NavMesh/RecastNavMeshGenerator.cpp
+++ b/Engine/Source/Runtime/NavigationSystem/Private/NavMesh/RecastNavMeshGenerator.cpp
@@ -4882,9 +4882,6 @@ FRecastNavMeshGenerator::FRecastNavMeshGenerator(ARecastNavMesh& InDestNavMesh)
 	, DestNavMesh(&InDestNavMesh)
 	, bInitialized(false)
 	, bRestrictBuildingToActiveTiles(false)
-PRAGMA_DISABLE_DEPRECATION_WARNINGS
-	, bSortTilesWithSeedLocations(true)
-PRAGMA_ENABLE_DEPRECATION_WARNINGS	
 	, Version(0)
 {
 	INC_DWORD_STAT_BY(STAT_NavigationMemory, sizeof(*this));
@@ -5632,17 +5629,6 @@ void FRecastNavMeshGenerator::RemoveTiles(const TArray<FIntPoint>& Tiles)
 	}
 }
 
-// Deprecated 
-void FRecastNavMeshGenerator::ReAddTiles(const TArray<FIntPoint>& Tiles)
-{
-	TArray<FNavMeshDirtyTileElement> ConvertedTiles;
-	for (const FIntPoint& Point : Tiles)
-	{
-		ConvertedTiles.Add(FNavMeshDirtyTileElement{Point, TNumericLimits<FVector::FReal>::Max(), ENavigationInvokerPriority::Default});
-	}
-	ReAddTiles(ConvertedTiles);
-}
-	
 void FRecastNavMeshGenerator::ReAddTiles(const TArray<FNavMeshDirtyTileElement>& Tiles)
 {
 	static const FVector Expansion(1, 1, BIG_NUMBER);
diff --git a/Engine/Source/Runtime/NavigationSystem/Private/NavigationData.cpp b/Engine/Source/Runtime/NavigationSystem/Private/NavigationData.cpp
--- a/Engine/Source/Runtime/NavigationSystem/Private/NavigationData.cpp
+++ b/Engine/Source/Runtime/NavigationSystem/Private/NavigationData.cpp
@@ -865,44 +865,4 @@ uint32 ANavigationData::LogMemUsed() const
 	}
 
 	return MemUsed;
-}
-
-//------------------------------------------------------------------------//
-// deprecated functions
-//------------------------------------------------------------------------//
-void ANavigationData::DrawDebugPath(FNavigationPath* Path, FColor PathColor, UCanvas* Canvas, bool bPersistent, const uint32 NextPathPointIndex) const
-{
-	DrawDebugPath(Path, PathColor, Canvas, bPersistent, -1.f, NextPathPointIndex);
-}
-
-ENavigationQueryResult::Type ANavigationData::CalcPathCost(const FVector& PathStart, const FVector& PathEnd, float& OutPathCost, FSharedConstNavQueryFilter QueryFilter, const UObject* Querier) const
-{
-	FVector::FReal PathCost = OutPathCost;
-
-	const ENavigationQueryResult::Type Result = CalcPathCost(PathStart, PathEnd, PathCost, QueryFilter, Querier);
-	OutPathCost = UE_REAL_TO_FLOAT_CLAMPED(PathCost);
-
-	return Result;
-}
-
-ENavigationQueryResult::Type ANavigationData::CalcPathLength(const FVector& PathStart, const FVector& PathEnd, float& OutPathLength, FSharedConstNavQueryFilter QueryFilter, const UObject* Querier) const
-{
-	FVector::FReal PathLength = OutPathLength;
-
-	const ENavigationQueryResult::Type Result = CalcPathLength(PathStart, PathEnd, PathLength, QueryFilter, Querier);
-	OutPathLength = UE_REAL_TO_FLOAT_CLAMPED(PathLength);
-
-	return Result;
-}
-
-ENavigationQueryResult::Type ANavigationData::CalcPathLengthAndCost(const FVector& PathStart, const FVector& PathEnd, float& OutPathLength, float& OutPathCost, FSharedConstNavQueryFilter QueryFilter, const UObject* Querier) const
-{
-	FVector::FReal PathLength = OutPathLength;
-	FVector::FReal PathCost = OutPathCost;
-
-	const ENavigationQueryResult::Type Result = CalcPathLengthAndCost(PathStart, PathEnd, PathLength, PathCost, QueryFilter, Querier);
-	OutPathLength = UE_REAL_TO_FLOAT_CLAMPED(PathLength);
-	OutPathCost = UE_REAL_TO_FLOAT_CLAMPED(PathCost);
-
-	return Result;
-}
+}
\ No newline at end of file
diff --git a/Engine/Source/Runtime/NavigationSystem/Private/NavigationPath.cpp b/Engine/Source/Runtime/NavigationSystem/Private/NavigationPath.cpp
--- a/Engine/Source/Runtime/NavigationSystem/Private/NavigationPath.cpp
+++ b/Engine/Source/Runtime/NavigationSystem/Private/NavigationPath.cpp
@@ -597,13 +597,4 @@ void UNavigationPath::SetPathPointsFromPath(FNavigationPath& NativePath)
 	{
 		PathPoints.Add(PathPoint.Location);
 	}
-}
-
-
-//------------------------------------------------------------------------//
-// deprecated functions
-//------------------------------------------------------------------------//
-void FNavigationPath::DebugDraw(const ANavigationData* NavData, FColor PathColor, UCanvas* Canvas, bool bPersistent, const uint32 NextPathPointIndex) const
-{
-	DebugDraw(NavData, PathColor, Canvas, bPersistent, -1.f, NextPathPointIndex);
-}
+}
\ No newline at end of file
diff --git a/Engine/Source/Runtime/NavigationSystem/Private/NavigationSystem.cpp b/Engine/Source/Runtime/NavigationSystem/Private/NavigationSystem.cpp
--- a/Engine/Source/Runtime/NavigationSystem/Private/NavigationSystem.cpp
+++ b/Engine/Source/Runtime/NavigationSystem/Private/NavigationSystem.cpp
@@ -1226,13 +1226,13 @@ void UNavigationSystemV1::OnWorldInitDone(FNavigationSystemRunMode Mode)
 			const bool bIsBuildLocked = IsNavigationBuildingLocked();
 			const bool bCanRebuild = !bIsBuildLocked && GetIsAutoUpdateEnabled();
 
-			if (GetDefaultNavDataInstance(FNavigationSystem::DontCreate) != NULL)
+			if (GetDefaultNavDataInstance(FNavigationSystem::DontCreate) != nullptr)
 			{
 				// trigger navmesh update
 				for (TActorIterator<ANavigationData> It(World); It; ++It)
 				{
 					ANavigationData* NavData = (*It);
-					if (NavData != NULL)
+					if (NavData != nullptr)
 					{
 						const ERegistrationResult Result = RegisterNavData(NavData);
 						LogNavDataRegistrationResult(Result);
@@ -1385,7 +1385,7 @@ void UNavigationSystemV1::SetCrowdManager(UCrowdManagerBase* NewCrowdManager)
 		CrowdManager->RemoveFromRoot();
 	}
 	CrowdManager = NewCrowdManager;
-	if (NewCrowdManager != NULL)
+	if (NewCrowdManager != nullptr)
 	{
 		CrowdManager->AddToRoot();
 	}
@@ -1901,12 +1901,12 @@ bool UNavigationSystemV1::GetRandomPoint(FNavLocation& ResultLocation, ANavigati
 {
 	SCOPE_CYCLE_COUNTER(STAT_Navigation_QueriesTimeSync);
 
-	if (NavData == NULL)
+	if (NavData == nullptr)
 	{
 		NavData = MainNavData;
 	}
 
-	if (NavData != NULL)
+	if (NavData != nullptr)
 	{
 		ResultLocation = NavData->GetRandomPoint(QueryFilter);
 		return true;
@@ -1939,101 +1939,66 @@ bool UNavigationSystemV1::GetRandomPointInNavigableRadius(const FVector& Origin,
 	return NavData != nullptr && NavData->GetRandomPointInNavigableRadius(Origin, Radius, ResultLocation, QueryFilter);
 }
 
-ENavigationQueryResult::Type UNavigationSystemV1::GetPathCost(const FVector& PathStart, const FVector& PathEnd, float& OutPathCost, const ANavigationData* NavData, FSharedConstNavQueryFilter QueryFilter) const
-{
-	FVector::FReal PathCost = OutPathCost;
-
-	const ENavigationQueryResult::Type Result = GetPathCost(PathStart, PathEnd, PathCost, NavData, QueryFilter);
-
-	OutPathCost = UE_REAL_TO_FLOAT_CLAMPED(PathCost);
-
-	return Result;
-}
-
 ENavigationQueryResult::Type UNavigationSystemV1::GetPathCost(const FVector& PathStart, const FVector& PathEnd, FVector::FReal& OutPathCost, const ANavigationData* NavData, FSharedConstNavQueryFilter QueryFilter) const
 {
 	SCOPE_CYCLE_COUNTER(STAT_Navigation_QueriesTimeSync);
 
-	if (NavData == NULL)
+	if (NavData == nullptr)
 	{
 		NavData = GetDefaultNavDataInstance();
 	}
 
-	return NavData != NULL ? NavData->CalcPathCost(PathStart, PathEnd, OutPathCost, QueryFilter) : ENavigationQueryResult::Error;
-}
-
-ENavigationQueryResult::Type UNavigationSystemV1::GetPathLength(const FVector& PathStart, const FVector& PathEnd, float& OutPathLength, const ANavigationData* NavData, FSharedConstNavQueryFilter QueryFilter) const
-{
-	FVector::FReal PathLength = OutPathLength;
-
-	const ENavigationQueryResult::Type Result = GetPathLength(PathStart, PathEnd, PathLength, NavData, QueryFilter);
-
-	OutPathLength = UE_REAL_TO_FLOAT_CLAMPED(PathLength);
-
-	return Result;
+	return NavData != nullptr ? NavData->CalcPathCost(PathStart, PathEnd, OutPathCost, QueryFilter) : ENavigationQueryResult::Error;
 }
 
 ENavigationQueryResult::Type UNavigationSystemV1::GetPathLength(const FVector& PathStart, const FVector& PathEnd, FVector::FReal& OutPathLength, const ANavigationData* NavData, FSharedConstNavQueryFilter QueryFilter) const
 {
 	SCOPE_CYCLE_COUNTER(STAT_Navigation_QueriesTimeSync);
 
-	if (NavData == NULL)
+	if (NavData == nullptr)
 	{
 		NavData = GetDefaultNavDataInstance();
 	}
 
-	return NavData != NULL ? NavData->CalcPathLength(PathStart, PathEnd, OutPathLength, QueryFilter) : ENavigationQueryResult::Error;
-}
-
-ENavigationQueryResult::Type UNavigationSystemV1::GetPathLengthAndCost(const FVector& PathStart, const FVector& PathEnd, float& OutPathLength, float& OutPathCost, const ANavigationData* NavData, FSharedConstNavQueryFilter QueryFilter) const
-{
-	FVector::FReal PathLength = OutPathLength;
-	FVector::FReal PathCost = OutPathCost;
-
-	const ENavigationQueryResult::Type Result = GetPathLengthAndCost(PathStart, PathEnd, PathLength, PathCost, NavData, QueryFilter);
-
-	OutPathLength = UE_REAL_TO_FLOAT_CLAMPED(PathLength);
-	OutPathCost = UE_REAL_TO_FLOAT_CLAMPED(PathCost);
-
-	return Result;
+	return NavData != nullptr ? NavData->CalcPathLength(PathStart, PathEnd, OutPathLength, QueryFilter) : ENavigationQueryResult::Error;
 }
 
 ENavigationQueryResult::Type UNavigationSystemV1::GetPathLengthAndCost(const FVector& PathStart, const FVector& PathEnd, FVector::FReal& OutPathLength, FVector::FReal& OutPathCost, const ANavigationData* NavData, FSharedConstNavQueryFilter QueryFilter) const
 {
 	SCOPE_CYCLE_COUNTER(STAT_Navigation_QueriesTimeSync);
 
-	if (NavData == NULL)
+	if (NavData == nullptr)
 	{
 		NavData = GetDefaultNavDataInstance();
 	}
 
-	return NavData != NULL ? NavData->CalcPathLengthAndCost(PathStart, PathEnd, OutPathLength, OutPathCost, QueryFilter) : ENavigationQueryResult::Error;
+	return NavData != nullptr ? NavData->CalcPathLengthAndCost(PathStart, PathEnd, OutPathLength, OutPathCost, QueryFilter) : ENavigationQueryResult::Error;
 }
 
 bool UNavigationSystemV1::ProjectPointToNavigation(const FVector& Point, FNavLocation& OutLocation, const FVector& Extent, const ANavigationData* NavData, FSharedConstNavQueryFilter QueryFilter) const
 {
 	SCOPE_CYCLE_COUNTER(STAT_Navigation_QueriesTimeSync);
 
-	if (NavData == NULL)
+	if (NavData == nullptr)
 	{
 		NavData = GetDefaultNavDataInstance();
 	}
 
-	return NavData != NULL && NavData->ProjectPoint(Point, OutLocation
-		, FNavigationSystem::IsValidExtent(Extent) ? Extent : NavData->GetConfig().DefaultQueryExtent
-		, QueryFilter);
+	return NavData != nullptr && NavData->ProjectPoint(Point, OutLocation
+	                                            , FNavigationSystem::IsValidExtent(Extent) ? Extent : NavData->GetConfig().DefaultQueryExtent
+	                                            , QueryFilter);
 }
 
 UNavigationPath* UNavigationSystemV1::FindPathToActorSynchronously(UObject* WorldContextObject, const FVector& PathStart, AActor* GoalActor, float TetherDistance, AActor* PathfindingContext, TSubclassOf<UNavigationQueryFilter> FilterClass)
 {
-	if (GoalActor == NULL)
+	if (GoalActor == nullptr)
 	{
-		return NULL; 
+		return nullptr; 
 	}
 
 	INavAgentInterface* NavAgent = Cast<INavAgentInterface>(GoalActor);
 	UNavigationPath* GeneratedPath = FindPathToLocationSynchronously(WorldContextObject, PathStart, NavAgent ? NavAgent->GetNavAgentLocation() : GoalActor->GetActorLocation(), PathfindingContext, FilterClass);
-	if (GeneratedPath != NULL && GeneratedPath->GetPath().IsValid() == true)
+	if (GeneratedPath != nullptr && GeneratedPath->GetPath().IsValid() == true)
 	{
 		GeneratedPath->GetPath()->SetGoalActorObservation(*GoalActor, TetherDistance);
 	}
@@ -2043,18 +2008,18 @@ UNavigationPath* UNavigationSystemV1::FindPathToActorSynchronously(UObject* Worl
 
 UNavigationPath* UNavigationSystemV1::FindPathToLocationSynchronously(UObject* WorldContextObject, const FVector& PathStart, const FVector& PathEnd, AActor* PathfindingContext, TSubclassOf<UNavigationQueryFilter> FilterClass)
 {
-	UWorld* World = NULL;
+	UWorld* World = nullptr;
 
-	if (WorldContextObject != NULL)
+	if (WorldContextObject != nullptr)
 	{
 		World = GEngine->GetWorldFromContextObject(WorldContextObject, EGetWorldErrorMode::LogAndReturnNull);
 	}
-	if (World == NULL && PathfindingContext != NULL)
+	if (World == nullptr && PathfindingContext != nullptr)
 	{
 		World = GEngine->GetWorldFromContextObject(PathfindingContext, EGetWorldErrorMode::LogAndReturnNull);
 	}
 
-	UNavigationPath* ResultPath = NULL;
+	UNavigationPath* ResultPath = nullptr;
 
 	UNavigationSystemV1* NavSys = FNavigationSystem::GetCurrent<UNavigationSystemV1>(World);
 
@@ -2062,13 +2027,13 @@ UNavigationPath* UNavigationSystemV1::FindPathToLocationSynchronously(UObject* W
 	{
 		ResultPath = NewObject<UNavigationPath>(NavSys);
 		bool bValidPathContext = false;
-		const ANavigationData* NavigationData = NULL;
+		const ANavigationData* NavigationData = nullptr;
 
-		if (PathfindingContext != NULL)
+		if (PathfindingContext != nullptr)
 		{
 			INavAgentInterface* NavAgent = Cast<INavAgentInterface>(PathfindingContext);
 			
-			if (NavAgent != NULL)
+			if (NavAgent != nullptr)
 			{
 				const FNavAgentProperties& AgentProps = NavAgent->GetNavAgentPropertiesRef();
 				NavigationData = NavSys->GetNavDataForProps(AgentProps, PathStart);
@@ -2101,13 +2066,13 @@ UNavigationPath* UNavigationSystemV1::FindPathToLocationSynchronously(UObject* W
 
 bool UNavigationSystemV1::NavigationRaycast(UObject* WorldContextObject, const FVector& RayStart, const FVector& RayEnd, FVector& HitLocation, TSubclassOf<UNavigationQueryFilter> FilterClass, AController* Querier)
 {
-	UWorld* World = NULL;
+	UWorld* World = nullptr;
 
-	if (WorldContextObject != NULL)
+	if (WorldContextObject != nullptr)
 	{
 		World = GEngine->GetWorldFromContextObject(WorldContextObject, EGetWorldErrorMode::LogAndReturnNull);
 	}
-	if (World == NULL && Querier != NULL)
+	if (World == nullptr && Querier != nullptr)
 	{
 		World = GEngine->GetWorldFromContextObject(Querier, EGetWorldErrorMode::LogAndReturnNull);
 	}
@@ -2121,19 +2086,19 @@ bool UNavigationSystemV1::NavigationRaycast(UObject* WorldContextObject, const F
 	if (NavSys)
 	{
 		// figure out which navigation data to use
-		const ANavigationData* NavData = NULL;
+		const ANavigationData* NavData = nullptr;
 		INavAgentInterface* MyNavAgent = Cast<INavAgentInterface>(Querier);
 		if (MyNavAgent)
 		{
 			const FNavAgentProperties& AgentProps = MyNavAgent->GetNavAgentPropertiesRef();
 			NavData = NavSys->GetNavDataForProps(AgentProps, RayStart);
 		}
-		if (NavData == NULL)
+		if (NavData == nullptr)
 		{
 			NavData = NavSys->GetDefaultNavDataInstance();
 		}
 
-		if (NavData != NULL)
+		if (NavData != nullptr)
 		{
 			bRaycastBlocked = NavData->Raycast(RayStart, RayEnd, HitLocation, UNavigationQueryFilter::GetQueryFilter(*NavData, Querier, FilterClass));
 		}
@@ -2358,7 +2323,7 @@ ANavigationData* UNavigationSystemV1::GetDefaultNavDataInstance(FNavigationSyste
 
 FSharedNavQueryFilter UNavigationSystemV1::CreateDefaultQueryFilterCopy() const 
 { 
-	return MainNavData ? MainNavData->GetDefaultQueryFilter()->GetCopy() : NULL; 
+	return MainNavData ? MainNavData->GetDefaultQueryFilter()->GetCopy() : nullptr; 
 }
 
 bool UNavigationSystemV1::IsNavigationBuilt(const AWorldSettings* Settings) const
@@ -2373,14 +2338,14 @@ bool UNavigationSystemV1::IsNavigationBuilt(const AWorldSettings* Settings) cons
 	for (int32 NavDataIndex = 0; NavDataIndex < NavDataSet.Num(); ++NavDataIndex)
 	{
 		ANavigationData* NavData = NavDataSet[NavDataIndex];
-		if (NavData != NULL && NavData->GetWorldSettings() == Settings)
+		if (NavData != nullptr && NavData->GetWorldSettings() == Settings)
 		{
 			FNavDataGenerator* Generator = NavData->GetGenerator();
 			if ((NavData->GetRuntimeGenerationMode() != ERuntimeGenerationType::Static
 #if WITH_EDITOR
-				|| GEditor != NULL
+				|| GEditor != nullptr
 #endif // WITH_EDITOR
-				) && (Generator == NULL || Generator->IsBuildInProgressCheckDirty() == true))
+				) && (Generator == nullptr || Generator->IsBuildInProgressCheckDirty() == true))
 			{
 				bIsBuilt = false;
 				break;
@@ -2592,7 +2557,7 @@ void UNavigationSystemV1::ProcessRegistrationCandidates()
 				NavDataPtr->Destroy();
 				if (NavDataPtr == MainNavData)
 				{
-					MainNavData = NULL;
+					MainNavData = nullptr;
 				}
 			}
 
@@ -2631,7 +2596,7 @@ UNavigationSystemV1::ERegistrationResult UNavigationSystemV1::RegisterNavData(AN
 {
 	UE_LOG(LogNavigation, Verbose, TEXT("%s %s"), ANSI_TO_TCHAR(__FUNCTION__), *GetFullNameSafe(NavData));
 	
-	if (NavData == NULL)
+	if (NavData == nullptr)
 	{
 		return RegistrationError;
 	}
@@ -2760,7 +2725,7 @@ void UNavigationSystemV1::UnregisterNavData(ANavigationData* NavData)
 	
 	NavDataSet.RemoveSingle(NavData);
 
-	if (NavData == NULL)
+	if (NavData == nullptr)
 	{
 		return;
 	}
@@ -2950,7 +2915,7 @@ void UNavigationSystemV1::RequestAreaRegistering(UClass* NavAreaClass)
 void UNavigationSystemV1::RegisterNavAreaClass(UClass* AreaClass)
 {
 	// can't be null
-	if (AreaClass == NULL)
+	if (AreaClass == nullptr)
 	{
 		return;
 	}
@@ -3148,14 +3113,14 @@ void UNavigationSystemV1::ResetCachedFilter(TSubclassOf<UNavigationQueryFilter>
 
 UNavigationSystemV1* UNavigationSystemV1::CreateNavigationSystem(UWorld* WorldOwner)
 {
-	UNavigationSystemV1* NavSys = NULL;
+	UNavigationSystemV1* NavSys = nullptr;
 
 	// create navigation system for editor and server targets, but remove it from game clients
 	if (WorldOwner && (*GEngine->NavigationSystemClass != nullptr) 
 		&& (GEngine->NavigationSystemClass->GetDefaultObject<UNavigationSystemV1>()->bAllowClientSideNavigation || WorldOwner->GetNetMode() != NM_Client))
 	{
 		AWorldSettings* WorldSettings = WorldOwner->GetWorldSettings();
-		if (WorldSettings == NULL || WorldSettings->IsNavigationSystemEnabled())
+		if (WorldSettings == nullptr || WorldSettings->IsNavigationSystemEnabled())
 		{
 			NavSys = NewObject<UNavigationSystemV1>(WorldOwner, GEngine->NavigationSystemClass);		
 			WorldOwner->SetNavigationSystem(NavSys);
@@ -3185,13 +3150,13 @@ ANavigationData* UNavigationSystemV1::GetNavDataWithID(const uint16 NavDataID) c
 	for (int32 NavDataIndex = 0; NavDataIndex < NavDataSet.Num(); ++NavDataIndex)
 	{
 		const ANavigationData* NavData = NavDataSet[NavDataIndex];
-		if (NavData != NULL && NavData->GetNavDataUniqueID() == NavDataID)
+		if (NavData != nullptr && NavData->GetNavDataUniqueID() == NavDataID)
 		{
 			return const_cast<ANavigationData*>(NavData);
 		}
 	}
 
-	return NULL;
+	return nullptr;
 }
 
 void UNavigationSystemV1::OnNavRelevantObjectRegistered(UObject& Object)
@@ -3685,7 +3650,7 @@ void UNavigationSystemV1::InitializeLevelCollisions()
 #if WITH_EDITOR
 void UNavigationSystemV1::UpdateLevelCollision(ULevel* InLevel)
 {
-	if (InLevel != NULL)
+	if (InLevel != nullptr)
 	{
 		UWorld* World = GetWorld();
 		OnLevelRemovedFromWorld(InLevel, World);
@@ -4255,7 +4220,7 @@ ANavigationData* UNavigationSystemV1::CreateNavigationDataInstanceInLevel(const
 
 	ANavigationData* Instance = World->SpawnActor<ANavigationData>(*NavConfig.GetNavDataClass<ANavigationData>(), SpawnInfo);
 
-	if (Instance != NULL)
+	if (Instance != nullptr)
 	{
 		Instance->SetConfig(NavConfig);
 		if (NavConfig.Name != NAME_None)
@@ -4263,8 +4228,8 @@ ANavigationData* UNavigationSystemV1::CreateNavigationDataInstanceInLevel(const
 			FString StrName = FString::Printf(TEXT("%s-%s"), *(Instance->GetFName().GetPlainNameString()), *(NavConfig.Name.ToString()));
 			// temporary solution to make sure we don't try to change name while there's already
 			// an object with this name
-			UObject* ExistingObject = StaticFindObject(/*Class=*/ NULL, Instance->GetOuter(), *StrName, true);
-			while (ExistingObject != NULL)
+			UObject* ExistingObject = StaticFindObject(/*Class=*/ nullptr, Instance->GetOuter(), *StrName, true);
+			while (ExistingObject != nullptr)
 			{
 				ANavigationData* ExistingNavigationData = Cast<ANavigationData>(ExistingObject);
 				if (ExistingNavigationData)
@@ -4273,14 +4238,14 @@ ANavigationData* UNavigationSystemV1::CreateNavigationDataInstanceInLevel(const
 				}
 
 				// Reset the existing object's name
-				ExistingObject->Rename(NULL, NULL, REN_DontCreateRedirectors | REN_ForceGlobalUnique | REN_DoNotDirty | REN_NonTransactional);
+				ExistingObject->Rename(nullptr, nullptr, REN_DontCreateRedirectors | REN_ForceGlobalUnique | REN_DoNotDirty | REN_NonTransactional);
 				// see if there's another one, it does happen when undo/redoing 
 				// nav instance deletion in the editor
-				ExistingObject = StaticFindObject(/*Class=*/ NULL, Instance->GetOuter(), *StrName, true);
+				ExistingObject = StaticFindObject(/*Class=*/ nullptr, Instance->GetOuter(), *StrName, true);
 			}
 
 			// Set descriptive name
-			Instance->Rename(*StrName, NULL, REN_DoNotDirty);
+			Instance->Rename(*StrName, nullptr, REN_DoNotDirty);
 #if WITH_EDITOR
 			if (World->WorldType == EWorldType::Editor)
 			{
@@ -4426,8 +4391,7 @@ bool UNavigationSystemV1::IsNavigationBuildInProgress()
 	for (int32 NavDataIndex = 0; NavDataIndex < NavDataSet.Num(); ++NavDataIndex)
 	{
 		ANavigationData* NavData = NavDataSet[NavDataIndex];
-		if (NavData != NULL && NavData->GetGenerator() != NULL
-			&& NavData->GetGenerator()->IsBuildInProgressCheckDirty() == true)
+		if (NavData != nullptr && NavData->GetGenerator() != nullptr && NavData->GetGenerator()->IsBuildInProgressCheckDirty() == true)
 		{
 			bRet = true;
 			break;
@@ -4682,7 +4646,7 @@ void UNavigationSystemV1::CleanUp(FNavigationSystem::ECleanupMode Mode)
 
 	DestroyNavOctree();
 	
-	SetCrowdManager(NULL);
+	SetCrowdManager(nullptr);
 
 	for (int32 Idx = NavDataSet.Num() - 1; Idx >= 0; --Idx)
 	{
@@ -4902,17 +4866,6 @@ bool UNavigationSystemV1::K2_GetRandomLocationInNavigableRadius(UObject* WorldCo
 	return bResult;
 }
 
-ENavigationQueryResult::Type UNavigationSystemV1::GetPathCost(UObject* WorldContextObject, const FVector& PathStart, const FVector& PathEnd, float& OutPathCost, ANavigationData* NavData, TSubclassOf<UNavigationQueryFilter> FilterClass)
-{
-	FVector::FReal PathCost = OutPathCost;
-
-	const ENavigationQueryResult::Type Result = GetPathCost(WorldContextObject, PathStart, PathEnd, PathCost, NavData, FilterClass);
-
-	OutPathCost = UE_REAL_TO_FLOAT_CLAMPED(PathCost);
-
-	return Result;
-}
-
 ENavigationQueryResult::Type UNavigationSystemV1::GetPathCost(UObject* WorldContextObject, const FVector& PathStart, const FVector& PathEnd, double& OutPathCost, ANavigationData* NavData, TSubclassOf<UNavigationQueryFilter> FilterClass)
 {
 	UWorld* World = GEngine->GetWorldFromContextObject(WorldContextObject, EGetWorldErrorMode::LogAndReturnNull);
@@ -4929,17 +4882,6 @@ ENavigationQueryResult::Type UNavigationSystemV1::GetPathCost(UObject* WorldCont
 	return ENavigationQueryResult::Error;
 }
 
-ENavigationQueryResult::Type UNavigationSystemV1::GetPathLength(UObject* WorldContextObject, const FVector& PathStart, const FVector& PathEnd, float& OutPathLength, ANavigationData* NavData, TSubclassOf<UNavigationQueryFilter> FilterClass)
-{
-	FVector::FReal PathLength = OutPathLength;
-
-	const ENavigationQueryResult::Type Result = GetPathLength(WorldContextObject, PathStart, PathEnd, PathLength, NavData, FilterClass);
-
-	OutPathLength = UE_REAL_TO_FLOAT_CLAMPED(PathLength);
-
-	return Result;
-}
-
 ENavigationQueryResult::Type UNavigationSystemV1::GetPathLength(UObject* WorldContextObject, const FVector& PathStart, const FVector& PathEnd, double& OutPathLength, ANavigationData* NavData, TSubclassOf<UNavigationQueryFilter> FilterClass)
 {
 	UWorld* World = GEngine->GetWorldFromContextObject(WorldContextObject, EGetWorldErrorMode::LogAndReturnNull);
@@ -4992,12 +4934,12 @@ bool UNavigationSystemV1::K2_ReplaceAreaInOctreeData(const UObject* Object, TSub
 //----------------------------------------------------------------------//
 bool UNavigationSystemV1::ShouldGeneratorRun(const FNavDataGenerator* Generator) const
 {
-	if (Generator != NULL && (IsNavigationSystemStatic() == false))
+	if (Generator != nullptr && (IsNavigationSystemStatic() == false))
 	{
 		for (int32 NavDataIndex = 0; NavDataIndex < NavDataSet.Num(); ++NavDataIndex)
 		{
 			ANavigationData* NavData = NavDataSet[NavDataIndex];
-			if (NavData != NULL && NavData->GetGenerator() == Generator)
+			if (NavData != nullptr && NavData->GetGenerator() == Generator)
 			{
 				return true;
 			}
@@ -5026,7 +4968,7 @@ bool UNavigationSystemV1::HandleCountNavMemCommand()
 	for (int32 NavDataIndex = 0; NavDataIndex < NavDataSet.Num(); ++NavDataIndex)
 	{
 		ANavigationData* NavData = NavDataSet[NavDataIndex];
-		if (NavData != NULL)
+		if (NavData != nullptr)
 		{
 			NavData->LogMemUsed();
 		}
@@ -5085,7 +5027,7 @@ void UNavigationSystemV1::CycleNavigationDataDrawn()
 	for (int32 NavDataIndex = 0; NavDataIndex < NavDataSet.Num(); ++NavDataIndex)
 	{
 		ANavigationData* NavData = NavDataSet[NavDataIndex];
-		if (NavData != NULL)
+		if (NavData != nullptr)
 		{
 			const bool bNewEnabledDrawing = (CurrentlyDrawnNavDataIndex == INDEX_NONE) || (NavDataIndex == CurrentlyDrawnNavDataIndex);
 			NavData->SetNavRenderingEnabled(bNewEnabledDrawing);
@@ -5142,12 +5084,12 @@ bool UNavigationSystemV1::CanRebuildDirtyNavigation() const
 
 bool UNavigationSystemV1::DoesPathIntersectBox(const FNavigationPath* Path, const FBox& Box, uint32 StartingIndex, FVector* AgentExtent)
 {
-	return Path != NULL && Path->DoesIntersectBox(Box, StartingIndex, NULL, AgentExtent);
+	return Path != nullptr && Path->DoesIntersectBox(Box, StartingIndex, nullptr, AgentExtent);
 }
 
 bool UNavigationSystemV1::DoesPathIntersectBox(const FNavigationPath* Path, const FBox& Box, const FVector& AgentLocation, uint32 StartingIndex, FVector* AgentExtent)
 {
-	return Path != NULL && Path->DoesIntersectBox(Box, AgentLocation, StartingIndex, NULL, AgentExtent);
+	return Path != nullptr && Path->DoesIntersectBox(Box, AgentLocation, StartingIndex, nullptr, AgentExtent);
 }
 
 void UNavigationSystemV1::SetMaxSimultaneousTileGenerationJobsCount(int32 MaxNumberOfJobs)
@@ -5242,20 +5184,6 @@ namespace UE::Navigation::Private
 	}
 }
 
-// Deprecated
-void UNavigationSystemV1::RegisterInvoker(AActor& Invoker, float TileGenerationRadius, float TileRemovalRadius)
-{
-PRAGMA_DISABLE_DEPRECATION_WARNINGS
-	RegisterInvoker(Invoker, TileGenerationRadius, TileRemovalRadius, FNavAgentSelector());
-PRAGMA_ENABLE_DEPRECATION_WARNINGS	
-}
-
-// Deprecated
-void UNavigationSystemV1::RegisterInvoker(AActor& Invoker, float TileGenerationRadius, float TileRemovalRadius, const FNavAgentSelector& Agents)
-{
-	RegisterInvoker(Invoker, TileGenerationRadius, TileRemovalRadius, Agents, ENavigationInvokerPriority::Default);
-}
-
 void UNavigationSystemV1::RegisterInvoker(AActor& Invoker, float TileGenerationRadius, float TileRemovalRadius, const FNavAgentSelector& Agents, ENavigationInvokerPriority InPriority)
 {
 	UE_CVLOG(bGenerateNavigationOnlyAroundNavigationInvokers == false, this, LogNavInvokers, Warning
diff --git a/Engine/Source/Runtime/NavigationSystem/Public/NavMesh/PImplRecastNavMesh.h b/Engine/Source/Runtime/NavigationSystem/Public/NavMesh/PImplRecastNavMesh.h
--- a/Engine/Source/Runtime/NavigationSystem/Public/NavMesh/PImplRecastNavMesh.h
+++ b/Engine/Source/Runtime/NavigationSystem/Public/NavMesh/PImplRecastNavMesh.h
@@ -100,17 +100,9 @@ public:
 
 	/** Supported queries */
 
-	/** Generates path from the given query. Synchronous. */
-	UE_DEPRECATED(5.2, "Please use FindPath with the added bRequireNavigableEndLocation parameter (true can be used as default).")
-	NAVIGATIONSYSTEM_API ENavigationQueryResult::Type FindPath(const FVector& StartLoc, const FVector& EndLoc, const FVector::FReal CostLimit, FNavMeshPath& Path, const FNavigationQueryFilter& Filter, const UObject* Owner) const;
-																																									  
 	/** Generates path from the given query. Synchronous. */
 	NAVIGATIONSYSTEM_API ENavigationQueryResult::Type FindPath(const FVector& StartLoc, const FVector& EndLoc, const FVector::FReal CostLimit, const bool bRequireNavigableEndLocation, FNavMeshPath& Path, const FNavigationQueryFilter& Filter, const UObject* Owner) const;
 
-	/** Check if path exists */
-	UE_DEPRECATED(5.2, "Please use TestPath with the added bRequireNavigableEndLocation parameter (true can be used as default).")
-	NAVIGATIONSYSTEM_API ENavigationQueryResult::Type TestPath(const FVector& StartLoc, const FVector& EndLoc, const FNavigationQueryFilter& Filter, const UObject* Owner, int32* NumVisitedNodes = 0) const;
-
 	/** Check if path exists */
 	NAVIGATIONSYSTEM_API ENavigationQueryResult::Type TestPath(const FVector& StartLoc, const FVector& EndLoc, const bool bRequireNavigableEndLocation, const FNavigationQueryFilter& Filter, const UObject* Owner, int32* NumVisitedNodes = 0) const;
 
@@ -121,10 +113,6 @@ public:
 	NAVIGATIONSYSTEM_API void Raycast(const FVector& StartLoc, const FVector& EndLoc, const FNavigationQueryFilter& InQueryFilter, const UObject* Owner,
 		ARecastNavMesh::FRaycastResult& RaycastResult, NavNodeRef StartNode = INVALID_NAVNODEREF) const;
 
-	/** Generates path from given query and collect data for every step of A* algorithm */
-	UE_DEPRECATED(5.2, "Please use DebugPathfinding with the added bRequireNavigableEndLocation parameter (true can be used as default).")
-	NAVIGATIONSYSTEM_API int32 DebugPathfinding(const FVector& StartLoc, const FVector& EndLoc, const FVector::FReal CostLimit, const FNavigationQueryFilter& Filter, const UObject* Owner, TArray<FRecastDebugPathfindingData>& Steps);
-
 	/** Generates path from given query and collect data for every step of A* algorithm */
 	NAVIGATIONSYSTEM_API int32 DebugPathfinding(const FVector& StartLoc, const FVector& EndLoc, const FVector::FReal CostLimit, const bool bRequireNavigableEndLocation, const FNavigationQueryFilter& Filter, const UObject* Owner, TArray<FRecastDebugPathfindingData>& Steps);
 
@@ -171,10 +159,6 @@ public:
 	/** finds stringpulled path from given corridor */
 	NAVIGATIONSYSTEM_API bool FindStraightPath(const FVector& StartLoc, const FVector& EndLoc, const TArray<NavNodeRef>& PathCorridor, TArray<FNavPathPoint>& PathPoints, TArray<FNavLinkId>* CustomLinks = NULL) const;
 
-	/** finds stringpulled path from given corridor */
-	UE_DEPRECATED(5.3, "Please use FindStraightPath with the TArray<FNavPathPoint>* CustomLinks. This function has no effect.")
-	bool FindStraightPath(const FVector& StartLoc, const FVector& EndLoc, const TArray<NavNodeRef>& PathCorridor, TArray<FNavPathPoint>& PathPoints, TArray<uint32>* CustomLinks) const { return false; }
-
 	/** Filters nav polys in PolyRefs with Filter */
 	NAVIGATIONSYSTEM_API bool FilterPolys(TArray<NavNodeRef>& PolyRefs, const FRecastQueryFilter* Filter, const UObject* Owner) const;
 
@@ -182,9 +166,6 @@ public:
 	UE_DEPRECATED(5.5, "Use the version of this function in ARecastNavMesh that takes a FNavTileRef instead.")
 	NAVIGATIONSYSTEM_API bool GetPolysInTile(int32 TileIndex, TArray<FNavPoly>& Polys) const;
 
-	UE_DEPRECATED(5.3, "Please use the version of this function that takes a FNavLinkId. This function has no effect.")
-	void UpdateNavigationLinkArea(int32 UserId, uint8 AreaType, uint16 PolyFlags) const {}
-
 	/** Updates area on polygons creating point-to-point connection with given UserId */
 	NAVIGATIONSYSTEM_API void UpdateNavigationLinkArea(FNavLinkId UserId, uint8 AreaType, uint16 PolyFlags) const;
 
diff --git a/Engine/Source/Runtime/NavigationSystem/Public/NavMesh/RecastHelpers.h b/Engine/Source/Runtime/NavigationSystem/Public/NavMesh/RecastHelpers.h
--- a/Engine/Source/Runtime/NavigationSystem/Public/NavMesh/RecastHelpers.h
+++ b/Engine/Source/Runtime/NavigationSystem/Public/NavMesh/RecastHelpers.h
@@ -23,20 +23,14 @@
 
 
 extern NAVIGATIONSYSTEM_API FVector Unreal2RecastPoint(const FVector::FReal* UnrealPoint);
-UE_DEPRECATED(5.0, "UnrealPoint should now be a FReal pointer!")
-extern NAVIGATIONSYSTEM_API FVector Unreal2RecastPoint(const float* UnrealPoint);
 extern NAVIGATIONSYSTEM_API FVector Unreal2RecastPoint(const FVector& UnrealPoint);
 extern NAVIGATIONSYSTEM_API FBox Unreal2RecastBox(const FBox& UnrealBox);
 extern NAVIGATIONSYSTEM_API FMatrix Unreal2RecastMatrix();
 
 extern NAVIGATIONSYSTEM_API FVector Recast2UnrealPoint(const FVector::FReal* RecastPoint);
-UE_DEPRECATED(5.0, "RecastPoint should now be a FReal pointer!")
-extern NAVIGATIONSYSTEM_API FVector Recast2UnrealPoint(const float* RecastPoint);
 extern NAVIGATIONSYSTEM_API FVector Recast2UnrealPoint(const FVector& RecastPoint);
 
 extern NAVIGATIONSYSTEM_API FBox Recast2UnrealBox(const FVector::FReal* RecastMin, const FVector::FReal* RecastMax);
-UE_DEPRECATED(5.0, "RecastMin and RecastMax should now be a FReal pointers!")
-extern NAVIGATIONSYSTEM_API FBox Recast2UnrealBox(const float* RecastMin, const float* RecastMax);
 extern NAVIGATIONSYSTEM_API FBox Recast2UnrealBox(const FBox& RecastBox);
 extern NAVIGATIONSYSTEM_API FMatrix Recast2UnrealMatrix();
 extern NAVIGATIONSYSTEM_API FColor Recast2UnrealColor(const unsigned int RecastColor);
diff --git a/Engine/Source/Runtime/NavigationSystem/Public/NavMesh/RecastNavMesh.h b/Engine/Source/Runtime/NavigationSystem/Public/NavMesh/RecastNavMesh.h
--- a/Engine/Source/Runtime/NavigationSystem/Public/NavMesh/RecastNavMesh.h
+++ b/Engine/Source/Runtime/NavigationSystem/Public/NavMesh/RecastNavMesh.h
@@ -777,16 +777,6 @@ class ARecastNavMesh : public ANavigationData
 	UPROPERTY(EditAnywhere, Category=Generation, config, meta=(ClampMin = "300.0"))
 	float TileSizeUU;
 
-	/** horizontal size of voxelization cell */
-	UE_DEPRECATED(5.2, "Set the CellSizes for the required navmesh resolutions in NavMeshResolutionParams.")
-	UPROPERTY(config)
-	float CellSize;
-
-	/** vertical size of voxelization cell */
-	UE_DEPRECATED(5.2, "Set the CellHeight for the required navmesh resolutions in NavMeshResolutionParams.")
-	UPROPERTY(config)
-	float CellHeight;
-
 	/** Resolution params 
 	 * If using multiple resolutions, it's recommended to chose the highest resolution first and 
 	 * set it according to the highest desired precision and then the other resolutions. */
@@ -805,11 +795,6 @@ class ARecastNavMesh : public ANavigationData
 	UPROPERTY(EditAnywhere, Category=Generation, config, meta=(ClampMin = "0.0", ClampMax = "89.0", UIMin = "0.0", UIMax = "89.0" ))
 	float AgentMaxSlope;
 
-	/** Largest vertical step the agent can perform */
-	UE_DEPRECATED(5.3, "Set the AgentMaxStepHeight for the required navmesh resolutions in NavMeshResolutionParams.")
-	UPROPERTY(config)
-	float AgentMaxStepHeight;
-
 	/* The minimum dimension of area. Areas smaller than this will be discarded */
 	UPROPERTY(EditAnywhere, Category=Generation, config, meta=(ClampMin = "0.0"))
 	float MinRegionArea;
@@ -1253,10 +1238,6 @@ public:
 	/** Creates a task to be executed on GameThread calling UpdateDrawing */
 	NAVIGATIONSYSTEM_API void RequestDrawingUpdate(bool bForce = false);
 
-	/** called after regenerating tiles */
-	UE_DEPRECATED(5.1, "Use new version with FNavTileRef")
-	NAVIGATIONSYSTEM_API virtual void OnNavMeshTilesUpdated(const TArray<uint32>& ChangedTiles);
-
 	/** called after regenerating tiles */
 	NAVIGATIONSYSTEM_API virtual void OnNavMeshTilesUpdated(const TArray<FNavTileRef>& ChangedTiles);
 
@@ -1281,9 +1262,6 @@ public:
 	//----------------------------------------------------------------------//
 	// Debug                                                                
 	//----------------------------------------------------------------------//
-	/** Debug rendering support. */
-	UE_DEPRECATED(5.1, "Please use the new signature of GetDebugGeometryForTile()")
-	NAVIGATIONSYSTEM_API void GetDebugGeometry(FRecastDebugGeometry& OutGeometry, int32 TileIndex = INDEX_NONE) const;
 
 	/* Gather debug geometry.
 	 * @params OutGeometry Output geometry.
@@ -1515,9 +1493,6 @@ public:
 	/** Check if poly is a custom link */
 	NAVIGATIONSYSTEM_API bool IsCustomLink(NavNodeRef PolyRef) const;
 
-	UE_DEPRECATED(5.3, "Use new override of this function with Array<FNavLinkId>* CustomLinks. This function has no effect.")
-	bool FindStraightPath(const FVector& StartLoc, const FVector& EndLoc, const TArray<NavNodeRef>& PathCorridor, TArray<FNavPathPoint>& PathPoints, TArray<uint32>* CustomLinks) const { return false; }
-
 	/** finds stringpulled path from given corridor */
 	NAVIGATIONSYSTEM_API bool FindStraightPath(const FVector& StartLoc, const FVector& EndLoc, const TArray<NavNodeRef>& PathCorridor, TArray<FNavPathPoint>& PathPoints, TArray<FNavLinkId>* CustomLinks = NULL) const;
 
@@ -1538,20 +1513,12 @@ public:
 	NAVIGATIONSYSTEM_API bool IsUsingActiveTilesGeneration(const UNavigationSystemV1& NavSys) const;
 
 	NAVIGATIONSYSTEM_API virtual void ConditionalConstructGenerator() override;
-
-PRAGMA_DISABLE_DEPRECATION_WARNINGS	
-	UE_DEPRECATED(5.2, "UpdateGenerationProperties is unused, it will be removed")
-	NAVIGATIONSYSTEM_API void UpdateGenerationProperties(const FRecastNavMeshGenerationProperties& GenerationProps);
-PRAGMA_ENABLE_DEPRECATION_WARNINGS	
 	
 	bool ShouldGatherDataOnGameThread() const { return bDoFullyAsyncNavDataGathering == false; }
 	int32 GetTileNumberHardLimit() const { return TileNumberHardLimit; }
 
 	NAVIGATIONSYSTEM_API virtual void UpdateActiveTiles(const TArray<FNavigationInvokerRaw>& InvokerLocations);
 	NAVIGATIONSYSTEM_API virtual void RemoveTiles(const TArray<FIntPoint>& Tiles);
-	
-	UE_DEPRECATED(5.3, "Use overload with FNavMeshDirtyTileElement instead.")
-	NAVIGATIONSYSTEM_API void RebuildTile(const TArray<FIntPoint>& Tiles);
 
 	NAVIGATIONSYSTEM_API void RebuildTile(const TArray<FNavMeshDirtyTileElement>& Tiles);
 	
@@ -1564,10 +1531,6 @@ PRAGMA_ENABLE_DEPRECATION_WARNINGS
 protected:
 
 	NAVIGATIONSYSTEM_API void UpdatePolyRefBitsPreview();
-	
-	/** Invalidates active paths that go through changed tiles  */
-	UE_DEPRECATED(5.1, "Use new version with FNavTileRef")
-	NAVIGATIONSYSTEM_API void InvalidateAffectedPaths(const TArray<uint32>& ChangedTiles);
 
 	/** Invalidates active paths that go through changed tiles  */
 	NAVIGATIONSYSTEM_API void InvalidateAffectedPaths(const TArray<FNavTileRef>& ChangedTiles);
diff --git a/Engine/Source/Runtime/NavigationSystem/Public/NavMesh/RecastNavMeshDataChunk.h b/Engine/Source/Runtime/NavigationSystem/Public/NavMesh/RecastNavMeshDataChunk.h
--- a/Engine/Source/Runtime/NavigationSystem/Public/NavMesh/RecastNavMeshDataChunk.h
+++ b/Engine/Source/Runtime/NavigationSystem/Public/NavMesh/RecastNavMeshDataChunk.h
@@ -70,22 +70,6 @@ class URecastNavMeshDataChunk : public UNavigationDataChunk
 	NAVIGATIONSYSTEM_API virtual void Serialize(FArchive& Ar) override;
 	//~ End UObject Interface
 
-	/** Attaches tiles to specified navmesh, transferring tile ownership to navmesh */
-	UE_DEPRECATED(5.1, "Use overload using ARecastNavMesh& instead")
-	NAVIGATIONSYSTEM_API TArray<uint32> AttachTiles(FPImplRecastNavMesh& NavMeshImpl);
-
-	/** Attaches tiles to specified navmesh */
-	UE_DEPRECATED(5.1, "Use overload using ARecastNavMesh& instead")
-	NAVIGATIONSYSTEM_API TArray<uint32> AttachTiles(FPImplRecastNavMesh& NavMeshImpl, const bool bKeepCopyOfData, const bool bKeepCopyOfCacheData);
-
-	/** Detaches tiles from specified navmesh, taking tile ownership */
-	UE_DEPRECATED(5.1, "Use overload using ARecastNavMesh& instead")
-	NAVIGATIONSYSTEM_API TArray<uint32> DetachTiles(FPImplRecastNavMesh& NavMeshImpl);
-
-	/** Detaches tiles from specified navmesh */
-	UE_DEPRECATED(5.1, "Use overload using ARecastNavMesh& instead")
-	NAVIGATIONSYSTEM_API TArray<uint32> DetachTiles(FPImplRecastNavMesh& NavMeshImpl, const bool bTakeDataOwnership, const bool bTakeCacheDataOwnership);
-
 #if WITH_RECAST
 	/** Attaches tiles to specified navmesh, transferring tile ownership to navmesh */
 	NAVIGATIONSYSTEM_API TArray<FNavTileRef> AttachTiles(ARecastNavMesh& NavMesh);
diff --git a/Engine/Source/Runtime/NavigationSystem/Public/NavMesh/RecastNavMeshGenerator.h b/Engine/Source/Runtime/NavigationSystem/Public/NavMesh/RecastNavMeshGenerator.h
--- a/Engine/Source/Runtime/NavigationSystem/Public/NavMesh/RecastNavMeshGenerator.h
+++ b/Engine/Source/Runtime/NavigationSystem/Public/NavMesh/RecastNavMeshGenerator.h
@@ -987,9 +987,6 @@ public:
 
 	NAVIGATIONSYSTEM_API void RemoveTiles(const TArray<FIntPoint>& Tiles);
 
-	UE_DEPRECATED(5.3, "Use overload with FNavMeshDirtyTileElement instead.")
-	NAVIGATIONSYSTEM_API void ReAddTiles(const TArray<FIntPoint>& Tiles);
-
 	NAVIGATIONSYSTEM_API void ReAddTiles(const TArray<FNavMeshDirtyTileElement>& Tiles);
 
 	bool IsBuildingRestrictedToActiveTiles() const { return bRestrictBuildingToActiveTiles; }
@@ -1012,13 +1009,6 @@ protected:
 
 	NAVIGATIONSYSTEM_API virtual TSharedRef<FRecastTileGenerator> CreateTileGenerator(const FIntPoint& Coord, const TArray<FBox>& DirtyAreas, const double PendingTileCreationTime = 0.);
 
-	template <typename T>
-	UE_DEPRECATED(5.3, "Use ConstructTileGeneratorImpl instead.")
-	TSharedRef<T> ConstuctTileGeneratorImpl(const FIntPoint& Coord, const TArray<FBox>& DirtyAreas, const double PendingTileCreationTime = 0.)
-	{
-		return ConstructTileGeneratorImpl<T>(Coord, DirtyAreas, PendingTileCreationTime);
-	}
-	
 	template <typename T>
 	TSharedRef<T> ConstructTileGeneratorImpl(const FIntPoint& Coord, const TArray<FBox>& DirtyAreas, const double PendingTileCreationTime)
 	{
@@ -1030,7 +1020,7 @@ protected:
 	UE_DEPRECATED(5.5, "Use BBoxGrowthLow and BBoxGrowthHigh properties instead.")
 	void SetBBoxGrowth(const FVector& InBBox)
 	{
-PRAGMA_DISABLE_DEPRECATION_WARNINGS		
+PRAGMA_DISABLE_DEPRECATION_WARNINGS
 		BBoxGrowth = InBBox;
 PRAGMA_ENABLE_DEPRECATION_WARNINGS
 	}
@@ -1100,7 +1090,7 @@ protected:
 	ARecastNavMesh*	DestNavMesh;
 	
 	/** List of dirty tiles that needs to be regenerated */
-	TNavStatArray<FPendingTileElement> PendingDirtyTiles;			
+	TNavStatArray<FPendingTileElement> PendingDirtyTiles;
 	
 	/** List of dirty tiles currently being regenerated */
 	TNavStatArray<TRunningTileElement<FRecastTileGeneratorWrapper>> RunningDirtyTiles;
@@ -1129,9 +1119,6 @@ protected:
 
 	uint32 bRestrictBuildingToActiveTiles:1;
 
-	UE_DEPRECATED(5.3, "Use SortPendingTilesMethod instead.")
-	uint32 bSortTilesWithSeedLocations:1;
-
 	/** Runtime generator's version, increased every time all tile generators get invalidated
 	 *	like when navmesh size changes */
 	uint32 Version;
diff --git a/Engine/Source/Runtime/NavigationSystem/Public/NavigationData.h b/Engine/Source/Runtime/NavigationSystem/Public/NavigationData.h
--- a/Engine/Source/Runtime/NavigationSystem/Public/NavigationData.h
+++ b/Engine/Source/Runtime/NavigationSystem/Public/NavigationData.h
@@ -234,8 +234,6 @@ struct FNavigationPath : public TSharedFromThis<FNavigationPath, ESPMode::Thread
 	  * This function will NOT reset setup variables like goal actor, filter, observer, etc */
 	NAVIGATIONSYSTEM_API virtual void ResetForRepath();
 
-	UE_DEPRECATED(5.0, "Use version that takes LifeTime instead.")
-	NAVIGATIONSYSTEM_API virtual void DebugDraw(const ANavigationData* NavData, FColor PathColor, class UCanvas* Canvas, bool bPersistent, const uint32 NextPathPointIndex = 0) const;
 	NAVIGATIONSYSTEM_API virtual void DebugDraw(const ANavigationData* NavData, const FColor PathColor, class UCanvas* Canvas, const bool bPersistent, const float LifeTime, const uint32 NextPathPointIndex = 0) const;
 	
 #if ENABLE_VISUAL_LOG
@@ -743,8 +741,6 @@ public:
 	//----------------------------------------------------------------------//
 	// Debug                                                                
 	//----------------------------------------------------------------------//
-	UE_DEPRECATED(5.0, "Use version that takes LifeTime instead.")
-	NAVIGATIONSYSTEM_API void DrawDebugPath(FNavigationPath* Path, const FColor PathColor, class UCanvas* Canvas, const bool bPersistent, const uint32 NextPathPointIndex) const;
 	NAVIGATIONSYSTEM_API void DrawDebugPath(FNavigationPath* Path, const FColor PathColor = FColor::White, class UCanvas* Canvas = nullptr, const bool bPersistent = true, const float LifeTime = -1.f, const uint32 NextPathPointIndex = 0) const;
 
 	FORCEINLINE bool IsDrawingEnabled() const { return bEnableDrawing; }
@@ -893,24 +889,16 @@ public:
 	 *	@note function should assert if item's FNavigationProjectionWork.ProjectionLimit is invalid */
 	NAVIGATIONSYSTEM_API virtual void BatchProjectPoints(TArray<FNavigationProjectionWork>& Workload, FSharedConstNavQueryFilter Filter = NULL, const UObject* Querier = NULL) const PURE_VIRTUAL(ANavigationData::BatchProjectPoints, );
 
-	UE_DEPRECATED(5.2, "Use new version with FVector::FReal")
-	NAVIGATIONSYSTEM_API virtual ENavigationQueryResult::Type CalcPathCost(const FVector& PathStart, const FVector& PathEnd, float& OutPathCost, FSharedConstNavQueryFilter QueryFilter = NULL, const UObject* Querier = NULL) const final;
-
 	/** Calculates path from PathStart to PathEnd and retrieves its cost.
  *	@NOTE this function does not generate string pulled path so the result is an (over-estimated) approximation
  *	@NOTE potentially expensive, so use it with caution */
 	NAVIGATIONSYSTEM_API virtual ENavigationQueryResult::Type CalcPathCost(const FVector& PathStart, const FVector& PathEnd, FVector::FReal& OutPathCost, FSharedConstNavQueryFilter QueryFilter = NULL, const UObject* Querier = NULL) const PURE_VIRTUAL(ANavigationData::CalcPathCost, return ENavigationQueryResult::Invalid;);
 
-	UE_DEPRECATED(5.2, "Use new version with FVector::FReal")
-	NAVIGATIONSYSTEM_API virtual ENavigationQueryResult::Type CalcPathLength(const FVector& PathStart, const FVector& PathEnd, float& OutPathLength, FSharedConstNavQueryFilter QueryFilter = NULL, const UObject* Querier = NULL) const final;
-
 	/** Calculates path from PathStart to PathEnd and retrieves its length.
 	 *	@NOTE this function does not generate string pulled path so the result is an (over-estimated) approximation
 	 *	@NOTE potentially expensive, so use it with caution */
 	NAVIGATIONSYSTEM_API virtual ENavigationQueryResult::Type CalcPathLength(const FVector& PathStart, const FVector& PathEnd, FVector::FReal& OutPathLength, FSharedConstNavQueryFilter QueryFilter = NULL, const UObject* Querier = NULL) const PURE_VIRTUAL(ANavigationData::CalcPathLength, return ENavigationQueryResult::Invalid;);
 
-	UE_DEPRECATED(5.2, "Use new version with FVector::FReal")
-	NAVIGATIONSYSTEM_API virtual ENavigationQueryResult::Type CalcPathLengthAndCost(const FVector& PathStart, const FVector& PathEnd, float& OutPathLength, float& OutPathCost, FSharedConstNavQueryFilter QueryFilter = NULL, const UObject* Querier = NULL) const final;
 	/** Calculates path from PathStart to PathEnd and retrieves its length.
 	 *	@NOTE this function does not generate string pulled path so the result is an (over-estimated) approximation
 	 *	@NOTE potentially expensive, so use it with caution */
diff --git a/Engine/Source/Runtime/NavigationSystem/Public/NavigationSystem.h b/Engine/Source/Runtime/NavigationSystem/Public/NavigationSystem.h
--- a/Engine/Source/Runtime/NavigationSystem/Public/NavigationSystem.h
+++ b/Engine/Source/Runtime/NavigationSystem/Public/NavigationSystem.h
@@ -370,13 +370,6 @@ public:
 	UPROPERTY(config, EditAnywhere, Category=NavigationSystem)
 	uint32 bSkipAgentHeightCheckWhenPickingNavData:1;
 
-#if WITH_EDITOR
-	/** Warnings are logged if exporting the navigation collision for an object exceed this vertex count.
-	 * Use -1 to disable. */
-	UE_DEPRECATED(5.2, "This property is deprecated. Please use GeometryExportTriangleCountWarningThreshold instead.")
-	int32 GeometryExportVertexCountWarningThreshold = 1000000;
-#endif // WITH_EDITOR
-
 	/** Warnings are logged if exporting the navigation collision for an object exceed this triangle count.
 	 * Use -1 to disable. */
 	UPROPERTY(config, EditAnywhere, AdvancedDisplay, Category = NavigationSystem)
@@ -502,19 +495,11 @@ public:
 	 *	@return Return Value represents if the call was successful */
 	UFUNCTION(BlueprintCallable, Category = "AI|Navigation", meta = (WorldContext = "WorldContextObject", DisplayName = "Get Random Location in Navigable Radius", ScriptName = "GetRandomLocationInNavigableRadius"))
 	static NAVIGATIONSYSTEM_API bool K2_GetRandomLocationInNavigableRadius(UObject* WorldContextObject, const FVector& Origin, FVector& RandomLocation, float Radius, ANavigationData* NavData = NULL, TSubclassOf<UNavigationQueryFilter> FilterClass = NULL);
-	
-	/** Potentially expensive. Use with caution. Consider using UPathFollowingComponent::GetRemainingPathCost instead */
-	UE_DEPRECATED(5.2, "Use new version with double")
-	static NAVIGATIONSYSTEM_API ENavigationQueryResult::Type GetPathCost(UObject* WorldContextObject, const FVector& PathStart, const FVector& PathEnd, float& PathCost, ANavigationData* NavData = NULL, TSubclassOf<UNavigationQueryFilter> FilterClass = NULL);
 
 	/** Potentially expensive. Use with caution. Consider using UPathFollowingComponent::GetRemainingPathCost instead */
 	UFUNCTION(BlueprintPure, Category = "AI|Navigation", meta = (WorldContext = "WorldContextObject"))
 	static NAVIGATIONSYSTEM_API ENavigationQueryResult::Type GetPathCost(UObject* WorldContextObject, const FVector& PathStart, const FVector& PathEnd, double& PathCost, ANavigationData* NavData = NULL, TSubclassOf<UNavigationQueryFilter> FilterClass = NULL);
 
-	/** Potentially expensive. Use with caution */
-	UE_DEPRECATED(5.2, "Use new version with double")
-	static NAVIGATIONSYSTEM_API ENavigationQueryResult::Type GetPathLength(UObject* WorldContextObject, const FVector& PathStart, const FVector& PathEnd, float& PathLength, ANavigationData* NavData = NULL, TSubclassOf<UNavigationQueryFilter> FilterClass = NULL);
-
 	/** Potentially expensive. Use with caution */
 	UFUNCTION(BlueprintPure, Category = "AI|Navigation", meta = (WorldContext = "WorldContextObject"))
 	static NAVIGATIONSYSTEM_API ENavigationQueryResult::Type GetPathLength(UObject* WorldContextObject, const FVector& PathStart, const FVector& PathEnd, double& PathLength, ANavigationData* NavData = NULL, TSubclassOf<UNavigationQueryFilter> FilterClass = NULL);
@@ -686,29 +671,14 @@ public:
 	 *	@return true if any location found, false otherwise */
 	NAVIGATIONSYSTEM_API bool GetRandomPointInNavigableRadius(const FVector& Origin, float Radius, FNavLocation& ResultLocation, ANavigationData* NavData = NULL, FSharedConstNavQueryFilter QueryFilter = NULL) const;
 	
-	/** Calculates a path from PathStart to PathEnd and retrieves its cost. 
-	 *	@NOTE potentially expensive, so use it with caution */
-	UE_DEPRECATED(5.2, "Use new version with FVector::FReal")
-	NAVIGATIONSYSTEM_API ENavigationQueryResult::Type GetPathCost(const FVector& PathStart, const FVector& PathEnd, float& PathCost, const ANavigationData* NavData = NULL, FSharedConstNavQueryFilter QueryFilter = NULL) const;
-
 	/** Calculates a path from PathStart to PathEnd and retrieves its cost.
 	 *	@NOTE potentially expensive, so use it with caution */
 	NAVIGATIONSYSTEM_API ENavigationQueryResult::Type GetPathCost(const FVector& PathStart, const FVector& PathEnd, FVector::FReal& PathCost, const ANavigationData* NavData = NULL, FSharedConstNavQueryFilter QueryFilter = NULL) const;
 
-	/** Calculates a path from PathStart to PathEnd and retrieves its overestimated length.
-	*	@NOTE potentially expensive, so use it with caution */
-	UE_DEPRECATED(5.2, "Use new version with FVector::FReal")
-	NAVIGATIONSYSTEM_API ENavigationQueryResult::Type GetPathLength(const FVector& PathStart, const FVector& PathEnd, float& PathLength, const ANavigationData* NavData = NULL, FSharedConstNavQueryFilter QueryFilter = NULL) const;
-
 	/** Calculates a path from PathStart to PathEnd and retrieves its overestimated length.
 	 *	@NOTE potentially expensive, so use it with caution */
 	NAVIGATIONSYSTEM_API ENavigationQueryResult::Type GetPathLength(const FVector& PathStart, const FVector& PathEnd, FVector::FReal& PathLength, const ANavigationData* NavData = NULL, FSharedConstNavQueryFilter QueryFilter = NULL) const;
 
-	/** Calculates a path from PathStart to PathEnd and retrieves its overestimated length and cost.
-	*	@NOTE potentially expensive, so use it with caution */
-	UE_DEPRECATED(5.2, "Use new version with FVector::FReal")
-	NAVIGATIONSYSTEM_API ENavigationQueryResult::Type GetPathLengthAndCost(const FVector& PathStart, const FVector& PathEnd, float& PathLength, float& PathCost, const ANavigationData* NavData = NULL, FSharedConstNavQueryFilter QueryFilter = NULL) const;
-
 	/** Calculates a path from PathStart to PathEnd and retrieves its overestimated length and cost.
 	 *	@NOTE potentially expensive, so use it with caution */
 	NAVIGATIONSYSTEM_API ENavigationQueryResult::Type GetPathLengthAndCost(const FVector& PathStart, const FVector& PathEnd, FVector::FReal& PathLength, FVector::FReal& PathCost, const ANavigationData* NavData = NULL, FSharedConstNavQueryFilter QueryFilter = NULL) const;
@@ -824,12 +794,6 @@ public:
 	//----------------------------------------------------------------------//
 	// Active tiles
 	//----------------------------------------------------------------------//
-	UE_DEPRECATED(5.2, "This function is deprecated. Please use the new RegisterInvoker method with the Agents parameter (FNavAgentSelector() can be used as default value to keep the same behavior)")
-	NAVIGATIONSYSTEM_API virtual void RegisterInvoker(AActor& Invoker, float TileGenerationRadius, float TileRemovalRadius);
-
-	UE_DEPRECATED(5.3, "This function is deprecated. Please use the new RegisterInvoker method with invoker priority.")
-	NAVIGATIONSYSTEM_API virtual void RegisterInvoker(AActor& Invoker, float TileGenerationRadius, float TileRemovalRadius, const FNavAgentSelector& Agents);
-
 	NAVIGATIONSYSTEM_API virtual void RegisterInvoker(AActor& Invoker, float TileGenerationRadius, float TileRemovalRadius, const FNavAgentSelector& Agents, ENavigationInvokerPriority InPriority);
 
 	NAVIGATIONSYSTEM_API virtual void RegisterInvoker(const TWeakInterfacePtr<INavigationInvokerInterface>& Invoker, float TileGenerationRadius, float TileRemovalRadius, const FNavAgentSelector& Agents, ENavigationInvokerPriority InPriority);
