Moved NavigationDirtyArea to a dedicated file to break some header dependencies

diff --git a/Engine/Source/Runtime/Engine/Classes/AI/Navigation/NavigationDirtyArea.h b/Engine/Source/Runtime/Engine/Classes/AI/Navigation/NavigationDirtyArea.h
new file mode 100644
--- /dev/null
+++ b/Engine/Source/Runtime/Engine/Classes/AI/Navigation/NavigationDirtyArea.h
@@ -0,0 +1,35 @@
+// Copyright Epic Games, Inc. All Rights Reserved.
+
+#pragma once
+
+#include "HAL/Platform.h"
+#include "Math/Box.h"
+#include "UObject/WeakObjectPtr.h"
+
+class UObject;
+namespace ENavigationDirtyFlag
+{
+	enum Type : uint8;
+}
+
+struct FNavigationDirtyArea
+{
+	FBox Bounds = FBox(ForceInit);
+	int32 Flags = 0;
+	TWeakObjectPtr<UObject> OptionalSourceObject;
+
+	FNavigationDirtyArea();
+	ENGINE_API FNavigationDirtyArea(const FBox& InBounds, int32 InFlags, UObject* const InOptionalSourceObject = nullptr);
+
+	FORCEINLINE bool HasFlag(ENavigationDirtyFlag::Type Flag) const { return (Flags & Flag) != 0; }
+
+	bool operator==(const FNavigationDirtyArea& Other) const
+	{ 
+		return Flags == Other.Flags && OptionalSourceObject == Other.OptionalSourceObject && Bounds.Equals(Other.Bounds);
+	}
+	
+	bool operator!=( const FNavigationDirtyArea& Other) const
+	{
+		return !(*this == Other);
+	}
+};
\ No newline at end of file
diff --git a/Engine/Source/Runtime/Engine/Classes/AI/Navigation/NavigationTypes.h b/Engine/Source/Runtime/Engine/Classes/AI/Navigation/NavigationTypes.h
--- a/Engine/Source/Runtime/Engine/Classes/AI/Navigation/NavigationTypes.h
+++ b/Engine/Source/Runtime/Engine/Classes/AI/Navigation/NavigationTypes.h
@@ -11,7 +11,10 @@
 #if UE_ENABLE_INCLUDE_ORDER_DEPRECATED_IN_5_4
 #include "GameFramework/Actor.h"
 #include "AI/Navigation/NavDataGatheringMode.h"
-#endif
+#endif // UE_ENABLE_INCLUDE_ORDER_DEPRECATED_IN_5_4
+#if UE_ENABLE_INCLUDE_ORDER_DEPRECATED_IN_5_5
+#include "NavigationDirtyArea.h"
+#endif //UE_ENABLE_INCLUDE_ORDER_DEPRECATED_IN_5_5
 #include "NavigationTypes.generated.h"
 
 #define INVALID_NAVQUERYID uint32(0)
@@ -79,7 +82,7 @@ namespace ENavigationOptionFlag
 
 namespace ENavigationDirtyFlag
 {
-	enum Type
+	enum Type : uint8
 	{
 		Geometry			= (1 << 0),
 		DynamicModifier		= (1 << 1),
@@ -90,27 +93,6 @@ namespace ENavigationDirtyFlag
 	};
 }
 
-struct FNavigationDirtyArea
-{
-	FBox Bounds;
-	int32 Flags;
-	TWeakObjectPtr<UObject> OptionalSourceObject;
-	
-	FNavigationDirtyArea() : Flags(0) {}
-	ENGINE_API FNavigationDirtyArea(const FBox& InBounds, int32 InFlags, UObject* const InOptionalSourceObject = nullptr);
-	FORCEINLINE bool HasFlag(ENavigationDirtyFlag::Type Flag) const { return (Flags & Flag) != 0; }
-
-	bool operator==(const FNavigationDirtyArea& Other) const 
-	{ 
-		return Flags == Other.Flags && OptionalSourceObject == Other.OptionalSourceObject && Bounds.Equals(Other.Bounds); 
-	}
-	
-	bool operator!=( const FNavigationDirtyArea& Other) const
-	{
-		return !(*this == Other);
-	}
-};
-
 UENUM()
 enum class ENavDataGatheringModeConfig : uint8
 {
diff --git a/Engine/Source/Runtime/Engine/Private/AI/Navigation/NavigationDirtyArea.cpp b/Engine/Source/Runtime/Engine/Private/AI/Navigation/NavigationDirtyArea.cpp
new file mode 100644
--- /dev/null
+++ b/Engine/Source/Runtime/Engine/Private/AI/Navigation/NavigationDirtyArea.cpp
@@ -0,0 +1,20 @@
+// Copyright Epic Games, Inc. All Rights Reserved.
+
+#include "AI/Navigation/NavigationDirtyArea.h"
+#if !NO_LOGGING
+#include "AI/NavigationSystemBase.h"
+#endif
+
+FNavigationDirtyArea::FNavigationDirtyArea(const FBox& InBounds, const int32 InFlags, UObject* const InOptionalSourceObject)
+	: Bounds(InBounds)
+	, Flags(InFlags)
+	, OptionalSourceObject(InOptionalSourceObject)
+{
+#if !NO_LOGGING
+	if (!Bounds.IsValid || Bounds.ContainsNaN())
+	{
+		UE_LOG(LogNavigation, Warning, TEXT("Creation of FNavigationDirtyArea with invalid bounds%s. Bounds: %s, SourceObject: %s."),
+			Bounds.ContainsNaN() ? TEXT(" (contains NaN)") : TEXT(""), *Bounds.ToString(), *GetFullNameSafe(OptionalSourceObject.Get()));
+	}
+#endif //!NO_LOGGING
+}
diff --git a/Engine/Source/Runtime/Engine/Private/AI/Navigation/NavigationTypes.cpp b/Engine/Source/Runtime/Engine/Private/AI/Navigation/NavigationTypes.cpp
--- a/Engine/Source/Runtime/Engine/Private/AI/Navigation/NavigationTypes.cpp
+++ b/Engine/Source/Runtime/Engine/Private/AI/Navigation/NavigationTypes.cpp
@@ -15,7 +15,7 @@
 
 DEFINE_STAT(STAT_Navigation_MetaAreaTranslation);
 
-static const uint32 MAX_NAV_SEARCH_NODES = 2048;
+static constexpr uint32 MAX_NAV_SEARCH_NODES = 2048;
 
 namespace FNavigationSystem
 {
@@ -76,20 +76,6 @@ namespace FNavigationSystem
 	}	
 }
 
-FNavigationDirtyArea::FNavigationDirtyArea(const FBox& InBounds, int32 InFlags, UObject* const InOptionalSourceObject /*= nullptr*/)
-	: Bounds(InBounds)
-	, Flags(InFlags)
-	, OptionalSourceObject(InOptionalSourceObject)
-{
-#if !NO_LOGGING
-	if (!Bounds.IsValid || Bounds.ContainsNaN())
-	{
-		UE_LOG(LogNavigation, Warning, TEXT("Creation of FNavigationDirtyArea with invalid bounds%s. Bounds: %s, SourceObject: %s."),
-			Bounds.ContainsNaN() ? TEXT(" (contains NaN)") : TEXT(""), *Bounds.ToString(), *GetFullNameSafe(OptionalSourceObject.Get()));
-	}
-#endif //!NO_LOGGING	
-}
-
 //----------------------------------------------------------------------//
 // FNavigationQueryFilter
 //----------------------------------------------------------------------//
@@ -141,12 +127,13 @@ FString FNavDataConfig::GetDescription() const
 	return FString::Printf(TEXT("Name %s class %s agent radius %.1f")
 		, *Name.ToString(), *NavDataClass.ToString(), AgentRadius);
 }
+
 //----------------------------------------------------------------------//
 // FNavigationRelevantData
 //----------------------------------------------------------------------//
 bool FNavigationRelevantData::FCollisionDataHeader::IsValid(const uint8* RawData, int32 RawDataSize)
 {
-	const int32 HeaderSize = sizeof(FCollisionDataHeader);
+	constexpr int32 HeaderSize = sizeof(FCollisionDataHeader);
 	return (RawDataSize == 0) || ((RawDataSize >= HeaderSize) && (((const FCollisionDataHeader*)RawData)->DataSize == RawDataSize));
 }
 
@@ -197,7 +184,7 @@ FNavigationQueryFilter::FNavigationQueryFilter(const FNavigationQueryFilter& Sou
 FNavigationQueryFilter::FNavigationQueryFilter(const FNavigationQueryFilter* Source)
 	: MaxSearchNodes(DefaultMaxSearchNodes)
 {
-	if (Source != NULL)
+	if (Source != nullptr)
 	{
 		Assign(*Source);
 	}
@@ -220,7 +207,7 @@ FNavigationQueryFilter& FNavigationQueryFilter::operator=(const FNavigationQuery
 
 void FNavigationQueryFilter::Assign(const FNavigationQueryFilter& Source)
 {
-	if (Source.GetImplementation() != NULL)
+	if (Source.GetImplementation() != nullptr)
 	{
 		QueryFilterImpl = Source.QueryFilterImpl;
 	}
@@ -361,7 +348,7 @@ FNavLinkId FNavLinkId::GenerateUniqueId(FNavLinkAuxiliaryId AuxiliaryId, FGuid A
 	// Apply NavLinkIdBitMask to differentiate Legacy Ids (that do not have the mask set).
 	const uint64 UniqueId = UE::Navigation::NavLinkIdHelpers::Private::MakeIdFromGUID(AuxiliaryId, ActorInstanceGuid) | NavLinkIdBitMask;
 
-	UE_LOG(LogNavLink, VeryVerbose, TEXT("%hs id: %u."), __FUNCTION__, UniqueId);
+	UE_LOG(LogNavLink, VeryVerbose, TEXT("%hs id: %llu."), __FUNCTION__, UniqueId);
 
 	return FNavLinkId(UniqueId);
 }
diff --git a/Engine/Source/Runtime/Engine/Public/AI/NavDataGenerator.h b/Engine/Source/Runtime/Engine/Public/AI/NavDataGenerator.h
--- a/Engine/Source/Runtime/Engine/Public/AI/NavDataGenerator.h
+++ b/Engine/Source/Runtime/Engine/Public/AI/NavDataGenerator.h
@@ -1,10 +1,16 @@
 // Copyright Epic Games, Inc. All Rights Reserved.
 
-#pragma once 
+#pragma once
 
+#if UE_ENABLE_INCLUDE_ORDER_DEPRECATED_IN_5_5
 #include "CoreMinimal.h"
 #include "AI/Navigation/NavigationTypes.h"
+#endif // UE_ENABLE_INCLUDE_ORDER_DEPRECATED_IN_5_5
+
 #include "EngineDefines.h"
+#include "Templates/SharedPointer.h"
+
+struct FNavigationDirtyArea;
 
 class FNavDataGenerator : public TSharedFromThis<FNavDataGenerator, ESPMode::ThreadSafe>
 {
diff --git a/Engine/Source/Runtime/NavigationSystem/Private/NavigationData.cpp b/Engine/Source/Runtime/NavigationSystem/Private/NavigationData.cpp
--- a/Engine/Source/Runtime/NavigationSystem/Private/NavigationData.cpp
+++ b/Engine/Source/Runtime/NavigationSystem/Private/NavigationData.cpp
@@ -2,16 +2,16 @@
 
 #include "NavigationData.h"
 
-#include "AssetCompilingManager.h"
-#include "EngineGlobals.h"
+#include "AI/NavDataGenerator.h"
 #include "AI/Navigation/NavAgentInterface.h"
+#include "AI/Navigation/NavAreaBase.h"
+#include "AI/Navigation/NavigationDirtyArea.h"
+#include "AssetCompilingManager.h"
 #include "Components/PrimitiveComponent.h"
-#include "AI/NavDataGenerator.h"
-#include "NavigationSystem.h"
-#include "NavFilters/NavigationQueryFilter.h"
 #include "Engine/Engine.h"
 #include "NavAreas/NavArea.h"
-#include "AI/Navigation/NavAreaBase.h"
+#include "NavFilters/NavigationQueryFilter.h"
+#include "NavigationSystem.h"
 #include "VisualLogger/VisualLogger.h"
 
 #include UE_INLINE_GENERATED_CPP_BY_NAME(NavigationData)
diff --git a/Engine/Source/Runtime/NavigationSystem/Private/NavigationDirtyAreasController.cpp b/Engine/Source/Runtime/NavigationSystem/Private/NavigationDirtyAreasController.cpp
--- a/Engine/Source/Runtime/NavigationSystem/Private/NavigationDirtyAreasController.cpp
+++ b/Engine/Source/Runtime/NavigationSystem/Private/NavigationDirtyAreasController.cpp
@@ -1,10 +1,11 @@
 // Copyright Epic Games, Inc. All Rights Reserved.
 
 #include "NavigationDirtyAreasController.h"
+#include "AI/Navigation/NavigationDirtyArea.h"
+#include "AI/Navigation/NavigationDirtyElement.h"
 #include "NavigationData.h"
 #include "NavigationSystem.h"
 #include "VisualLogger/VisualLogger.h"
-#include "AI/Navigation/NavigationDirtyElement.h"
 
 DEFINE_LOG_CATEGORY(LogNavigationDirtyArea);
 
diff --git a/Engine/Source/Runtime/NavigationSystem/Private/NavigationSystem.cpp b/Engine/Source/Runtime/NavigationSystem/Private/NavigationSystem.cpp
--- a/Engine/Source/Runtime/NavigationSystem/Private/NavigationSystem.cpp
+++ b/Engine/Source/Runtime/NavigationSystem/Private/NavigationSystem.cpp
@@ -1,59 +1,60 @@
 // Copyright Epic Games, Inc. All Rights Reserved.
 
 #include "NavigationSystem.h"
-#include "NavigationDataHandler.h"
-#include "Misc/ScopeLock.h"
-#include "Stats/StatsMisc.h"
-#include "Modules/ModuleManager.h"
+#include "AbstractNavData.h"
+#include "AI/NavDataGenerator.h"
 #include "AI/Navigation/NavAgentInterface.h"
-#include "Engine/LocalPlayer.h"
-#include "Engine/World.h"
-#include "GameFramework/Controller.h"
-#include "AI/Navigation/NavRelevantInterface.h"
+#include "AI/Navigation/NavigationDataChunk.h"
+#include "AI/Navigation/NavigationDirtyArea.h"
 #include "AI/Navigation/NavigationDirtyElement.h"
 #include "AI/Navigation/NavigationInvokerInterface.h"
 #include "AI/Navigation/NavigationInvokerPriority.h"
-#include "NavFilters/NavigationQueryFilter.h"
-#include "UObject/UObjectIterator.h"
+#include "AI/Navigation/NavRelevantInterface.h"
+#include "AI/NavigationModifier.h"
+#include "Components/PrimitiveComponent.h"
+#include "CrowdManagerBase.h"
+#include "Engine/Engine.h"
+#include "Engine/LocalPlayer.h"
+#include "Engine/World.h"
 #include "EngineUtils.h"
+#include "GameFramework/Controller.h"
+#include "GameFramework/Pawn.h"
 #include "Logging/MessageLog.h"
-#include "ProfilingDebugging/CsvProfiler.h"
+#include "Misc/ScopeLock.h"
+#include "Modules/ModuleManager.h"
 #include "NavAreas/NavArea.h"
-#include "NavigationOctree.h"
-#include "VisualLogger/VisualLogger.h"
-#include "NavMesh/NavMeshBoundsVolume.h"
+#include "NavAreas/NavArea_Default.h"
+#include "NavAreas/NavArea_Obstacle.h"
+#include "NavAreas/NavAreaMeta_SwitchByAgent.h"
+#include "NavFilters/NavigationQueryFilter.h"
+#include "NavigationDataHandler.h"
 #include "NavigationInvokerComponent.h"
 #include "NavigationObjectRepository.h"
-#include "AI/Navigation/NavigationDataChunk.h"
-#include "Engine/Engine.h"
+#include "NavigationOctree.h"
+#include "NavigationPath.h"
+#include "NavLinkCustomInterface.h"
+#include "NavMesh/NavMeshBoundsVolume.h"
+#include "ProfilingDebugging/CsvProfiler.h"
+#include "Stats/StatsMisc.h"
 #include "UObject/Package.h"
-#include "Components/PrimitiveComponent.h"
+#include "UObject/UObjectIterator.h"
 #include "UObject/UObjectThreadContext.h"
-#include "GameFramework/Pawn.h"
-#include "AI/NavDataGenerator.h"
+#include "VisualLogger/VisualLogger.h"
 
 #if WITH_RECAST
-#include "NavMesh/RecastNavMesh.h"
-#include "NavMesh/RecastHelpers.h"
 #include "NavMesh/RecastGeometryExport.h"
+#include "NavMesh/RecastHelpers.h"
+#include "NavMesh/RecastNavMesh.h"
 #endif // WITH_RECAST
+
 #if WITH_EDITOR
-#include "EditorModeManager.h"
-#include "EditorModes.h"
-#include "LevelEditor.h"
-#include "Misc/MessageDialog.h"
 #include "WorldPartition/WorldPartition.h"
+#include "Misc/MessageDialog.h"
+#include "LevelEditor.h"
+#include "EditorModes.h"
+#include "EditorModeManager.h"
 #endif
 
-#include "NavAreas/NavArea_Null.h"
-#include "NavAreas/NavArea_Obstacle.h"
-#include "NavAreas/NavArea_Default.h"
-#include "NavAreas/NavAreaMeta_SwitchByAgent.h"
-#include "NavLinkCustomInterface.h"
-#include "NavigationPath.h"
-#include "AbstractNavData.h"
-#include "CrowdManagerBase.h"
-#include "AI/NavigationModifier.h"
 
 #include UE_INLINE_GENERATED_CPP_BY_NAME(NavigationSystem)
 
diff --git a/Engine/Source/Runtime/NavigationSystem/Public/NavMesh/RecastNavMeshGenerator.h b/Engine/Source/Runtime/NavigationSystem/Public/NavMesh/RecastNavMeshGenerator.h
--- a/Engine/Source/Runtime/NavigationSystem/Public/NavMesh/RecastNavMeshGenerator.h
+++ b/Engine/Source/Runtime/NavigationSystem/Public/NavMesh/RecastNavMeshGenerator.h
@@ -3,21 +3,24 @@
 
 #if UE_ENABLE_INCLUDE_ORDER_DEPRECATED_IN_5_4
 #include "CoreMinimal.h"
-#endif
+#endif // UE_ENABLE_INCLUDE_ORDER_DEPRECATED_IN_5_4
+#if UE_ENABLE_INCLUDE_ORDER_DEPRECATED_IN_5_5
 #include "Stats/Stats.h"
-#include "AI/NavigationSystemBase.h"
-#include "AI/Navigation/NavigationTypes.h"
+#endif // UE_ENABLE_INCLUDE_ORDER_DEPRECATED_IN_5_5
+#include "AI/NavDataGenerator.h"
+#include "AI/Navigation/NavigationDirtyArea.h"
 #include "AI/Navigation/NavigationInvokerPriority.h"
 #include "AI/Navigation/NavigationRelevantData.h"
-#include "EngineDefines.h"
+#include "AI/Navigation/NavigationTypes.h"
 #include "AI/NavigationModifier.h"
+#include "AI/NavigationSystemBase.h"
+#include "Async/AsyncWork.h"
+#include "EngineDefines.h"
+#include "NavDebugTypes.h"
 #include "NavigationOctree.h"
+#include "NavMesh/RecastHelpers.h"
 #include "NavMesh/RecastNavMesh.h"
-#include "Async/AsyncWork.h"
 #include "UObject/GCObject.h"
-#include "AI/NavDataGenerator.h"
-#include "NavMesh/RecastHelpers.h"
-#include "NavDebugTypes.h"
 
 #if WITH_RECAST
 
diff --git a/Engine/Source/Runtime/NavigationSystem/Public/NavigationData.h b/Engine/Source/Runtime/NavigationSystem/Public/NavigationData.h
--- a/Engine/Source/Runtime/NavigationSystem/Public/NavigationData.h
+++ b/Engine/Source/Runtime/NavigationSystem/Public/NavigationData.h
@@ -28,6 +28,7 @@ class INavLinkCustomInterface;
 class UNavArea;
 class UPrimitiveComponent;
 class UNavigationQueryFilter;
+struct FNavigationDirtyArea;
 
 USTRUCT()
 struct FSupportedAreaData
diff --git a/Engine/Source/Runtime/NavigationSystem/Public/NavigationDirtyAreasController.h b/Engine/Source/Runtime/NavigationSystem/Public/NavigationDirtyAreasController.h
--- a/Engine/Source/Runtime/NavigationSystem/Public/NavigationDirtyAreasController.h
+++ b/Engine/Source/Runtime/NavigationSystem/Public/NavigationDirtyAreasController.h
@@ -11,7 +11,7 @@
 #include "UObject/UnrealNames.h"
 
 class UObject;
-
+struct FNavigationDirtyArea;
 
 NAVIGATIONSYSTEM_API DECLARE_LOG_CATEGORY_EXTERN(LogNavigationDirtyArea, Warning, All);
 
