[AINavigation] Converted ENavigationOptionFlag and ENavigationDirtyFlag to enum classes

diff --git a/Engine/Source/Runtime/Engine/Classes/AI/Navigation/NavigationDirtyArea.h b/Engine/Source/Runtime/Engine/Classes/AI/Navigation/NavigationDirtyArea.h
--- a/Engine/Source/Runtime/Engine/Classes/AI/Navigation/NavigationDirtyArea.h
+++ b/Engine/Source/Runtime/Engine/Classes/AI/Navigation/NavigationDirtyArea.h
@@ -7,21 +7,35 @@
 #include "UObject/WeakObjectPtr.h"
 
 class UObject;
-namespace ENavigationDirtyFlag
+
+enum class ENavigationDirtyFlag : uint8
 {
-	enum Type : uint8;
-}
+	None				= 0,
+	Geometry			= (1 << 0),
+	DynamicModifier		= (1 << 1),
+	UseAgentHeight		= (1 << 2),
+	NavigationBounds	= (1 << 3),
+
+	All = Geometry | DynamicModifier, // all rebuild steps here without additional flags
+};
+ENUM_CLASS_FLAGS(ENavigationDirtyFlag);
 
 struct FNavigationDirtyArea
 {
 	FBox Bounds = FBox(ForceInit);
-	int32 Flags = 0;
 	TWeakObjectPtr<UObject> OptionalSourceObject;
+	ENavigationDirtyFlag Flags = ENavigationDirtyFlag::None;
+
+	FNavigationDirtyArea() = default;
+	ENGINE_API FNavigationDirtyArea(const FBox& InBounds, ENavigationDirtyFlag InFlags, UObject* const InOptionalSourceObject = nullptr);
 
-	FNavigationDirtyArea();
+	UE_DEPRECATED(5.5, "Use constructor taking ENavigationDirtyFlag instead.")
 	ENGINE_API FNavigationDirtyArea(const FBox& InBounds, int32 InFlags, UObject* const InOptionalSourceObject = nullptr);
 
-	FORCEINLINE bool HasFlag(ENavigationDirtyFlag::Type Flag) const { return (Flags & Flag) != 0; }
+	bool HasFlag(const ENavigationDirtyFlag Flag) const
+	{
+		return (Flags & Flag) != ENavigationDirtyFlag::None;
+	}
 
 	bool operator==(const FNavigationDirtyArea& Other) const
 	{ 
@@ -32,4 +46,4 @@ struct FNavigationDirtyArea
 	{
 		return !(*this == Other);
 	}
-};
\ No newline at end of file
+};
diff --git a/Engine/Source/Runtime/Engine/Classes/AI/Navigation/NavigationDirtyElement.h b/Engine/Source/Runtime/Engine/Classes/AI/Navigation/NavigationDirtyElement.h
--- a/Engine/Source/Runtime/Engine/Classes/AI/Navigation/NavigationDirtyElement.h
+++ b/Engine/Source/Runtime/Engine/Classes/AI/Navigation/NavigationDirtyElement.h
@@ -5,6 +5,7 @@
 #include "Math/Box.h"
 #include "UObject/WeakObjectPtr.h"
 
+enum class ENavigationDirtyFlag : uint8;
 class INavRelevantInterface;
 
 struct FNavigationDirtyElement
@@ -18,41 +19,37 @@ struct FNavigationDirtyElement
 
 	/** object owning this element */
 	FWeakObjectPtr Owner;
-	
+
 	/** cached interface pointer */
-	INavRelevantInterface* NavInterface;
+	INavRelevantInterface* NavInterface = nullptr;
+
+	/** bounds of already existing entry for this actor */
+	FBox PrevBounds = FBox(ForceInit);
 
 	/** override for update flags */
-	int32 FlagsOverride;
-	
+	ENavigationDirtyFlag FlagsOverride;
+
 	/** flags of already existing entry for this actor */
-	int32 PrevFlags;
-	
-	/** bounds of already existing entry for this actor */
-	FBox PrevBounds;
+	ENavigationDirtyFlag PrevFlags;
 
 	/** prev flags & bounds data are set */
-	uint8 bHasPrevData : 1;
+	uint8 bHasPrevData : 1 = false;
 
 	/** request was invalidated while queued, use prev values to dirty area */
-	uint8 bInvalidRequest : 1;
+	uint8 bInvalidRequest : 1 = false;
 
 	/** requested during visibility change of the owning level (loading/unloading) */
-	uint8 bIsFromVisibilityChange : 1;
+	uint8 bIsFromVisibilityChange : 1 = false;
 
 	/** part of the base navmesh */
-	uint8 bIsInBaseNavmesh : 1;
+	uint8 bIsInBaseNavmesh : 1 = false;
 
-	FNavigationDirtyElement()
-		: NavInterface(nullptr), FlagsOverride(0), PrevFlags(0), PrevBounds(ForceInit), bHasPrevData(false), bInvalidRequest(false), bIsFromVisibilityChange(false), bIsInBaseNavmesh(false)
-	{
-	}
-
-	FNavigationDirtyElement(UObject* InOwner)
-		: Owner(InOwner), NavInterface(nullptr), FlagsOverride(0), PrevFlags(0), PrevBounds(ForceInit), bHasPrevData(false), bInvalidRequest(false), bIsFromVisibilityChange(false), bIsInBaseNavmesh(false)
-	{
-	}
+	ENGINE_API FNavigationDirtyElement();
+	ENGINE_API explicit FNavigationDirtyElement(UObject* InOwner);
+	ENGINE_API FNavigationDirtyElement(UObject* InOwner, INavRelevantInterface* InNavInterface, ENavigationDirtyFlag InFlagsOverride, const bool bUseWorldPartitionedDynamicMode = false);
+	ENGINE_API FNavigationDirtyElement(UObject* InOwner, INavRelevantInterface* InNavInterface, const bool bUseWorldPartitionedDynamicMode = false);
 
+	UE_DEPRECATED(5.5, "Use the version taking ENavigationDirtyFlag instead.")
 	ENGINE_API FNavigationDirtyElement(UObject* InOwner, INavRelevantInterface* InNavInterface, int32 InFlagsOverride = 0, const bool bUseWorldPartitionedDynamicMode = false);
 
 	bool operator==(const FNavigationDirtyElement& Other) const 
diff --git a/Engine/Source/Runtime/Engine/Classes/AI/Navigation/NavigationRelevantData.h b/Engine/Source/Runtime/Engine/Classes/AI/Navigation/NavigationRelevantData.h
--- a/Engine/Source/Runtime/Engine/Classes/AI/Navigation/NavigationRelevantData.h
+++ b/Engine/Source/Runtime/Engine/Classes/AI/Navigation/NavigationRelevantData.h
@@ -2,9 +2,12 @@
 
 #pragma once
 
-#include "UObject/ObjectMacros.h"
+#include "AI/Navigation/NavigationDirtyArea.h"
 #include "AI/Navigation/NavigationTypes.h"
 #include "AI/NavigationModifier.h"
+#if UE_ENABLE_INCLUDE_ORDER_DEPRECATED_IN_5_5
+#include "UObject/ObjectMacros.h"
+#endif //UE_ENABLE_INCLUDE_ORDER_DEPRECATED_IN_5_5
 
 struct FNavigationRelevantDataFilter 
 {
@@ -109,15 +112,21 @@ struct FNavigationRelevantData : public TSharedFromThis<FNavigationRelevantData,
 	FORCEINLINE bool IsEmpty() const { return !HasGeometry() && !HasModifiers(); }
 	FORCEINLINE SIZE_T GetAllocatedSize() const { return CollisionData.GetAllocatedSize() + VoxelData.GetAllocatedSize() + Modifiers.GetAllocatedSize(); }
 	FORCEINLINE SIZE_T GetGeometryAllocatedSize() const { return CollisionData.GetAllocatedSize() + VoxelData.GetAllocatedSize(); }
-	FORCEINLINE int32 GetDirtyFlag() const
+	FORCEINLINE ENavigationDirtyFlag GetDirtyFlag() const
 	{
 		const bool bSetGeometryFlag = HasGeometry() || IsPendingLazyGeometryGathering() ||
 			Modifiers.GetFillCollisionUnderneathForNavmesh() || Modifiers.GetMaskFillCollisionUnderneathForNavmesh() ||
 			(Modifiers.GetNavMeshResolution() != ENavigationDataResolution::Invalid);
 		
-		return (bSetGeometryFlag ? ENavigationDirtyFlag::Geometry : 0) |
-			((HasDynamicModifiers() || NeedAnyPendingLazyModifiersGathering()) ? ENavigationDirtyFlag::DynamicModifier : 0) |
-			(Modifiers.HasAgentHeightAdjust() ? ENavigationDirtyFlag::UseAgentHeight : 0);
+		return (bSetGeometryFlag
+				? ENavigationDirtyFlag::Geometry
+				: ENavigationDirtyFlag::None)
+			| ((HasDynamicModifiers() || NeedAnyPendingLazyModifiersGathering())
+				? ENavigationDirtyFlag::DynamicModifier
+				: ENavigationDirtyFlag::None)
+			| (Modifiers.HasAgentHeightAdjust()
+				? ENavigationDirtyFlag::UseAgentHeight
+				: ENavigationDirtyFlag::None);
 	}
 
 	FORCEINLINE FCompositeNavModifier GetModifierForAgent(const struct FNavAgentProperties* NavAgent = nullptr) const
diff --git a/Engine/Source/Runtime/Engine/Classes/AI/Navigation/NavigationTypes.h b/Engine/Source/Runtime/Engine/Classes/AI/Navigation/NavigationTypes.h
--- a/Engine/Source/Runtime/Engine/Classes/AI/Navigation/NavigationTypes.h
+++ b/Engine/Source/Runtime/Engine/Classes/AI/Navigation/NavigationTypes.h
@@ -64,35 +64,17 @@ namespace FNavigationSystem
 }
 
 UENUM()
-namespace ENavigationOptionFlag
+enum class ENavigationOptionFlag : uint8
 {
-	enum Type : int
-	{
-		Default,
-		Enable UMETA(DisplayName = "Yes"),	// UHT was complaining when tried to use True as value instead of Enable
-
-		Disable UMETA(DisplayName = "No"),
-
-		MAX UMETA(Hidden)
-	};
-}
+	Default,
+	Enable UMETA(DisplayName = "Yes"), // UHT was complaining when tried to use True as value instead of Enable
+	Disable UMETA(DisplayName = "No"),
+	MAX UMETA(Hidden)
+};
 
 //////////////////////////////////////////////////////////////////////////
 // Navigation data generation
 
-namespace ENavigationDirtyFlag
-{
-	enum Type : uint8
-	{
-		Geometry			= (1 << 0),
-		DynamicModifier		= (1 << 1),
-		UseAgentHeight		= (1 << 2),
-		NavigationBounds	= (1 << 3),
-
-		All				= Geometry | DynamicModifier,		// all rebuild steps here without additional flags
-	};
-}
-
 UENUM()
 enum class ENavDataGatheringModeConfig : uint8
 {
diff --git a/Engine/Source/Runtime/Engine/Private/AI/Navigation/NavigationDirtyArea.cpp b/Engine/Source/Runtime/Engine/Private/AI/Navigation/NavigationDirtyArea.cpp
--- a/Engine/Source/Runtime/Engine/Private/AI/Navigation/NavigationDirtyArea.cpp
+++ b/Engine/Source/Runtime/Engine/Private/AI/Navigation/NavigationDirtyArea.cpp
@@ -5,10 +5,10 @@
 #include "AI/NavigationSystemBase.h"
 #endif
 
-FNavigationDirtyArea::FNavigationDirtyArea(const FBox& InBounds, const int32 InFlags, UObject* const InOptionalSourceObject)
+FNavigationDirtyArea::FNavigationDirtyArea(const FBox& InBounds, const ENavigationDirtyFlag InFlags, UObject* const InOptionalSourceObject)
 	: Bounds(InBounds)
-	, Flags(InFlags)
 	, OptionalSourceObject(InOptionalSourceObject)
+	, Flags(InFlags)
 {
 #if !NO_LOGGING
 	if (!Bounds.IsValid || Bounds.ContainsNaN())
@@ -18,3 +18,9 @@ FNavigationDirtyArea::FNavigationDirtyArea(const FBox& InBounds, const int32 InF
 	}
 #endif //!NO_LOGGING
 }
+
+// Deprecated
+FNavigationDirtyArea::FNavigationDirtyArea(const FBox& InBounds, int32 InFlags, UObject* const InOptionalSourceObject /*= nullptr*/)
+	: FNavigationDirtyArea(InBounds, static_cast<ENavigationDirtyFlag>(InFlags), InOptionalSourceObject)
+{
+}
\ No newline at end of file
diff --git a/Engine/Source/Runtime/Engine/Private/AI/Navigation/NavigationDirtyElement.cpp b/Engine/Source/Runtime/Engine/Private/AI/Navigation/NavigationDirtyElement.cpp
--- a/Engine/Source/Runtime/Engine/Private/AI/Navigation/NavigationDirtyElement.cpp
+++ b/Engine/Source/Runtime/Engine/Private/AI/Navigation/NavigationDirtyElement.cpp
@@ -1,16 +1,18 @@
 ﻿// Copyright Epic Games, Inc. All Rights Reserved.
 
 #include "AI/Navigation/NavigationDirtyElement.h"
+#include "AI/Navigation/NavigationDirtyArea.h"
 #include "AI/Navigation/NavigationTypes.h"
 
-FNavigationDirtyElement::FNavigationDirtyElement(UObject* InOwner, INavRelevantInterface* InNavInterface, int32 InFlagsOverride /*= 0*/, const bool bUseWorldPartitionedDynamicMode /*= false*/)
+FNavigationDirtyElement::FNavigationDirtyElement(
+	UObject* InOwner,
+	INavRelevantInterface* InNavInterface,
+	const ENavigationDirtyFlag InFlagsOverride,
+	const bool bUseWorldPartitionedDynamicMode /*= false*/)
 	: Owner(InOwner)
 	, NavInterface(InNavInterface)
 	, FlagsOverride(InFlagsOverride)
-	, PrevFlags(0)
-	, PrevBounds(ForceInit)
-	, bHasPrevData(false)
-	, bInvalidRequest(false)
+	, PrevFlags(ENavigationDirtyFlag::None)
 {
 	if (bUseWorldPartitionedDynamicMode)
 	{
@@ -23,3 +25,24 @@ FNavigationDirtyElement::FNavigationDirtyElement(UObject* InOwner, INavRelevantI
 		bIsInBaseNavmesh = false;
 	}
 }
+
+FNavigationDirtyElement::FNavigationDirtyElement(UObject* InOwner, INavRelevantInterface* InNavInterface, const bool bUseWorldPartitionedDynamicMode /*= false*/)
+	: FNavigationDirtyElement(InOwner, InNavInterface, ENavigationDirtyFlag::None, bUseWorldPartitionedDynamicMode)
+{
+}
+
+FNavigationDirtyElement::FNavigationDirtyElement(UObject* InOwner)
+	: FNavigationDirtyElement(InOwner, /*InNavInterface*/nullptr, /*bUseWorldPartitionedDynamicMode*/false)
+{
+}
+
+FNavigationDirtyElement::FNavigationDirtyElement()
+	: FNavigationDirtyElement(/*InOwner*/nullptr, /*InNavInterface*/nullptr, /*bUseWorldPartitionedDynamicMode*/false)
+{
+}
+
+// Deprecated
+FNavigationDirtyElement::FNavigationDirtyElement(UObject* InOwner, INavRelevantInterface* InNavInterface, int32 InFlagsOverride /*= 0*/, const bool bUseWorldPartitionedDynamicMode /*= false*/)
+	: FNavigationDirtyElement(InOwner, InNavInterface, static_cast<ENavigationDirtyFlag>(InFlagsOverride), bUseWorldPartitionedDynamicMode)
+{
+}
diff --git a/Engine/Source/Runtime/NavigationSystem/Private/NavigationDataHandler.cpp b/Engine/Source/Runtime/NavigationSystem/Private/NavigationDataHandler.cpp
--- a/Engine/Source/Runtime/NavigationSystem/Private/NavigationDataHandler.cpp
+++ b/Engine/Source/Runtime/NavigationSystem/Private/NavigationDataHandler.cpp
@@ -10,7 +10,7 @@ DEFINE_LOG_CATEGORY_STATIC(LogNavOctree, Warning, All);
 
 namespace UE::NavigationHelper::Private
 {
-	int32 GetDirtyFlag(int32 UpdateFlags, int32 DefaultValue)
+	ENavigationDirtyFlag GetDirtyFlag(const int32 UpdateFlags, const ENavigationDirtyFlag DefaultValue)
 	{
 		return ((UpdateFlags & FNavigationOctreeController::OctreeUpdate_Geometry) != 0) ? ENavigationDirtyFlag::All :
 			((UpdateFlags & FNavigationOctreeController::OctreeUpdate_Modifiers) != 0) ? ENavigationDirtyFlag::DynamicModifier :
@@ -51,7 +51,7 @@ void FNavigationDataHandler::RemoveNavOctreeElementId(const FOctreeElementId2& E
 		// mark area occupied by given element as dirty except if explicitly set to skip this default behavior
 		if (!ElementData.Data->bShouldSkipDirtyAreaOnAddOrRemove)
 		{
-			const int32 DirtyFlag = UE::NavigationHelper::Private::GetDirtyFlag(UpdateFlags, ElementData.Data->GetDirtyFlag());
+			const ENavigationDirtyFlag DirtyFlag = UE::NavigationHelper::Private::GetDirtyFlag(UpdateFlags, ElementData.Data->GetDirtyFlag());
 			DirtyAreasController.AddArea(ElementData.Bounds.GetBox(), DirtyFlag, [&ElementData] { return ElementData.Data->SourceObject.Get(); }, nullptr, "Remove from navoctree");
 		}
 
@@ -96,7 +96,7 @@ FSetElementId FNavigationDataHandler::RegisterNavOctreeElement(UObject& ElementO
 
 		if (bCanAdd)
 		{
-			FNavigationDirtyElement UpdateInfo(&ElementOwner, &ElementInterface, UE::NavigationHelper::Private::GetDirtyFlag(UpdateFlags, 0), DirtyAreasController.bUseWorldPartitionedDynamicMode);
+			FNavigationDirtyElement UpdateInfo(&ElementOwner, &ElementInterface, UE::NavigationHelper::Private::GetDirtyFlag(UpdateFlags, ENavigationDirtyFlag::None), DirtyAreasController.bUseWorldPartitionedDynamicMode);
 
 			SetId = OctreeController.PendingOctreeUpdates.FindId(UpdateInfo);
 			if (SetId.IsValidId())
@@ -180,7 +180,7 @@ void FNavigationDataHandler::AddElementToNavOctree(const FNavigationDirtyElement
 	}
 
 	// mark area occupied by given element as dirty except if explicitly set to skip this default behavior
-	const int32 DirtyFlag = DirtyElement.FlagsOverride ? DirtyElement.FlagsOverride : GeneratedData.Data->GetDirtyFlag();
+	const ENavigationDirtyFlag DirtyFlag = DirtyElement.FlagsOverride != ENavigationDirtyFlag::None ? DirtyElement.FlagsOverride : GeneratedData.Data->GetDirtyFlag();
 	if (GeneratedData.Data->bShouldSkipDirtyAreaOnAddOrRemove)
 	{
 		if (DirtyElement.ExplicitAreasToDirty.Num() > 0)
@@ -268,7 +268,7 @@ void FNavigationDataHandler::UpdateNavOctreeElement(UObject& ElementOwner, INavR
 
 	// grab existing octree data
 	FBox CurrentBounds;
-	int32 CurrentFlags;
+	ENavigationDirtyFlag CurrentFlags;
 	const bool bAlreadyExists = OctreeController.GetNavOctreeElementData(ElementOwner, CurrentFlags, CurrentBounds);
 
 	// don't invalidate pending requests
diff --git a/Engine/Source/Runtime/NavigationSystem/Private/NavigationDirtyAreasController.cpp b/Engine/Source/Runtime/NavigationSystem/Private/NavigationDirtyAreasController.cpp
--- a/Engine/Source/Runtime/NavigationSystem/Private/NavigationDirtyAreasController.cpp
+++ b/Engine/Source/Runtime/NavigationSystem/Private/NavigationDirtyAreasController.cpp
@@ -115,17 +115,30 @@ void FNavigationDirtyAreasController::Tick(const float DeltaSeconds, const TArra
 	}
 }
 
+// Deprecated
 void FNavigationDirtyAreasController::AddArea(const FBox& NewArea, const int32 Flags, const TFunction<UObject*()>& ObjectProviderFunc /*= nullptr*/,
 	const FNavigationDirtyElement* DirtyElement /*= nullptr*/, const FName& DebugReason /*= NAME_None*/)
 {
-	AddAreas({NewArea}, Flags, ObjectProviderFunc, DirtyElement, DebugReason);
+	AddAreas({NewArea}, static_cast<ENavigationDirtyFlag>(Flags), ObjectProviderFunc, DirtyElement, DebugReason);
 }
 
+// Deprecated
 void FNavigationDirtyAreasController::AddAreas(const TConstArrayView<FBox> NewAreas, const int32 Flags, const TFunction<UObject*()>& ObjectProviderFunc, const FNavigationDirtyElement* DirtyElement, const FName& DebugReason)
+{
+	AddAreas(NewAreas, static_cast<ENavigationDirtyFlag>(Flags), ObjectProviderFunc, DirtyElement, DebugReason);
+}
+
+void FNavigationDirtyAreasController::AddArea(const FBox& NewArea, const ENavigationDirtyFlag Flags, const TFunction<UObject*()>& ObjectProviderFunc /*= nullptr*/,
+	const FNavigationDirtyElement* DirtyElement /*= nullptr*/, const FName& DebugReason /*= NAME_None*/)
+{
+	AddAreas({NewArea}, Flags, ObjectProviderFunc, DirtyElement, DebugReason);
+}
+
+void FNavigationDirtyAreasController::AddAreas(const TConstArrayView<FBox> NewAreas, const ENavigationDirtyFlag Flags, const TFunction<UObject*()>& ObjectProviderFunc, const FNavigationDirtyElement* DirtyElement, const FName& DebugReason)
 {
 #if !UE_BUILD_SHIPPING
 	// always keep track of reported areas even when filtered out by invalid area as long as flags are valid
-	bDirtyAreasReportedWhileAccumulationLocked = bDirtyAreasReportedWhileAccumulationLocked || (Flags > 0 && !bCanAccumulateDirtyAreas);
+	bDirtyAreasReportedWhileAccumulationLocked |= (Flags != ENavigationDirtyFlag::None) && !bCanAccumulateDirtyAreas;
 
 	checkf(NewAreas.Num() > 0, TEXT("All callers of this method are expected to provide at least one area."));
 #endif // !UE_BUILD_SHIPPING
@@ -220,7 +233,7 @@ void FNavigationDirtyAreasController::AddAreas(const TConstArrayView<FBox> NewAr
 		}
 #endif // !UE_BUILD_SHIPPING
 
-		if (Flags > 0 && bCanAccumulateDirtyAreas)
+		if (Flags != ENavigationDirtyFlag::None && bCanAccumulateDirtyAreas)
 		{
 			DirtyAreas.Add(FNavigationDirtyArea(NewArea, Flags, SourceObject));
 		}
diff --git a/Engine/Source/Runtime/NavigationSystem/Private/NavigationOctreeController.cpp b/Engine/Source/Runtime/NavigationSystem/Private/NavigationOctreeController.cpp
--- a/Engine/Source/Runtime/NavigationSystem/Private/NavigationOctreeController.cpp
+++ b/Engine/Source/Runtime/NavigationSystem/Private/NavigationOctreeController.cpp
@@ -24,21 +24,30 @@ void FNavigationOctreeController::SetNavigableGeometryStoringMode(FNavigationOct
 	NavOctree->SetNavigableGeometryStoringMode(NavGeometryMode);
 }
 
-bool FNavigationOctreeController::GetNavOctreeElementData(const UObject& NodeOwner, int32& DirtyFlags, FBox& DirtyBounds)
+bool FNavigationOctreeController::GetNavOctreeElementData(const UObject& NodeOwner, ENavigationDirtyFlag& OutDirtyFlags, FBox& OutDirtyBounds)
 {
 	const FOctreeElementId2* ElementId = GetObjectsNavOctreeId(NodeOwner);
 	if (ElementId != nullptr && IsValidElement(*ElementId))
 	{
 		// mark area occupied by given actor as dirty
 		const FNavigationOctreeElement& ElementData = NavOctree->GetElementById(*ElementId);
-		DirtyFlags = ElementData.Data->GetDirtyFlag();
-		DirtyBounds = ElementData.Bounds.GetBox();
+		OutDirtyFlags = ElementData.Data->GetDirtyFlag();
+		OutDirtyBounds = ElementData.Bounds.GetBox();
 		return true;
 	}
 
 	return false;
 }
 
+// Deprecated
+bool FNavigationOctreeController::GetNavOctreeElementData(const UObject& NodeOwner, int32& DirtyFlags, FBox& DirtyBounds)
+{
+	ENavigationDirtyFlag TmpDirtyFlags = ENavigationDirtyFlag::None;
+	const bool bSuccess = GetNavOctreeElementData(NodeOwner, TmpDirtyFlags, DirtyBounds);
+	DirtyFlags = static_cast<int32>(TmpDirtyFlags);
+	return bSuccess;
+}
+
 const FNavigationRelevantData* FNavigationOctreeController::GetDataForObject(const UObject& Object) const
 {
 	const FOctreeElementId2* ElementId = GetObjectsNavOctreeId(Object);
diff --git a/Engine/Source/Runtime/NavigationSystem/Private/NavigationPath.cpp b/Engine/Source/Runtime/NavigationSystem/Private/NavigationPath.cpp
--- a/Engine/Source/Runtime/NavigationSystem/Private/NavigationPath.cpp
+++ b/Engine/Source/Runtime/NavigationSystem/Private/NavigationPath.cpp
@@ -519,7 +519,7 @@ void UNavigationPath::EnableDebugDrawing(bool bShouldDrawDebugData, FLinearColor
 	}
 }
 
-void UNavigationPath::EnableRecalculationOnInvalidation(TEnumAsByte<ENavigationOptionFlag::Type> DoRecalculation)
+void UNavigationPath::EnableRecalculationOnInvalidation(const ENavigationOptionFlag DoRecalculation)
 {
 	if (DoRecalculation != RecalculateOnInvalidation)
 	{
diff --git a/Engine/Source/Runtime/NavigationSystem/Private/NavigationSystem.cpp b/Engine/Source/Runtime/NavigationSystem/Private/NavigationSystem.cpp
--- a/Engine/Source/Runtime/NavigationSystem/Private/NavigationSystem.cpp
+++ b/Engine/Source/Runtime/NavigationSystem/Private/NavigationSystem.cpp
@@ -2802,7 +2802,7 @@ void UNavigationSystemV1::RegisterCustomLink(INavLinkCustomInterface& CustomLink
 			const FBox LinkBounds = ComputeCustomLinkBounds(CustomLink);
 			if (LinkBounds.IsValid)
 			{
-				AddDirtyArea(LinkBounds, FNavigationOctreeController::OctreeUpdate_Modifiers);
+				AddDirtyArea(LinkBounds, ENavigationDirtyFlag::DynamicModifier);
 			}
 		}
 	}
@@ -3240,19 +3240,37 @@ void UNavigationSystemV1::UnregisterComponentToNavOctree(UActorComponent* Comp)
 	}
 }
 
+// Deprecated
 void UNavigationSystemV1::AddDirtyArea(const FBox& NewArea, int32 Flags, const FName& DebugReason /*= NAME_None*/)
 {
-	DefaultDirtyAreasController.AddArea(NewArea, Flags, nullptr, nullptr, DebugReason);
+	AddDirtyArea(NewArea, static_cast<ENavigationDirtyFlag>(Flags), DebugReason);
 }
 
+// Deprecated
 void UNavigationSystemV1::AddDirtyArea(const FBox& NewArea, int32 Flags, const TFunction<UObject*()>& ObjectProviderFunc, const FName& DebugReason /*= NAME_None*/)
 {
-	DefaultDirtyAreasController.AddArea(NewArea, Flags, ObjectProviderFunc, nullptr, DebugReason);
+	AddDirtyArea(NewArea, static_cast<ENavigationDirtyFlag>(Flags), ObjectProviderFunc, DebugReason);
 }
 
+// Deprecated
 void UNavigationSystemV1::AddDirtyAreas(const TArray<FBox>& NewAreas, int32 Flags, const FName& DebugReason /*= NAME_None*/)
 { 
-	if (Flags == 0)
+	AddDirtyAreas(NewAreas, static_cast<ENavigationDirtyFlag>(Flags), DebugReason);
+}
+
+void UNavigationSystemV1::AddDirtyArea(const FBox& NewArea, const ENavigationDirtyFlag Flags, const FName& DebugReason /*= NAME_None*/)
+{
+	DefaultDirtyAreasController.AddArea(NewArea, Flags, nullptr, nullptr, DebugReason);
+}
+
+void UNavigationSystemV1::AddDirtyArea(const FBox& NewArea, const ENavigationDirtyFlag Flags, const TFunction<UObject*()>& ObjectProviderFunc, const FName& DebugReason /*= NAME_None*/)
+{
+	DefaultDirtyAreasController.AddArea(NewArea, Flags, ObjectProviderFunc, nullptr, DebugReason);
+}
+
+void UNavigationSystemV1::AddDirtyAreas(const TArray<FBox>& NewAreas, const ENavigationDirtyFlag Flags, const FName& DebugReason /*= NAME_None*/)
+{
+	if (Flags == ENavigationDirtyFlag::None)
 	{
 		return;
 	}
@@ -3285,9 +3303,18 @@ void UNavigationSystemV1::AddElementToNavOctree(const FNavigationDirtyElement& D
 	FNavigationDataHandler(DefaultOctreeController, DefaultDirtyAreasController).AddElementToNavOctree(DirtyElement);
 }
 
+bool UNavigationSystemV1::GetNavOctreeElementData(const UObject& NodeOwner, ENavigationDirtyFlag& OutDirtyFlags, FBox& OutDirtyBounds)
+{
+	return DefaultOctreeController.GetNavOctreeElementData(NodeOwner, OutDirtyFlags, OutDirtyBounds);
+}
+
+// Deprecated
 bool UNavigationSystemV1::GetNavOctreeElementData(const UObject& NodeOwner, int32& DirtyFlags, FBox& DirtyBounds)
 {
-	return DefaultOctreeController.GetNavOctreeElementData(NodeOwner, DirtyFlags, DirtyBounds);
+	ENavigationDirtyFlag TmpDirtyFlags = ENavigationDirtyFlag::None;
+	const bool bSuccess = GetNavOctreeElementData(NodeOwner, TmpDirtyFlags, DirtyBounds);
+	DirtyFlags = static_cast<int32>(TmpDirtyFlags);
+	return bSuccess;
 }
 
 void UNavigationSystemV1::UnregisterNavOctreeElement(UObject* ElementOwner, INavRelevantInterface* ElementInterface, int32 UpdateFlags)
diff --git a/Engine/Source/Runtime/NavigationSystem/Public/NavigationDirtyAreasController.h b/Engine/Source/Runtime/NavigationSystem/Public/NavigationDirtyAreasController.h
--- a/Engine/Source/Runtime/NavigationSystem/Public/NavigationDirtyAreasController.h
+++ b/Engine/Source/Runtime/NavigationSystem/Public/NavigationDirtyAreasController.h
@@ -12,6 +12,7 @@
 
 class UObject;
 struct FNavigationDirtyArea;
+enum class ENavigationDirtyFlag : uint8;
 
 NAVIGATIONSYSTEM_API DECLARE_LOG_CATEGORY_EXTERN(LogNavigationDirtyArea, Warning, All);
 
@@ -64,6 +65,9 @@ public:
 	 * @param DirtyElement Optional dirty element
 	 * @param DebugReason Source of the new area
 	 */
+	NAVIGATIONSYSTEM_API void AddArea(const FBox& NewArea, const ENavigationDirtyFlag Flags, const TFunction<UObject*()>& ObjectProviderFunc = nullptr,
+		const FNavigationDirtyElement* DirtyElement = nullptr, const FName& DebugReason = NAME_None);
+	UE_DEPRECATED(5.5, "Use the version taking ENavigationDirtyFlag instead.")
 	NAVIGATIONSYSTEM_API void AddArea(const FBox& NewArea, const int32 Flags, const TFunction<UObject*()>& ObjectProviderFunc = nullptr,
 		const FNavigationDirtyElement* DirtyElement = nullptr, const FName& DebugReason = NAME_None);
 
@@ -77,6 +81,9 @@ public:
 	 * @param DirtyElement Optional dirty element
 	 * @param DebugReason Source of the new area
 	 */
+	NAVIGATIONSYSTEM_API void AddAreas(const TConstArrayView<FBox> NewAreas, const ENavigationDirtyFlag Flags, const TFunction<UObject*()>& ObjectProviderFunc = nullptr,
+		const FNavigationDirtyElement* DirtyElement = nullptr, const FName& DebugReason = NAME_None);
+	UE_DEPRECATED(5.5, "Use the version taking ENavigationDirtyFlag instead.")
 	NAVIGATIONSYSTEM_API void AddAreas(const TConstArrayView<FBox> NewAreas, const int32 Flags, const TFunction<UObject*()>& ObjectProviderFunc = nullptr,
 		const FNavigationDirtyElement* DirtyElement = nullptr, const FName& DebugReason = NAME_None);
 	
diff --git a/Engine/Source/Runtime/NavigationSystem/Public/NavigationOctreeController.h b/Engine/Source/Runtime/NavigationSystem/Public/NavigationOctreeController.h
--- a/Engine/Source/Runtime/NavigationSystem/Public/NavigationOctreeController.h
+++ b/Engine/Source/Runtime/NavigationSystem/Public/NavigationOctreeController.h
@@ -32,6 +32,8 @@ struct FNavigationOctreeController
 	NAVIGATIONSYSTEM_API const FNavigationOctree* GetOctree() const;
 	NAVIGATIONSYSTEM_API FNavigationOctree* GetMutableOctree();
 	NAVIGATIONSYSTEM_API const FOctreeElementId2* GetObjectsNavOctreeId(const UObject& Object) const;
+	NAVIGATIONSYSTEM_API bool GetNavOctreeElementData(const UObject& NodeOwner, ENavigationDirtyFlag& OutDirtyFlags, FBox& OutDirtyBounds);
+	UE_DEPRECATED(5.5, "Use the version taking ENavigationDirtyFlag& instead.")
 	NAVIGATIONSYSTEM_API bool GetNavOctreeElementData(const UObject& NodeOwner, int32& DirtyFlags, FBox& DirtyBounds);
 	NAVIGATIONSYSTEM_API const FNavigationRelevantData* GetDataForObject(const UObject& Object) const;
 	NAVIGATIONSYSTEM_API FNavigationRelevantData* GetMutableDataForObject(const UObject& Object);
diff --git a/Engine/Source/Runtime/NavigationSystem/Public/NavigationPath.h b/Engine/Source/Runtime/NavigationSystem/Public/NavigationPath.h
--- a/Engine/Source/Runtime/NavigationSystem/Public/NavigationPath.h
+++ b/Engine/Source/Runtime/NavigationSystem/Public/NavigationPath.h
@@ -32,7 +32,7 @@ class UNavigationPath : public UObject
 	TArray<FVector> PathPoints;
 
 	UPROPERTY(BlueprintReadOnly, Category = Navigation)
-	TEnumAsByte<ENavigationOptionFlag::Type> RecalculateOnInvalidation;
+	ENavigationOptionFlag RecalculateOnInvalidation = ENavigationOptionFlag::Default;
 
 private:	
 	uint32 bIsValid : 1;
@@ -61,7 +61,7 @@ public:
 
 	/** if enabled path will request recalculation if it gets invalidated due to a change to underlying navigation */
 	UFUNCTION(BlueprintCallable, Category = "AI|Navigation")
-	NAVIGATIONSYSTEM_API void EnableRecalculationOnInvalidation(TEnumAsByte<ENavigationOptionFlag::Type> DoRecalculation);
+	NAVIGATIONSYSTEM_API void EnableRecalculationOnInvalidation(ENavigationOptionFlag DoRecalculation);
 
 	UFUNCTION(BlueprintCallable, Category = "AI|Navigation")
 	NAVIGATIONSYSTEM_API double GetPathLength() const;
diff --git a/Engine/Source/Runtime/NavigationSystem/Public/NavigationSystem.h b/Engine/Source/Runtime/NavigationSystem/Public/NavigationSystem.h
--- a/Engine/Source/Runtime/NavigationSystem/Public/NavigationSystem.h
+++ b/Engine/Source/Runtime/NavigationSystem/Public/NavigationSystem.h
@@ -916,8 +916,14 @@ public:
 	/** updates bounds of all components implementing INavRelevantInterface */
 	static NAVIGATIONSYSTEM_API void UpdateNavOctreeBounds(AActor* Actor);
 
+	NAVIGATIONSYSTEM_API void AddDirtyArea(const FBox& NewArea, ENavigationDirtyFlag Flags, const FName& DebugReason = NAME_None);
+	NAVIGATIONSYSTEM_API void AddDirtyArea(const FBox& NewArea, ENavigationDirtyFlag Flags, const TFunction<UObject*()>& ObjectProviderFunc, const FName& DebugReason = NAME_None);
+	NAVIGATIONSYSTEM_API void AddDirtyAreas(const TArray<FBox>& NewAreas, ENavigationDirtyFlag Flags, const FName& DebugReason = NAME_None);
+	UE_DEPRECATED(5.5, "Use the version taking ENavigationDirtyFlag instead.")
 	NAVIGATIONSYSTEM_API void AddDirtyArea(const FBox& NewArea, int32 Flags, const FName& DebugReason = NAME_None);
+	UE_DEPRECATED(5.5, "Use the version taking ENavigationDirtyFlag instead.")
 	NAVIGATIONSYSTEM_API void AddDirtyArea(const FBox& NewArea, int32 Flags, const TFunction<UObject*()>& ObjectProviderFunc, const FName& DebugReason = NAME_None);
+	UE_DEPRECATED(5.5, "Use the version taking ENavigationDirtyFlag instead.")
 	NAVIGATIONSYSTEM_API void AddDirtyAreas(const TArray<FBox>& NewAreas, int32 Flags, const FName& DebugReason = NAME_None);
 	NAVIGATIONSYSTEM_API bool HasDirtyAreasQueued() const;
 	NAVIGATIONSYSTEM_API int32 GetNumDirtyAreas() const;
@@ -1272,8 +1278,9 @@ protected:
 	NAVIGATIONSYSTEM_API void UnregisterNavOctreeElement(UObject* ElementOwner, INavRelevantInterface* ElementInterface, int32 UpdateFlags);
 	
 	/** read element data from navigation octree */
+	NAVIGATIONSYSTEM_API bool GetNavOctreeElementData(const UObject& NodeOwner, ENavigationDirtyFlag& OutDirtyFlags, FBox& OutDirtyBounds);
+	UE_DEPRECATED(5.5, "Use the version taking ENavigationDirtyFlag& instead.")
 	NAVIGATIONSYSTEM_API bool GetNavOctreeElementData(const UObject& NodeOwner, int32& DirtyFlags, FBox& DirtyBounds);
-	//bool GetNavOctreeElementData(UObject* NodeOwner, int32& DirtyFlags, FBox& DirtyBounds);
 
 	/** Adds given element to NavOctree. No check for owner's validity are performed, 
 	 *	nor its presence in NavOctree - function assumes callee responsibility 
