[Navigation] Restored deprecated properties that shouldn't have been removed in recent cleanup.

diff --git a/Engine/Source/Runtime/NavigationSystem/Private/NavMesh/PImplRecastNavMesh.cpp b/Engine/Source/Runtime/NavigationSystem/Private/NavMesh/PImplRecastNavMesh.cpp
--- a/Engine/Source/Runtime/NavigationSystem/Private/NavMesh/PImplRecastNavMesh.cpp
+++ b/Engine/Source/Runtime/NavigationSystem/Private/NavMesh/PImplRecastNavMesh.cpp
@@ -606,6 +606,19 @@ void FPImplRecastNavMesh::Serialize( FArchive& Ar, int32 NavMeshVersion )
 		// assumes tiles are rectangular
 		
 		float DefaultCellSize = NavMeshOwner->GetCellSize(ENavigationDataResolution::Default);
+
+#if WITH_EDITORONLY_DATA
+		PRAGMA_DISABLE_DEPRECATION_WARNINGS
+		if (NavMeshVersion < NAVMESHVER_TILE_RESOLUTIONS)
+		{
+			// For backward compatibility, read original CellSize value.
+			// In ARecastNavMesh::PostLoad(), cell sizes for the different resolutions are set to the old CellSize value but it occurs later (PostLoad).
+			// This why we explicitly need to read CellSize for older versions here.
+			DefaultCellSize = NavMeshOwner->CellSize;
+		}
+		PRAGMA_ENABLE_DEPRECATION_WARNINGS
+#endif // WITH_EDITORONLY_DATA
+
 		const FVector::FReal ActorsTileSize = static_cast<FVector::FReal>(FMath::TruncToInt(NavMeshOwner->TileSizeUU / DefaultCellSize)) * DefaultCellSize;
 
 		if (ActorsTileSize != Params.tileWidth)
diff --git a/Engine/Source/Runtime/NavigationSystem/Private/NavMesh/RecastNavMesh.cpp b/Engine/Source/Runtime/NavigationSystem/Private/NavMesh/RecastNavMesh.cpp
--- a/Engine/Source/Runtime/NavigationSystem/Private/NavMesh/RecastNavMesh.cpp
+++ b/Engine/Source/Runtime/NavigationSystem/Private/NavMesh/RecastNavMesh.cpp
@@ -679,7 +679,38 @@ void ARecastNavMesh::PostLoad()
 			UNavigationSystemBase::OnNavigationInitStartStaticDelegate().AddUObject(this, &ARecastNavMesh::CheckToDiscardSubLevelNavData);
 		}
 	}
-	
+
+#if WITH_EDITORONLY_DATA
+	PRAGMA_DISABLE_DEPRECATION_WARNINGS
+	// If needed, initialize from deprecated value.
+	if (NavMeshVersion < NAVMESHVER_TILE_RESOLUTIONS)
+	{
+		for (int i = 0; i < (uint8)ENavigationDataResolution::MAX; ++i)
+		{
+			SetCellSize((ENavigationDataResolution)i, CellSize);
+		}
+	}
+
+	// If needed, initialize CellHeight from the deprecated value.
+	if (NavMeshVersion < NAVMESHVER_TILE_RESOLUTIONS_CELLHEIGHT)
+	{
+		for (int i = 0; i < (uint8)ENavigationDataResolution::MAX; ++i)
+		{
+			SetCellHeight((ENavigationDataResolution)i, CellHeight);
+		}
+	}
+
+	// If needed, initialize AgentMaxStepHeight from the deprecated value.
+	if (NavMeshVersion < NAVMESHVER_TILE_RESOLUTIONS_AGENTMAXSTEPHEIGHT)
+	{
+		for (int i = 0; i < (uint8)ENavigationDataResolution::MAX; ++i)
+		{
+			SetAgentMaxStepHeight((ENavigationDataResolution)i, AgentMaxStepHeight);
+		}
+	}
+	PRAGMA_ENABLE_DEPRECATION_WARNINGS
+#endif // WITH_EDITORONLY_DATA
+
 	for (uint8 Index = 0; Index < (uint8)ENavigationDataResolution::MAX; Index++)
 	{
 		UE_CLOG(TileSizeUU < GetCellSize((ENavigationDataResolution)Index), LogNavigation, Error, TEXT("%s: TileSizeUU (%f) being less than CellSize (%f) is an invalid case and will cause navmesh generation issues.")
diff --git a/Engine/Source/Runtime/NavigationSystem/Public/NavMesh/RecastNavMesh.h b/Engine/Source/Runtime/NavigationSystem/Public/NavMesh/RecastNavMesh.h
--- a/Engine/Source/Runtime/NavigationSystem/Public/NavMesh/RecastNavMesh.h
+++ b/Engine/Source/Runtime/NavigationSystem/Public/NavMesh/RecastNavMesh.h
@@ -777,6 +777,26 @@ class ARecastNavMesh : public ANavigationData
 	UPROPERTY(EditAnywhere, Category=Generation, config, meta=(ClampMin = "300.0"))
 	float TileSizeUU;
 
+#if WITH_EDITORONLY_DATA
+	/**
+	 * Note that we are not using _DEPRECATED on the following deprecated properties
+	 * since it prevents the property from being serialized back which can break the
+	 * process of duplicating the navmesh for PIE
+	 */
+
+	UE_DEPRECATED(all, "Use NavMeshResolutionParams to set CellSize for the different resolutions instead")
+	UPROPERTY(config, meta = (DeprecatedProperty, DeprecationMessage = "Use NavMeshResolutionParams to set CellSize for the different resolutions instead"))
+	float CellSize;
+
+	UE_DEPRECATED(all, "Use NavMeshResolutionParams to set CellHeight for the different resolutions instead")
+	UPROPERTY(config, meta = (DeprecatedProperty, DeprecationMessage = "Use NavMeshResolutionParams to set CellHeight for the different resolutions instead"))
+	float CellHeight;
+
+	UE_DEPRECATED(all, "Use NavMeshResolutionParams to set AgentMaxStepHeight for the different resolutions instead")
+	UPROPERTY(config, meta = (DeprecatedProperty, DeprecationMessage = "Use NavMeshResolutionParams to set AgentMaxStepHeight for the different resolutions instead"))
+	float AgentMaxStepHeight;
+#endif // WITH_EDITORONLY_DATA
+
 	/** Resolution params 
 	 * If using multiple resolutions, it's recommended to chose the highest resolution first and 
 	 * set it according to the highest desired precision and then the other resolutions. */
