[NavLink Generation] Allow generated navlinks to be discovered and handled by pathfollow in a similar way to custom links.
Addition of UGeneratedNavLinksProxy class
Addition of properties to FNavLinkGenerationJumpDownConfig to configure the link handler

diff --git a/Engine/Source/Runtime/AIModule/Classes/Navigation/GeneratedNavLinksProxy.h b/Engine/Source/Runtime/AIModule/Classes/Navigation/GeneratedNavLinksProxy.h
new file mode 100644
--- /dev/null
+++ b/Engine/Source/Runtime/AIModule/Classes/Navigation/GeneratedNavLinksProxy.h
@@ -0,0 +1,40 @@
+﻿// Copyright Epic Games, Inc. All Rights Reserved.
+
+#pragma once
+
+#include "BaseGeneratedNavLinksProxy.h"
+#include "GeneratedNavLinksProxy.generated.h"
+
+class AActor;
+class UObject;
+
+DECLARE_DYNAMIC_MULTICAST_DELEGATE_TwoParams(FLinkReachedSignature, AActor*, MovingActor, const FVector, DestinationPoint);
+
+/**
+ * Experimental
+ * Blueprintable class used to handle generated links as custom links.
+ */
+UCLASS(Blueprintable, BlueprintType, MinimalAPI)
+class UGeneratedNavLinksProxy : public UBaseGeneratedNavLinksProxy
+{
+	GENERATED_UCLASS_BODY()
+
+	virtual UWorld* GetWorld() const override;
+
+	// BEGIN INavLinkCustomInterface
+	AIMODULE_API virtual bool OnLinkMoveStarted(class UObject* PathComp, const FVector& DestPoint) override;
+	// END INavLinkCustomInterface
+
+	//////////////////////////////////////////////////////////////////////////
+	// Blueprint interface for smart links
+	
+	/** Called when agent reaches smart link during path following. */
+	UFUNCTION(BlueprintImplementableEvent)
+	AIMODULE_API void ReceiveSmartLinkReached(AActor* Agent, const FVector Destination);
+
+protected:
+	void NotifySmartLinkReached(UObject* PathingAgent, const FVector DestPoint);
+	
+	UPROPERTY(BlueprintAssignable)
+	FLinkReachedSignature OnSmartLinkReached;
+};
diff --git a/Engine/Source/Runtime/AIModule/Private/Navigation/GeneratedNavLinksProxy.cpp b/Engine/Source/Runtime/AIModule/Private/Navigation/GeneratedNavLinksProxy.cpp
new file mode 100644
--- /dev/null
+++ b/Engine/Source/Runtime/AIModule/Private/Navigation/GeneratedNavLinksProxy.cpp
@@ -0,0 +1,59 @@
+﻿// Copyright Epic Games, Inc. All Rights Reserved.
+
+#include "Navigation/GeneratedNavLinksProxy.h"
+
+#include "Engine/GameEngine.h"
+#include "GameFramework/Controller.h"
+#include "Navigation/PathFollowingComponent.h"
+#include "GameFramework/Actor.h"
+#include "GameFramework/Pawn.h"
+
+#if WITH_EDITOR
+#include "Editor.h"
+#endif
+
+UGeneratedNavLinksProxy::UGeneratedNavLinksProxy(const FObjectInitializer& ObjectInitializer)
+	: Super(ObjectInitializer)
+{
+}
+
+bool UGeneratedNavLinksProxy::OnLinkMoveStarted(UObject* PathComp, const FVector& DestPoint)
+{
+	NotifySmartLinkReached(PathComp, DestPoint);
+	return true;
+}
+
+void UGeneratedNavLinksProxy::NotifySmartLinkReached(UObject* PathingAgent, const FVector DestPoint)
+{
+	UPathFollowingComponent* PathComp = Cast<UPathFollowingComponent>(PathingAgent);
+	if (PathComp)
+	{
+		AActor* PathOwner = PathComp->GetOwner();
+		const AController* ControllerOwner = Cast<AController>(PathOwner);
+		if (ControllerOwner)
+		{
+			PathOwner = ControllerOwner->GetPawn();
+		}
+	
+		ReceiveSmartLinkReached(PathOwner, DestPoint);
+		OnSmartLinkReached.Broadcast(PathOwner, DestPoint);
+	}
+}
+
+UWorld* UGeneratedNavLinksProxy::GetWorld() const
+{
+#if WITH_EDITOR
+	if (GEditor)
+	{
+		constexpr bool bEnsureIsGWorld = false;
+		return GEditor->GetEditorWorldContext(bEnsureIsGWorld).World();
+	}
+#endif // WITH_EDITOR
+
+	if (UGameEngine* GameEngine = Cast<UGameEngine>(GEngine))
+	{
+		return GameEngine->GetGameWorld();
+	}
+
+	return Super::GetWorld();
+}
diff --git a/Engine/Source/Runtime/Engine/Classes/AI/Navigation/NavigationTypes.h b/Engine/Source/Runtime/Engine/Classes/AI/Navigation/NavigationTypes.h
--- a/Engine/Source/Runtime/Engine/Classes/AI/Navigation/NavigationTypes.h
+++ b/Engine/Source/Runtime/Engine/Classes/AI/Navigation/NavigationTypes.h
@@ -188,6 +188,11 @@ struct FNavLinkId
 		return (Id & NavLinkIdBitMask) == 0;
 	}
 
+	/**
+	 *  Helper function: returns unique ID number for custom links.
+	 */
+	static ENGINE_API FNavLinkId GenerateUniqueId();
+	
 	/**
 	 *  Helper function: returns unique ID number for custom links.
 	 */
diff --git a/Engine/Source/Runtime/Engine/Private/AI/Navigation/NavigationTypes.cpp b/Engine/Source/Runtime/Engine/Private/AI/Navigation/NavigationTypes.cpp
--- a/Engine/Source/Runtime/Engine/Private/AI/Navigation/NavigationTypes.cpp
+++ b/Engine/Source/Runtime/Engine/Private/AI/Navigation/NavigationTypes.cpp
@@ -343,6 +343,14 @@ FNavLinkAuxiliaryId FNavLinkAuxiliaryId::GenerateUniqueAuxiliaryId(FStringView P
 	return FNavLinkAuxiliaryId(AuxiliaryId);
 }
 
+FNavLinkId FNavLinkId::GenerateUniqueId()
+{
+	// Apply NavLinkIdBitMask to differentiate Legacy Ids (that do not have the mask set).
+	const uint64 UniqueId = UE::Navigation::NavLinkIdHelpers::Private::MakeIdFromGUID(FGuid::NewGuid()) | NavLinkIdBitMask;
+	UE_LOG(LogNavLink, VeryVerbose, TEXT("%hs id: %u."), __FUNCTION__, UniqueId);
+	return FNavLinkId(UniqueId);
+}
+
 FNavLinkId FNavLinkId::GenerateUniqueId(FNavLinkAuxiliaryId AuxiliaryId, FGuid ActorInstanceGuid)
 {
 	// Apply NavLinkIdBitMask to differentiate Legacy Ids (that do not have the mask set).
diff --git a/Engine/Source/Runtime/NavigationSystem/Private/BaseGeneratedNavLinksProxy.cpp b/Engine/Source/Runtime/NavigationSystem/Private/BaseGeneratedNavLinksProxy.cpp
new file mode 100644
--- /dev/null
+++ b/Engine/Source/Runtime/NavigationSystem/Private/BaseGeneratedNavLinksProxy.cpp
@@ -0,0 +1,28 @@
+﻿// Copyright Epic Games, Inc. All Rights Reserved.
+
+#include "BaseGeneratedNavLinksProxy.h"
+
+UBaseGeneratedNavLinksProxy::UBaseGeneratedNavLinksProxy(const FObjectInitializer& ObjectInitializer)
+	: Super(ObjectInitializer)
+{
+}
+
+void UBaseGeneratedNavLinksProxy::GetLinkData(FVector& LeftPt, FVector& RightPt, ENavLinkDirection::Type& Direction) const
+{
+	ensureMsgf(false, TEXT("Should not be called on a generated navlink proxy since it's not representing a single link."));
+}
+
+FNavLinkId UBaseGeneratedNavLinksProxy::GetId() const
+{
+	return LinkProxyId;
+}
+
+void UBaseGeneratedNavLinksProxy::UpdateLinkId(FNavLinkId ProxyId)
+{
+	LinkProxyId = ProxyId;
+}
+
+UObject* UBaseGeneratedNavLinksProxy::GetLinkOwner() const
+{
+	return Owner;
+}
diff --git a/Engine/Source/Runtime/NavigationSystem/Private/NavMesh/LinkGenerationConfig.cpp b/Engine/Source/Runtime/NavigationSystem/Private/NavMesh/LinkGenerationConfig.cpp
--- a/Engine/Source/Runtime/NavigationSystem/Private/NavMesh/LinkGenerationConfig.cpp
+++ b/Engine/Source/Runtime/NavigationSystem/Private/NavMesh/LinkGenerationConfig.cpp
@@ -1,6 +1,7 @@
 ﻿// Copyright Epic Games, Inc. All Rights Reserved.
 
 #include "NavMesh/LinkGenerationConfig.h"
+#include "BaseGeneratedNavLinksProxy.h"
 
 #if WITH_RECAST
 #include "Detour/DetourNavLinkBuilderConfig.h"
@@ -13,7 +14,12 @@ void FNavLinkGenerationJumpDownConfig::CopyToDetourConfig(dtNavLinkBuilderJumpDo
 	OutDetourConfig.jumpMaxDepth = JumpMaxDepth;
 	OutDetourConfig.jumpEndsHeightTolerance	= JumpEndsHeightTolerance;
 	OutDetourConfig.samplingSeparationFactor = SamplingSeparationFactor;
-	OutDetourConfig.filterDistanceThreshold = FilterDistanceThreshold;  
+	OutDetourConfig.filterDistanceThreshold = FilterDistanceThreshold;
+
+	if (LinkProxy)
+	{
+		OutDetourConfig.linkUserId = LinkProxy->GetId().GetId();	
+	}
 }
 
 void FNavLinkGenerationJumpOverConfig::CopyToDetourConfig(dtNavLinkBuilderJumpOverConfig& OutDetourConfig) const
diff --git a/Engine/Source/Runtime/NavigationSystem/Private/NavMesh/RecastNavMesh.cpp b/Engine/Source/Runtime/NavigationSystem/Private/NavMesh/RecastNavMesh.cpp
--- a/Engine/Source/Runtime/NavigationSystem/Private/NavMesh/RecastNavMesh.cpp
+++ b/Engine/Source/Runtime/NavigationSystem/Private/NavMesh/RecastNavMesh.cpp
@@ -23,6 +23,7 @@
 #include "WorldPartition/NavigationData/NavigationDataChunkActor.h"
 #include "Math/Color.h"
 #include "NavigationDataHandler.h"
+#include "BaseGeneratedNavLinksProxy.h"
 
 #if WITH_EDITOR
 #include "EditorSupportDelegates.h"
@@ -756,6 +757,61 @@ void ARecastNavMesh::CheckToDiscardSubLevelNavData(const UNavigationSystemBase&
 	}
 }
 
+void ARecastNavMesh::RegisterGeneratedLinksProxy()
+{
+	if (!NavLinkJumpDownConfig.bLinkProxyRegistered && NavLinkJumpDownConfig.LinkProxy)
+	{
+		if (UWorld* World = GetWorld())
+		{
+			UNavigationSystemV1* NavSys = FNavigationSystem::GetCurrent<UNavigationSystemV1>(World);
+			if (NavSys)
+			{
+				UE_LOG(LogNavLink, Log, TEXT("RegisterLinkProxy id: %llu ptr: 0x%p"), NavLinkJumpDownConfig.LinkProxyId.GetId(), NavLinkJumpDownConfig.LinkProxy.Get());
+
+				NavSys->RegisterCustomLink(*NavLinkJumpDownConfig.LinkProxy);
+				NavLinkJumpDownConfig.bLinkProxyRegistered = true;
+			}
+		}
+	}
+}
+
+void ARecastNavMesh::UnregisterGeneratedLinksProxy()
+{
+	if (NavLinkJumpDownConfig.bLinkProxyRegistered && NavLinkJumpDownConfig.LinkProxy)
+	{
+		if (UWorld* World = GetWorld())
+		{
+			UNavigationSystemV1* NavSys = FNavigationSystem::GetCurrent<UNavigationSystemV1>(World);
+			if (NavSys)
+			{
+				UE_LOG(LogNavLink, Log, TEXT("UnregisterLinkProxy id: %llu ptr: 0x%p"), NavLinkJumpDownConfig.LinkProxyId.GetId(), NavLinkJumpDownConfig.LinkProxy.Get());
+
+				NavSys->UnregisterCustomLink(*NavLinkJumpDownConfig.LinkProxy);
+				NavLinkJumpDownConfig.bLinkProxyRegistered = false;
+			}
+		}
+	}
+}
+
+void ARecastNavMesh::CreateAndRegisterJumpDownLinksProxy(const FNavLinkId LinkProxyId /*= FNavLinkId::GenerateUniqueId()*/)
+{
+	if (!ensureMsgf(NavLinkJumpDownConfig.LinkProxy == nullptr, TEXT("LinkProxy is expected to have been cleanup before creating a new one. A new proxy was not created.")))
+	{
+		return;
+	}
+	
+	UBaseGeneratedNavLinksProxy* LinkProxy = NewObject<UBaseGeneratedNavLinksProxy>(this, NavLinkJumpDownConfig.LinkProxyClass);
+	if (LinkProxy)
+	{
+		NavLinkJumpDownConfig.LinkProxyId = LinkProxyId;
+		LinkProxy->UpdateLinkId(NavLinkJumpDownConfig.LinkProxyId);
+		LinkProxy->SetOwner(this);
+		NavLinkJumpDownConfig.LinkProxy = LinkProxy;
+
+		RegisterGeneratedLinksProxy();
+	}
+}
+
 void ARecastNavMesh::PostRegisterAllComponents()
 {
 	Super::PostRegisterAllComponents();
@@ -764,6 +820,13 @@ void ARecastNavMesh::PostRegisterAllComponents()
 	{
 		ApplyWorldOffset(GetActorLocation(), /*unused*/false);
 	}
+
+	// Create and register link proxy if enabled.
+	if (NavLinkJumpDownConfig.bEnabled && NavLinkJumpDownConfig.LinkProxyClass.Get() != nullptr)
+	{
+		// Use existing proxy id.
+		CreateAndRegisterJumpDownLinksProxy(NavLinkJumpDownConfig.LinkProxyId);
+	}
 }
 
 void ARecastNavMesh::PostInitProperties()
@@ -3606,6 +3669,20 @@ bool ARecastNavMesh::IsUsingActiveTilesGeneration(const UNavigationSystemV1& Nav
 	return SupportsRuntimeGeneration() && (NavSys.IsActiveTilesGenerationEnabled() || bIsWorldPartitioned);
 }
 
+void ARecastNavMesh::PostLoadPreRebuild()
+{
+	// If any, unregister previous generated links proxy.
+	UnregisterGeneratedLinksProxy();
+	NavLinkJumpDownConfig.LinkProxy = nullptr;
+
+	// Register new one if needed.
+	if (NavLinkJumpDownConfig.bEnabled && NavLinkJumpDownConfig.LinkProxyClass.Get() != nullptr)
+	{
+		// Since this is on rebuild, create a new proxy id.
+		CreateAndRegisterJumpDownLinksProxy();
+	}
+}
+
 void ARecastNavMesh::ConditionalConstructGenerator()
 {	
 	if (NavDataGenerator.IsValid())
diff --git a/Engine/Source/Runtime/NavigationSystem/Private/NavMesh/RecastNavMeshGenerator.cpp b/Engine/Source/Runtime/NavigationSystem/Private/NavMesh/RecastNavMeshGenerator.cpp
--- a/Engine/Source/Runtime/NavigationSystem/Private/NavMesh/RecastNavMeshGenerator.cpp
+++ b/Engine/Source/Runtime/NavigationSystem/Private/NavMesh/RecastNavMeshGenerator.cpp
@@ -3860,6 +3860,11 @@ dtStatus FRecastTileGenerator::BuildTileCacheLinks(FNavMeshBuildContext& BuildCo
 
 		FNavigationLink& NewLink = OutGeneratedLinks.Emplace_GetRef();
 		NewLink.bIsGenerated = true;
+		if (link.action == DT_LINK_ACTION_JUMP_DOWN)
+		{
+			NewLink.NavLinkId = FNavLinkId(linkBuilderConfig.jumpDownConfig.linkUserId);	
+		}
+		// @todo: Set NavLinkId for jump over links if we keep them.
 		NewLink.Left = Recast2UnrealPoint(midA);
 		NewLink.Right = Recast2UnrealPoint(midB);
 	}
diff --git a/Engine/Source/Runtime/NavigationSystem/Private/NavigationData.cpp b/Engine/Source/Runtime/NavigationSystem/Private/NavigationData.cpp
--- a/Engine/Source/Runtime/NavigationSystem/Private/NavigationData.cpp
+++ b/Engine/Source/Runtime/NavigationSystem/Private/NavigationData.cpp
@@ -545,6 +545,8 @@ void ANavigationData::RebuildAll()
 	LoadBeforeGeneratorRebuild();
 	FAssetCompilingManager::Get().FinishAllCompilation();
 	UE_LOG(LogNavigationDataBuild, Display, TEXT("   %s load time: %.2fs"), ANSI_TO_TCHAR(__FUNCTION__), (FPlatformTime::Seconds() - LoadTime));
+
+	PostLoadPreRebuild();
 	
 	ConditionalConstructGenerator(); //recreate generator
 	
diff --git a/Engine/Source/Runtime/NavigationSystem/Private/NavigationSystem.cpp b/Engine/Source/Runtime/NavigationSystem/Private/NavigationSystem.cpp
--- a/Engine/Source/Runtime/NavigationSystem/Private/NavigationSystem.cpp
+++ b/Engine/Source/Runtime/NavigationSystem/Private/NavigationSystem.cpp
@@ -2841,6 +2841,7 @@ void UNavigationSystemV1::RequestCustomLinkRegistering(INavLinkCustomInterface&
 	{
 		if (UNavigationObjectRepository* Repository = UWorld::GetSubsystem<UNavigationObjectRepository>(Owner->GetWorld()))
 		{
+			UE_LOG(LogNavLink, Log, TEXT("%hs 0x%p"), __FUNCTION__, &CustomLink);
 			Repository->RegisterCustomNavLinkObject(CustomLink);
 		}
 	}
@@ -2852,6 +2853,7 @@ void UNavigationSystemV1::RequestCustomLinkUnregistering(INavLinkCustomInterface
 	{
 		if (UNavigationObjectRepository* Repository = UWorld::GetSubsystem<UNavigationObjectRepository>(Owner->GetWorld()))
 		{
+			UE_LOG(LogNavLink, Log, TEXT("%hs 0x%p"), __FUNCTION__, &CustomLink);
 			Repository->UnregisterCustomNavLinkObject(CustomLink);
 		}
 	}
diff --git a/Engine/Source/Runtime/NavigationSystem/Public/BaseGeneratedNavLinksProxy.h b/Engine/Source/Runtime/NavigationSystem/Public/BaseGeneratedNavLinksProxy.h
new file mode 100644
--- /dev/null
+++ b/Engine/Source/Runtime/NavigationSystem/Public/BaseGeneratedNavLinksProxy.h
@@ -0,0 +1,38 @@
+﻿// Copyright Epic Games, Inc. All Rights Reserved.
+
+#pragma once
+
+#include "UObject/ObjectMacros.h"
+#include "UObject/Object.h"
+#include "NavLinkCustomInterface.h"
+
+#include "BaseGeneratedNavLinksProxy.generated.h"
+
+/**
+ * Experimental
+ * Base class used to create generated navlinks proxy.
+ * The proxy id is used to represent multiple links generated from the same configuration.
+ */
+UCLASS(Blueprintable, MinimalAPI)
+class UBaseGeneratedNavLinksProxy : public UObject, public INavLinkCustomInterface
+{
+	GENERATED_UCLASS_BODY()
+	
+	// BEGIN INavLinkCustomInterface
+	NAVIGATIONSYSTEM_API virtual void GetLinkData(FVector& LeftPt, FVector& RightPt, ENavLinkDirection::Type& Direction) const override;
+	NAVIGATIONSYSTEM_API virtual FNavLinkId GetId() const override;
+	NAVIGATIONSYSTEM_API virtual void UpdateLinkId(FNavLinkId ProxyId) override;
+	NAVIGATIONSYSTEM_API virtual UObject* GetLinkOwner() const override;
+	// END INavLinkCustomInterface
+
+	void SetOwner(UObject* NewOwner) { Owner = NewOwner; }
+	
+protected:
+	/** The LinkID will be the same for all navlinks using the proxy. */
+	UPROPERTY(Transient)
+	FNavLinkId LinkProxyId;
+
+	/** Proxy owner. */
+	UPROPERTY(Transient)
+	TObjectPtr<UObject> Owner = nullptr;
+};
diff --git a/Engine/Source/Runtime/NavigationSystem/Public/NavMesh/LinkGenerationConfig.h b/Engine/Source/Runtime/NavigationSystem/Public/NavMesh/LinkGenerationConfig.h
--- a/Engine/Source/Runtime/NavigationSystem/Public/NavMesh/LinkGenerationConfig.h
+++ b/Engine/Source/Runtime/NavigationSystem/Public/NavMesh/LinkGenerationConfig.h
@@ -1,8 +1,13 @@
 ﻿// Copyright Epic Games, Inc. All Rights Reserved.
 #pragma once
 
+#include "Templates/SubclassOf.h"
+#include "AI/Navigation/NavigationTypes.h"
+
 #include "LinkGenerationConfig.generated.h"
 
+class UBaseGeneratedNavLinksProxy;
+
 #if WITH_RECAST
 struct dtNavLinkBuilderJumpDownConfig;
 struct dtNavLinkBuilderJumpOverConfig;
@@ -12,7 +17,7 @@ struct dtNavLinkBuilderJumpOverConfig;
 USTRUCT()
 struct FNavLinkGenerationJumpDownConfig
 {
-	GENERATED_USTRUCT_BODY()
+	GENERATED_BODY()
 
 	/** Should this config be used to generate links. */
 	UPROPERTY(EditAnywhere, Config, Category = Settings)
@@ -43,6 +48,31 @@ struct FNavLinkGenerationJumpDownConfig
 	UPROPERTY(EditAnywhere, Config, Category = Settings, meta=(Units=cm, UIMin=0, ClampMin=0))
 	float FilterDistanceThreshold = 80.f;
 
+	/** Class used to handle links made with this configuration.
+	 * Using this allows to implement custom behaviors when using navlinks, for example during the pathfollow.
+	 * Note that having a proxy is not required for successful navlink pathfinding, but it does allow for custom behavior at the start and the end of a given navlink.
+	 * This implies that using LinkProxyClass is optional and it can remain empty (the default value).
+	 * @see INavLinkCustomInterface 
+	 * @see UGeneratedNavLinksProxy
+	 */
+	UPROPERTY(EditAnywhere, Category= Settings)
+	TSubclassOf<UBaseGeneratedNavLinksProxy> LinkProxyClass;
+
+	/** Identifier used identify the current proxy handler. All links generated through this config will use the same handler. */
+	UPROPERTY()
+	FNavLinkId LinkProxyId;
+
+	/** Current proxy. The proxy instance is build from the LinkProxyClass (provided it's not null).
+	 * A proxy will be created if a @see LinkProxyClass is used.
+	 */
+	UPROPERTY(Transient)
+	TObjectPtr<UBaseGeneratedNavLinksProxy> LinkProxy = nullptr;
+
+	/** Is the link proxy registered to the navigation system CustomNavLinksMap.
+	 * Registration occurs on PostRegisterAllComponents or on PostLoadPreRebuild if a new proxy was created. */
+	UPROPERTY(Transient)
+	bool bLinkProxyRegistered = false;
+
 #if WITH_RECAST	
 	/** Copy configuration to dtNavLinkBuilderJumpDownConfig. */
 	void CopyToDetourConfig(dtNavLinkBuilderJumpDownConfig& OutDetourConfig) const;
@@ -53,7 +83,7 @@ struct FNavLinkGenerationJumpDownConfig
 USTRUCT()
 struct FNavLinkGenerationJumpOverConfig
 {
-	GENERATED_USTRUCT_BODY()
+	GENERATED_BODY()
 
 	/** Should this config be used to generate links. */
 	UPROPERTY(EditAnywhere, Config, Category = Settings)
diff --git a/Engine/Source/Runtime/NavigationSystem/Public/NavMesh/RecastNavMesh.h b/Engine/Source/Runtime/NavigationSystem/Public/NavMesh/RecastNavMesh.h
--- a/Engine/Source/Runtime/NavigationSystem/Public/NavMesh/RecastNavMesh.h
+++ b/Engine/Source/Runtime/NavigationSystem/Public/NavMesh/RecastNavMesh.h
@@ -1532,6 +1532,9 @@ public:
 	 *  The subset is be defined by navinvokers or loaded world partitioned cells. */
 	NAVIGATIONSYSTEM_API bool IsUsingActiveTilesGeneration(const UNavigationSystemV1& NavSys) const;
 
+	/** Runs after LoadBeforeGeneratorRebuild but before the rebuild. */
+	NAVIGATIONSYSTEM_API virtual void PostLoadPreRebuild() override;
+	
 	NAVIGATIONSYSTEM_API virtual void ConditionalConstructGenerator() override;
 	
 	bool ShouldGatherDataOnGameThread() const { return bDoFullyAsyncNavDataGathering == false; }
@@ -1560,6 +1563,12 @@ protected:
 	NAVIGATIONSYSTEM_API virtual FRecastNavMeshGenerator* CreateGeneratorInstance();
 
 	NAVIGATIONSYSTEM_API void CheckToDiscardSubLevelNavData(const UNavigationSystemBase& NavSys);
+	
+	void RegisterGeneratedLinksProxy();
+	void UnregisterGeneratedLinksProxy();
+
+	/* Create and register links proxy. It's expected to be called on load or when the navmesh is rebuilt. */
+	void CreateAndRegisterJumpDownLinksProxy(const FNavLinkId LinkProxyId = FNavLinkId::GenerateUniqueId());
 
 private:
 	friend struct FRecastGraphWrapper;
diff --git a/Engine/Source/Runtime/NavigationSystem/Public/NavigationData.h b/Engine/Source/Runtime/NavigationSystem/Public/NavigationData.h
--- a/Engine/Source/Runtime/NavigationSystem/Public/NavigationData.h
+++ b/Engine/Source/Runtime/NavigationSystem/Public/NavigationData.h
@@ -638,6 +638,9 @@ public:
 
 	/** Any loading before NavDataGenerator->RebuildAll() */
 	virtual void LoadBeforeGeneratorRebuild() {}
+
+	/** Runs after LoadBeforeGeneratorRebuild but before the rebuild. */
+	virtual void PostLoadPreRebuild() {}
 	
 	/** Triggers rebuild in case navigation supports it */
 	NAVIGATIONSYSTEM_API virtual void RebuildAll();
diff --git a/Engine/Source/Runtime/Navmesh/Public/Detour/DetourNavLinkBuilderConfig.h b/Engine/Source/Runtime/Navmesh/Public/Detour/DetourNavLinkBuilderConfig.h
--- a/Engine/Source/Runtime/Navmesh/Public/Detour/DetourNavLinkBuilderConfig.h
+++ b/Engine/Source/Runtime/Navmesh/Public/Detour/DetourNavLinkBuilderConfig.h
@@ -26,6 +26,9 @@ struct dtNavLinkBuilderJumpDownConfig
 	/// When filtering similar links, distance used to compare between segment endpoints to match similar links.
 	/// Use greater distance for more filtering (0 to deactivate filtering). [Limit: > 0] [Units: wu]
 	float filterDistanceThreshold = 80.f;
+
+	/// User id used to handle links made from this configuration.
+	unsigned long long linkUserId = 0;
 };
 
 /** Configuration for generated jump over links. */
