Fix new LinkProxy ensure occurring when changing properties on ARecastNavmesh

#rnx

diff --git a/Engine/Source/Runtime/NavigationSystem/Private/NavMesh/RecastNavMesh.cpp b/Engine/Source/Runtime/NavigationSystem/Private/NavMesh/RecastNavMesh.cpp
--- a/Engine/Source/Runtime/NavigationSystem/Private/NavMesh/RecastNavMesh.cpp
+++ b/Engine/Source/Runtime/NavigationSystem/Private/NavMesh/RecastNavMesh.cpp
@@ -829,6 +829,14 @@ void ARecastNavMesh::PostRegisterAllComponents()
 	}
 }
 
+void ARecastNavMesh::PostUnregisterAllComponents()
+{
+	Super::PostUnregisterAllComponents();
+
+	UnregisterGeneratedLinksProxy();
+	NavLinkJumpDownConfig.LinkProxy = nullptr;
+}
+
 void ARecastNavMesh::PostInitProperties()
 {
 	if (HasAnyFlags(RF_ClassDefaultObject) == true)
diff --git a/Engine/Source/Runtime/NavigationSystem/Public/NavMesh/RecastNavMesh.h b/Engine/Source/Runtime/NavigationSystem/Public/NavMesh/RecastNavMesh.h
--- a/Engine/Source/Runtime/NavigationSystem/Public/NavMesh/RecastNavMesh.h
+++ b/Engine/Source/Runtime/NavigationSystem/Public/NavMesh/RecastNavMesh.h
@@ -1078,6 +1078,7 @@ public:
 	NAVIGATIONSYSTEM_API virtual void PostInitProperties() override;
 	NAVIGATIONSYSTEM_API virtual void PostLoad() override;
 	NAVIGATIONSYSTEM_API virtual void PostRegisterAllComponents() override;
+	NAVIGATIONSYSTEM_API virtual void PostUnregisterAllComponents() override;
 	NAVIGATIONSYSTEM_API virtual void BeginDestroy() override;
 
 #if WITH_EDITOR
