[Navmesh] Avoid registering link proxy with invalid id.

diff --git a/Engine/Source/Runtime/Engine/Classes/AI/Navigation/NavigationTypes.h b/Engine/Source/Runtime/Engine/Classes/AI/Navigation/NavigationTypes.h
--- a/Engine/Source/Runtime/Engine/Classes/AI/Navigation/NavigationTypes.h
+++ b/Engine/Source/Runtime/Engine/Classes/AI/Navigation/NavigationTypes.h
@@ -156,6 +156,11 @@ struct FNavLinkId
 		: Id(InId)
 	{}
 
+	bool IsValid() const
+	{
+		return Id != InvalidLinkId;
+	}
+
 	bool operator==(const FNavLinkId& Other) const
 	{
 		return Id == Other.Id;
diff --git a/Engine/Source/Runtime/NavigationSystem/Private/NavMesh/RecastNavMesh.cpp b/Engine/Source/Runtime/NavigationSystem/Private/NavMesh/RecastNavMesh.cpp
--- a/Engine/Source/Runtime/NavigationSystem/Private/NavMesh/RecastNavMesh.cpp
+++ b/Engine/Source/Runtime/NavigationSystem/Private/NavMesh/RecastNavMesh.cpp
@@ -824,8 +824,9 @@ void ARecastNavMesh::PostRegisterAllComponents()
 	// Create and register link proxy if enabled.
 	if (NavLinkJumpDownConfig.bEnabled && NavLinkJumpDownConfig.LinkProxyClass.Get() != nullptr)
 	{
-		// Use existing proxy id.
-		CreateAndRegisterJumpDownLinksProxy(NavLinkJumpDownConfig.LinkProxyId);
+		// If it's the first time, LinkProxyId will have the default invalid value so we need to generate an Id.
+		const FNavLinkId Id = NavLinkJumpDownConfig.LinkProxyId.IsValid() ? NavLinkJumpDownConfig.LinkProxyId :  FNavLinkId::GenerateUniqueId(); 
+		CreateAndRegisterJumpDownLinksProxy(Id);
 	}
 }
 
diff --git a/Engine/Source/Runtime/NavigationSystem/Private/NavigationSystem.cpp b/Engine/Source/Runtime/NavigationSystem/Private/NavigationSystem.cpp
--- a/Engine/Source/Runtime/NavigationSystem/Private/NavigationSystem.cpp
+++ b/Engine/Source/Runtime/NavigationSystem/Private/NavigationSystem.cpp
@@ -2807,6 +2807,8 @@ void UNavigationSystemV1::RegisterCustomLink(INavLinkCustomInterface& CustomLink
 		}
 	}
 
+	ensureMsgf(CustomLink.GetId().IsValid(), TEXT("%hs, registering a CustomLink with an invalid id."), __FUNCTION__);
+	
 	UE_CLOG(bGenerateNewId && CustomNavLinksMap.Contains(CustomLink.GetId()), LogNavLink, Warning, TEXT("%hs New navlink ID %llu is clashing with existing ID (Owner: %s)."),
 		__FUNCTION__, CustomLink.GetId().GetId(), *GetFullNameSafe(CustomLink.GetLinkOwner()));
 	CustomNavLinksMap.Add(CustomLink.GetId(), FNavigationSystem::FCustomLinkOwnerInfo(&CustomLink));
