PR #12062: exposed GetPolyWallSegments to game code via ARecastNavMesh

diff --git a/Engine/Source/Runtime/NavigationSystem/Private/NavMesh/PImplRecastNavMesh.cpp b/Engine/Source/Runtime/NavigationSystem/Private/NavMesh/PImplRecastNavMesh.cpp
--- a/Engine/Source/Runtime/NavigationSystem/Private/NavMesh/PImplRecastNavMesh.cpp
+++ b/Engine/Source/Runtime/NavigationSystem/Private/NavMesh/PImplRecastNavMesh.cpp
@@ -2357,6 +2357,50 @@ bool FPImplRecastNavMesh::GetPolyEdges(NavNodeRef PolyID, TArray<FNavigationPort
 	return false;
 }
 
+bool FPImplRecastNavMesh::GetPolyWallSegments(NavNodeRef PolyID, const FNavigationQueryFilter& InQueryFilter, const UObject* QueryOwner, TArray<FNavigationPortalEdge>& OutNeighbors) const
+{
+	const FRecastQueryFilter* FilterImplementation = (const FRecastQueryFilter*)(InQueryFilter.GetImplementation());
+	if (FilterImplementation == nullptr)
+	{
+		UE_VLOG(NavMeshOwner, LogNavigation, Error, TEXT("%hs failed due to passed filter having NULL implementation!"), __FUNCTION__);
+		return false;
+	}
+
+	const dtQueryFilter* QueryFilter = FilterImplementation->GetAsDetourQueryFilter();
+	if (QueryFilter == nullptr)
+	{
+		UE_VLOG(NavMeshOwner, LogNavigation, Warning, TEXT("%hs failed due to QueryFilter == nullptr"), __FUNCTION__);
+		return false;
+	}
+	
+	FRecastSpeciaLinkFilter LinkFilter(FNavigationSystem::GetCurrent<UNavigationSystemV1>(NavMeshOwner->GetWorld()), QueryOwner);
+	INITIALIZE_NAVQUERY(NavQuery, InQueryFilter.GetMaxSearchNodes(), LinkFilter);
+
+	constexpr int32 MaxSegments = 64;
+	constexpr int32 ComponentsPerSegment = 6;
+	dtReal RcVertices[MaxSegments * ComponentsPerSegment] = { 0 }; // segments * ax,ay,az,bx,by,bz
+
+	int32 NumSegments = 0;
+	dtPolyRef SegmentRefs[MaxSegments] = { 0 };
+	
+	if (dtStatusSucceed(NavQuery.getPolyWallSegments(PolyID, QueryFilter, RcVertices, SegmentRefs, &NumSegments, MaxSegments)))
+	{
+		OutNeighbors.SetNum(NumSegments);
+	
+		for (int32 i = 0; i < NumSegments; ++i)
+		{
+			FNavigationPortalEdge& Edge = OutNeighbors[i];
+			Edge.Left = Recast2UnrealPoint(&RcVertices[i * ComponentsPerSegment]);
+			Edge.Right = Recast2UnrealPoint(&RcVertices[i * ComponentsPerSegment + 3]);
+			Edge.ToRef = SegmentRefs[i];
+		}
+		return true;
+	}
+
+	OutNeighbors.SetNum(0);
+	return false;
+}
+
 bool FPImplRecastNavMesh::GetPolyTileIndex(NavNodeRef PolyID, uint32& PolyIndex, uint32& TileIndex) const
 {
 	if (DetourNavMesh && PolyID)
diff --git a/Engine/Source/Runtime/NavigationSystem/Private/NavMesh/RecastNavMesh.cpp b/Engine/Source/Runtime/NavigationSystem/Private/NavMesh/RecastNavMesh.cpp
--- a/Engine/Source/Runtime/NavigationSystem/Private/NavMesh/RecastNavMesh.cpp
+++ b/Engine/Source/Runtime/NavigationSystem/Private/NavMesh/RecastNavMesh.cpp
@@ -2665,6 +2665,17 @@ bool ARecastNavMesh::GetPolyEdges(NavNodeRef PolyID, TArray<FNavigationPortalEdg
 	return bFound;
 }
 
+bool ARecastNavMesh::GetPolyWallSegments(NavNodeRef PolyID, FSharedConstNavQueryFilter Filter, const UObject* QueryOwner, TArray<FNavigationPortalEdge>& OutNeighbors) const
+{
+	bool bFound = false;
+	if (RecastNavMeshImpl)
+	{
+		bFound = RecastNavMeshImpl->GetPolyWallSegments(PolyID, GetRightFilterRef(Filter), QueryOwner, OutNeighbors);
+	}
+
+	return bFound;
+}
+
 bool ARecastNavMesh::GetClosestPointOnPoly(NavNodeRef PolyID, const FVector& TestPt, FVector& PointOnPoly) const
 {
 	return RecastNavMeshImpl && RecastNavMeshImpl->GetClosestPointOnPoly(PolyID, TestPt, PointOnPoly);
diff --git a/Engine/Source/Runtime/NavigationSystem/Public/NavMesh/PImplRecastNavMesh.h b/Engine/Source/Runtime/NavigationSystem/Public/NavMesh/PImplRecastNavMesh.h
--- a/Engine/Source/Runtime/NavigationSystem/Public/NavMesh/PImplRecastNavMesh.h
+++ b/Engine/Source/Runtime/NavigationSystem/Public/NavMesh/PImplRecastNavMesh.h
@@ -194,6 +194,8 @@ public:
 	NAVIGATIONSYSTEM_API bool GetPolyNeighbors(NavNodeRef PolyID, TArray<NavNodeRef>& Neighbors) const;
 	/** Finds all polys connected with specified one */
 	NAVIGATIONSYSTEM_API bool GetPolyEdges(NavNodeRef PolyID, TArray<FNavigationPortalEdge>& Edges) const;
+	/** Finds all wall segments for the specified polygon (walls or area borders) */
+	NAVIGATIONSYSTEM_API bool GetPolyWallSegments(NavNodeRef PolyID, const FNavigationQueryFilter& InQueryFilter, const UObject* QueryOwner, TArray<FNavigationPortalEdge>& OutNeighbors) const;
 	/** Finds closest point constrained to given poly */
 	NAVIGATIONSYSTEM_API bool GetClosestPointOnPoly(NavNodeRef PolyID, const FVector& TestPt, FVector& PointOnPoly) const;
 	/** Decode poly ID into tile index and poly index */
diff --git a/Engine/Source/Runtime/NavigationSystem/Public/NavMesh/RecastNavMesh.h b/Engine/Source/Runtime/NavigationSystem/Public/NavMesh/RecastNavMesh.h
--- a/Engine/Source/Runtime/NavigationSystem/Public/NavMesh/RecastNavMesh.h
+++ b/Engine/Source/Runtime/NavigationSystem/Public/NavMesh/RecastNavMesh.h
@@ -1433,6 +1433,9 @@ public:
 	/** Finds edges of specified poly */
 	NAVIGATIONSYSTEM_API bool GetPolyEdges(NavNodeRef PolyID, TArray<FNavigationPortalEdge>& Neighbors) const;
 
+	/** Finds all wall segments for the specified polygon (walls or area borders) */
+	NAVIGATIONSYSTEM_API bool GetPolyWallSegments(NavNodeRef PolyID, FSharedConstNavQueryFilter Filter, const UObject* QueryOwner, TArray<FNavigationPortalEdge>& OutNeighbors) const;
+
 	/** Finds closest point constrained to given poly */
 	NAVIGATIONSYSTEM_API bool GetClosestPointOnPoly(NavNodeRef PolyID, const FVector& TestPt, FVector& PointOnPoly) const;
 
