Addition of a console command to force disable navlink generation (by setting ai.nav.AllowLinkGeneration to 0)

diff --git a/Engine/Source/Runtime/NavigationSystem/Private/NavMesh/RecastNavMeshGenerator.cpp b/Engine/Source/Runtime/NavigationSystem/Private/NavMesh/RecastNavMeshGenerator.cpp
--- a/Engine/Source/Runtime/NavigationSystem/Private/NavMesh/RecastNavMeshGenerator.cpp
+++ b/Engine/Source/Runtime/NavigationSystem/Private/NavMesh/RecastNavMeshGenerator.cpp
@@ -86,6 +86,10 @@ namespace UE::NavMesh::Private
 
 	static bool bUseAsymetricBorderSizes = false;
 	static FAutoConsoleVariableRef CVarUseAsymetricBorderSizes(TEXT("ai.nav.UseAsymetricBorderSizes"), bUseAsymetricBorderSizes, TEXT("Active by default. When generating links, use asymetric tile border sizes to improve generation speed."), ECVF_Default);
+
+	static bool bAllowLinkGeneration = true;
+	static FAutoConsoleVariableRef CVarAllowLinkGeneration(TEXT("ai.nav.AllowLinkGeneration"), bAllowLinkGeneration, TEXT("Set to false to force disabling link generation."), ECVF_Default);
+
 }
 
 static FOodleDataCompression::ECompressor GNavmeshTileCacheCompressor = FOodleDataCompression::ECompressor::Mermaid;
@@ -4915,7 +4919,7 @@ void FRecastNavMeshGenerator::SetupTileConfig(const ENavigationDataResolution Ti
 	OutConfig.walkableRadius = FMath::CeilToInt(DestNavMesh->AgentRadius / CellSize);
 	OutConfig.maxStepFromWalkableSlope = OutConfig.cs * FMath::Tan(FMath::DegreesToRadians(OutConfig.walkableSlopeAngle));
 
-	UE::NavMesh::Private::ComputeConfigBorderSizes(DestNavMesh->bGenerateNavLinks, OutConfig);
+	UE::NavMesh::Private::ComputeConfigBorderSizes(IsGeneratingLinks(), OutConfig);
 
 	OutConfig.maxEdgeLen = (int32)(1200.0f / CellSize);
 
@@ -4985,7 +4989,7 @@ void FRecastNavMeshGenerator::ConfigureBuildProperties(FRecastBuildConfig& OutCo
 		}
 	}
 
-	if (DestNavMesh->bGenerateNavLinks)
+	if (IsGeneratingLinks())
 	{
 		// NavLink builder configuration
 		const FNavLinkGenerationJumpDownConfig& JumpDown = DestNavMesh->NavLinkJumpDownConfig;
@@ -5029,7 +5033,7 @@ void FRecastNavMeshGenerator::ConfigureBuildProperties(FRecastBuildConfig& OutCo
 		OutConfig.walkableHeight = 1;
 	}
 
-	OutConfig.bGenerateLinks = DestNavMesh->bGenerateNavLinks;
+	OutConfig.bGenerateLinks = IsGeneratingLinks();
 
 	const UNavigationSystemV1* NavSys = FNavigationSystem::GetCurrent<UNavigationSystemV1>(GetWorld());
 	OutConfig.AgentIndex = NavSys ? NavSys->GetSupportedAgentIndex(DestNavMesh) : 0;
@@ -5311,6 +5315,11 @@ void FRecastNavMeshGenerator::CalcPolyRefBits(ARecastNavMesh* NavMeshOwner, int3
 #endif//USE_64BIT_ADDRESS
 }
 
+bool FRecastNavMeshGenerator::IsGeneratingLinks() const
+{
+	return DestNavMesh && DestNavMesh->bGenerateNavLinks && UE::NavMesh::Private::bAllowLinkGeneration; 
+}
+
 void FRecastNavMeshGenerator::CalcNavMeshProperties(int32& MaxTiles, int32& MaxPolys)
 {
 	int32 MaxTileBits = -1;
diff --git a/Engine/Source/Runtime/NavigationSystem/Public/NavMesh/RecastNavMeshGenerator.h b/Engine/Source/Runtime/NavigationSystem/Public/NavMesh/RecastNavMeshGenerator.h
--- a/Engine/Source/Runtime/NavigationSystem/Public/NavMesh/RecastNavMeshGenerator.h
+++ b/Engine/Source/Runtime/NavigationSystem/Public/NavMesh/RecastNavMeshGenerator.h
@@ -1004,6 +1004,9 @@ public:
 
 	static NAVIGATIONSYSTEM_API void CalcPolyRefBits(ARecastNavMesh* NavMeshOwner, int32& MaxTileBits, int32& MaxPolyBits);
 
+	/** Returns true if bGenerateNavLinks is enabled and bAllowLinkGeneration is true. */ 
+	bool IsGeneratingLinks() const;
+
 protected:
 	NAVIGATIONSYSTEM_API virtual void RestrictBuildingToActiveTiles(bool InRestrictBuildingToActiveTiles);
 	
