Remove experimental FNavLinkGenerationJumpOverConfig for now

#rnx

diff --git a/Engine/Source/Runtime/NavigationSystem/Private/NavMesh/LinkGenerationConfig.cpp b/Engine/Source/Runtime/NavigationSystem/Private/NavMesh/LinkGenerationConfig.cpp
--- a/Engine/Source/Runtime/NavigationSystem/Private/NavMesh/LinkGenerationConfig.cpp
+++ b/Engine/Source/Runtime/NavigationSystem/Private/NavMesh/LinkGenerationConfig.cpp
@@ -22,15 +22,4 @@ void FNavLinkGenerationJumpDownConfig::CopyToDetourConfig(dtNavLinkBuilderJumpDo
 	}
 }
 
-void FNavLinkGenerationJumpOverConfig::CopyToDetourConfig(dtNavLinkBuilderJumpOverConfig& OutDetourConfig) const
-{
-	OutDetourConfig.enabled = bEnabled;
-	OutDetourConfig.jumpGapWidth = JumpGapWidth;
-	OutDetourConfig.jumpGapHeightTolerance = JumpGapHeightTolerance;
-	OutDetourConfig.jumpDistanceFromGapCenter = JumpDistanceFromGapCenter;
-	OutDetourConfig.jumpHeight = JumpHeight;
-	OutDetourConfig.jumpEndsHeightTolerance	= JumpEndsHeightTolerance;
-	OutDetourConfig.samplingSeparationFactor = SamplingSeparationFactor;
-	OutDetourConfig.filterDistanceThreshold = FilterDistanceThreshold;
-}
 #endif //WITH_RECAST
diff --git a/Engine/Source/Runtime/NavigationSystem/Private/NavMesh/RecastNavMesh.cpp b/Engine/Source/Runtime/NavigationSystem/Private/NavMesh/RecastNavMesh.cpp
--- a/Engine/Source/Runtime/NavigationSystem/Private/NavMesh/RecastNavMesh.cpp
+++ b/Engine/Source/Runtime/NavigationSystem/Private/NavMesh/RecastNavMesh.cpp
@@ -3508,7 +3508,6 @@ void ARecastNavMesh::PostEditChangeChainProperty(FPropertyChangedChainEvent& Pro
 		if (CategoryName == NAME_Generation)
 		{
 			static const FName NAME_NavLinkJumpDownConfig = GET_MEMBER_NAME_CHECKED(ARecastNavMesh, NavLinkJumpDownConfig);
-			static const FName NAME_NavLinkJumpOverConfig = GET_MEMBER_NAME_CHECKED(ARecastNavMesh, NavLinkJumpOverConfig);
 			
 			const FName PropName = PropertyChangedChainEvent.Property->GetFName();
 			bool bRebuild = false;
@@ -3545,7 +3544,7 @@ void ARecastNavMesh::PostEditChangeChainProperty(FPropertyChangedChainEvent& Pro
 			{
 				bRebuild = true;
 			}
-			else if (MemberProperty && (MemberProperty->GetFName() == NAME_NavLinkJumpDownConfig || MemberProperty->GetFName() == NAME_NavLinkJumpOverConfig))
+			else if (MemberProperty && (MemberProperty->GetFName() == NAME_NavLinkJumpDownConfig))
 			{
 				bRebuild = true;
 			}
diff --git a/Engine/Source/Runtime/NavigationSystem/Private/NavMesh/RecastNavMeshGenerator.cpp b/Engine/Source/Runtime/NavigationSystem/Private/NavMesh/RecastNavMeshGenerator.cpp
--- a/Engine/Source/Runtime/NavigationSystem/Private/NavMesh/RecastNavMeshGenerator.cpp
+++ b/Engine/Source/Runtime/NavigationSystem/Private/NavMesh/RecastNavMeshGenerator.cpp
@@ -4993,12 +4993,10 @@ void FRecastNavMeshGenerator::ConfigureBuildProperties(FRecastBuildConfig& OutCo
 	{
 		// NavLink builder configuration
 		const FNavLinkGenerationJumpDownConfig& JumpDown = DestNavMesh->NavLinkJumpDownConfig;
-		const FNavLinkGenerationJumpOverConfig& JumpOver = DestNavMesh->NavLinkJumpOverConfig;
 		JumpDown.CopyToDetourConfig(OutConfig.JumpDownConfig);
-		JumpOver.CopyToDetourConfig(OutConfig.JumpOverConfig);
 
 		const float JumpDownSpillDistance = JumpDown.bEnabled ? JumpDown.JumpLength - JumpDown.JumpDistanceFromEdge : 0.f;
-		const float JumpOverSpillDistance = JumpOver.bEnabled ? JumpOver.JumpDistanceFromGapCenter : 0.f;
+		constexpr float JumpOverSpillDistance = 0.f; //JumpOver.bEnabled ? JumpOver.JumpDistanceFromGapCenter : 0.f;		// @todo: jump over config is not exposed for now
 		OutConfig.LinkSpillDistance = FMath::Max(JumpDownSpillDistance, JumpOverSpillDistance);
 	}
 	
diff --git a/Engine/Source/Runtime/NavigationSystem/Public/NavMesh/LinkGenerationConfig.h b/Engine/Source/Runtime/NavigationSystem/Public/NavMesh/LinkGenerationConfig.h
--- a/Engine/Source/Runtime/NavigationSystem/Public/NavMesh/LinkGenerationConfig.h
+++ b/Engine/Source/Runtime/NavigationSystem/Public/NavMesh/LinkGenerationConfig.h
@@ -79,47 +79,3 @@ struct FNavLinkGenerationJumpDownConfig
 #endif //WITH_RECAST
 };
 
-/** Experimental configuration for generated jump over links. */
-USTRUCT()
-struct FNavLinkGenerationJumpOverConfig
-{
-	GENERATED_BODY()
-
-	/** Should this config be used to generate links. */
-	UPROPERTY(EditAnywhere, Config, Category = Settings)
-	bool bEnabled = true;
-	
-	/* Maximum jumpable gap size used when matching edges to jump over. */
-	UPROPERTY(EditAnywhere, Config, Category = Settings, meta=(Units=cm, UIMin=0, ClampMin=0))
-	float JumpGapWidth = 200.f; 
-
-	/** Vertical tolerance used when matching edges to jump over. */
-	UPROPERTY(EditAnywhere, Config, Category = Settings, meta=(Units=cm, UIMin=0, ClampMin=0))
-	float JumpGapHeightTolerance = 100.f; 
-	
-	/** How far from the center of the gap is the jump started. */
-	UPROPERTY(EditAnywhere, Config, Category = Settings, meta=(Units=cm, UIMin=0, ClampMin=0))
-	float JumpDistanceFromGapCenter = 100.f; 
-
-	/** Height at the top of the jump trajectory. */
-	UPROPERTY(EditAnywhere, Config, Category = Settings, meta=(Units=cm, UIMin=0, ClampMin=0))
-	float JumpHeight = 100.f; 
-
-	/** Tolerance at both ends of the jump to find ground. */
-	UPROPERTY(EditAnywhere, Config, Category = Settings, meta=(Units=cm, UIMin=0, ClampMin=0))
-	float JumpEndsHeightTolerance = 80.f;
-
-	/** Value multiplied by CellSize to find the distance between sampling trajectories. Default is 1. */
-    /*  Larger values improve generation speed but might introduce sampling errors.  */
-	UPROPERTY(EditAnywhere, Config, Category = Settings, meta=(UIMin=1, ClampMin=1))
-	float SamplingSeparationFactor = 1.f;
-
-	/** When filtering similar links, distance used to compare between segment endpoints to match similar links. Use greater distance for more filtering (0 to deactivate filtering). */
-	UPROPERTY(EditAnywhere, Config, Category = Settings, meta=(Units=cm, UIMin=0, ClampMin=0))
-	float FilterDistanceThreshold = 80.f;
-
-#if WITH_RECAST
-	/** Copy configuration to dtNavLinkBuilderJumpOverConfig. */
-	void CopyToDetourConfig(dtNavLinkBuilderJumpOverConfig& OutDetourConfig) const;
-#endif //WITH_RECAST
-};
diff --git a/Engine/Source/Runtime/NavigationSystem/Public/NavMesh/RecastNavMesh.h b/Engine/Source/Runtime/NavigationSystem/Public/NavMesh/RecastNavMesh.h
--- a/Engine/Source/Runtime/NavigationSystem/Public/NavMesh/RecastNavMesh.h
+++ b/Engine/Source/Runtime/NavigationSystem/Public/NavMesh/RecastNavMesh.h
@@ -986,10 +986,6 @@ protected:
 	/** Experimental: paramters for building jump down links.  */
 	UPROPERTY(EditAnywhere, Category=Generation, config)
 	FNavLinkGenerationJumpDownConfig NavLinkJumpDownConfig;
-
-	/** Experimental: paramters for building jump over links.  */
-	UPROPERTY(EditAnywhere, Category=Generation, config)
-	FNavLinkGenerationJumpOverConfig NavLinkJumpOverConfig;
 	
 private:
 	/** Cache rasterized voxels instead of just collision vertices/indices in navigation octree */
