Fix cylinder collision wrongly offset when used as navcollision on dynamic obstacle.

diff --git a/Engine/Source/Runtime/NavigationSystem/Private/NavCollision.cpp b/Engine/Source/Runtime/NavigationSystem/Private/NavCollision.cpp
--- a/Engine/Source/Runtime/NavigationSystem/Private/NavCollision.cpp
+++ b/Engine/Source/Runtime/NavigationSystem/Private/NavCollision.cpp
@@ -292,7 +292,7 @@ void UNavCollision::GatherCollision()
 		const FNavCollisionCylinder& CylinderInfo = CylinderCollision[Idx];
 
 		FKSphylElem SphylElem(CylinderInfo.Radius, CylinderInfo.Height);
-		SphylElem.SetTransform(FTransform(CylinderInfo.Offset));
+		SphylElem.SetTransform(FTransform(CylinderInfo.Offset + FVector(0.f, 0.f, 0.5f*CylinderInfo.Height)));
 
 		SimpleGeom.SphylElems.Add(SphylElem);
 	}
