[NavMesh] Fix crash occuring when undoing removal of ARecastNavmesh

diff --git a/Engine/Source/Runtime/NavigationSystem/Private/NavMesh/RecastNavMesh.cpp b/Engine/Source/Runtime/NavigationSystem/Private/NavMesh/RecastNavMesh.cpp
--- a/Engine/Source/Runtime/NavigationSystem/Private/NavMesh/RecastNavMesh.cpp
+++ b/Engine/Source/Runtime/NavigationSystem/Private/NavMesh/RecastNavMesh.cpp
@@ -3647,6 +3647,18 @@ void ARecastNavMesh::PostEditChangeProperty(FPropertyChangedEvent& PropertyChang
 	}
 }
 
+void ARecastNavMesh::PostEditUndo()
+{
+	// ANavigationData will register again this navdata in Super::PostEditUndo(). 
+    // Since it was removed, we need to recreate it's RecastNavMeshImpl.
+	if (RecastNavMeshImpl == nullptr)
+	{
+		RecastNavMeshImpl = new FPImplRecastNavMesh(this);
+	}
+	
+	Super::PostEditUndo();
+}
+
 #endif // WITH_EDITOR
 
 bool ARecastNavMesh::NeedsRebuild() const
diff --git a/Engine/Source/Runtime/NavigationSystem/Public/NavMesh/RecastNavMesh.h b/Engine/Source/Runtime/NavigationSystem/Public/NavMesh/RecastNavMesh.h
--- a/Engine/Source/Runtime/NavigationSystem/Public/NavMesh/RecastNavMesh.h
+++ b/Engine/Source/Runtime/NavigationSystem/Public/NavMesh/RecastNavMesh.h
@@ -1080,6 +1080,7 @@ public:
 #if WITH_EDITOR
 	NAVIGATIONSYSTEM_API virtual void PostEditChangeChainProperty( struct FPropertyChangedChainEvent& PropertyChangedChainEvent) override;
 	NAVIGATIONSYSTEM_API virtual void PostEditChangeProperty( struct FPropertyChangedEvent& PropertyChangedEvent) override;
+	NAVIGATIONSYSTEM_API virtual void PostEditUndo() override;
 #endif // WITH_EDITOR
 	//~ End UObject Interface
 
