Remove deprecated FRecastNavMeshGenerationProperties

diff --git a/Engine/Source/Runtime/NavigationSystem/Private/NavMesh/RecastNavMesh.cpp b/Engine/Source/Runtime/NavigationSystem/Private/NavMesh/RecastNavMesh.cpp
--- a/Engine/Source/Runtime/NavigationSystem/Private/NavMesh/RecastNavMesh.cpp
+++ b/Engine/Source/Runtime/NavigationSystem/Private/NavMesh/RecastNavMesh.cpp
@@ -423,69 +423,6 @@ namespace FNavMeshConfig
 
 #endif // WITH_RECAST
 
-// Deprecated
-FRecastNavMeshGenerationProperties::FRecastNavMeshGenerationProperties()
-{
-	TilePoolSize = 1024;
-	TileSizeUU = 988.f;
-	CellSize = 19;
-	CellHeight = 10;
-	AgentRadius = 34.f;
-	AgentHeight = 144.f;
-	AgentMaxSlope = 44.f;
-	AgentMaxStepHeight = 35.f;
-	MinRegionArea = 0.f;
-	MergeRegionSize = 400.f;
-	MaxSimplificationError = 1.3f;	// from RecastDemo
-	TileNumberHardLimit = 1 << 20;
-	RegionPartitioning = ERecastPartitioning::Watershed;
-	LayerPartitioning = ERecastPartitioning::Watershed;
-	RegionChunkSplits = 2;
-	LayerChunkSplits = 2;
-	bSortNavigationAreasByCost = true;
-	bPerformVoxelFiltering = true;
-	bMarkLowHeightAreas = false;
-	bUseExtraTopCellWhenMarkingAreas = true;
-	bFilterLowSpanSequences = false;
-	bFilterLowSpanFromTileCache = false;
-	bFixedTilePoolSize = false;
-	bIsWorldPartitioned = false;
-}
-
-#if WITH_RECAST
-
-// Deprecated
-FRecastNavMeshGenerationProperties::FRecastNavMeshGenerationProperties(const ARecastNavMesh& RecastNavMesh)
-{
-	TilePoolSize = RecastNavMesh.TilePoolSize;
-	TileSizeUU = RecastNavMesh.TileSizeUU;
-	
-	CellSize = RecastNavMesh.GetCellSize(ENavigationDataResolution::Default);
-	CellHeight = RecastNavMesh.GetCellHeight(ENavigationDataResolution::Default);
-	AgentRadius = RecastNavMesh.AgentRadius;
-	AgentHeight = RecastNavMesh.AgentHeight;
-	AgentMaxSlope = RecastNavMesh.AgentMaxSlope;
-	AgentMaxStepHeight = RecastNavMesh.GetAgentMaxStepHeight(ENavigationDataResolution::Default); //FRecastNavMeshGenerationProperties is getting deprecated 
-	MinRegionArea = RecastNavMesh.MinRegionArea;
-	MergeRegionSize = RecastNavMesh.MergeRegionSize;
-	MaxSimplificationError = RecastNavMesh.MaxSimplificationError;
-	TileNumberHardLimit = RecastNavMesh.TileNumberHardLimit;
-	RegionPartitioning = RecastNavMesh.RegionPartitioning;
-	LayerPartitioning = RecastNavMesh.LayerPartitioning;
-	RegionChunkSplits = RecastNavMesh.RegionChunkSplits;
-	LayerChunkSplits = RecastNavMesh.LayerChunkSplits;
-	bSortNavigationAreasByCost = RecastNavMesh.bSortNavigationAreasByCost;
-	bPerformVoxelFiltering = RecastNavMesh.bPerformVoxelFiltering;
-	bMarkLowHeightAreas = RecastNavMesh.bMarkLowHeightAreas;
-	bUseExtraTopCellWhenMarkingAreas = RecastNavMesh.bUseExtraTopCellWhenMarkingAreas;
-	bFilterLowSpanSequences = RecastNavMesh.bFilterLowSpanSequences;
-	bFilterLowSpanFromTileCache = RecastNavMesh.bFilterLowSpanFromTileCache;
-	bFixedTilePoolSize = RecastNavMesh.bFixedTilePoolSize;
-	bIsWorldPartitioned = RecastNavMesh.bIsWorldPartitioned;
-}
-
-#endif // WITH_RECAST
-
 FRecastNavMeshTileGenerationDebug::FRecastNavMeshTileGenerationDebug()
 {
 	bEnabled = false;
diff --git a/Engine/Source/Runtime/NavigationSystem/Public/NavMesh/RecastNavMesh.h b/Engine/Source/Runtime/NavigationSystem/Public/NavMesh/RecastNavMesh.h
--- a/Engine/Source/Runtime/NavigationSystem/Public/NavMesh/RecastNavMesh.h
+++ b/Engine/Source/Runtime/NavigationSystem/Public/NavMesh/RecastNavMesh.h
@@ -331,122 +331,6 @@ struct FNavigationWallEdge
 };
 #endif //WITH_RECAST
 
-USTRUCT(meta=(Deprecated = "5.2"))
-struct FRecastNavMeshGenerationProperties
-{
-	GENERATED_BODY()
-
-	/** maximum number of tiles NavMesh can hold */
-	UPROPERTY(EditAnywhere, Category = Generation, meta = (editcondition = "bFixedTilePoolSize"))
-	int32 TilePoolSize;
-
-	/** size of single tile, expressed in uu */
-	UPROPERTY(EditAnywhere, Category = Generation, meta = (ClampMin = "300.0"))
-	float TileSizeUU;
-
-	/** horizontal size of voxelization cell */
-	UPROPERTY(EditAnywhere, Category = Generation, meta = (ClampMin = "1.0", ClampMax = "1024.0"))
-	float CellSize;
-
-	/** vertical size of voxelization cell */
-	UPROPERTY(EditAnywhere, Category = Generation, meta = (ClampMin = "1.0", ClampMax = "1024.0"))
-	float CellHeight;
-
-	/** Radius of largest agent that can freely traverse the generated navmesh */
-	UPROPERTY(EditAnywhere, Category = Generation, meta = (ClampMin = "0.0"))
-	float AgentRadius;
-
-	/** Size of the tallest agent that will path with this navmesh. */
-	UPROPERTY(EditAnywhere, Category = Generation, meta = (ClampMin = "0.0"))
-	float AgentHeight;
-
-	/* The maximum slope (angle) that the agent can move on. */
-	UPROPERTY(EditAnywhere, Category = Generation, meta = (ClampMin = "0.0", ClampMax = "89.0", UIMin = "0.0", UIMax = "89.0"))
-	float AgentMaxSlope;
-
-	/** Largest vertical step the agent can perform */
-	UPROPERTY(EditAnywhere, Category = Generation, meta = (ClampMin = "0.0"))
-	float AgentMaxStepHeight;
-
-	/* The minimum dimension of area. Areas smaller than this will be discarded */
-	UPROPERTY(EditAnywhere, Category = Generation, meta = (ClampMin = "0.0"))
-	float MinRegionArea;
-
-	/* The size limit of regions to be merged with bigger regions (watershed partitioning only) */
-	UPROPERTY(EditAnywhere, Category = Generation, meta = (ClampMin = "0.0"))
-	float MergeRegionSize;
-
-	/** How much navigable shapes can get simplified - the higher the value the more freedom */
-	UPROPERTY(EditAnywhere, Category = Generation, meta = (ClampMin = "0.0"))
-	float MaxSimplificationError;
-
-	/** Absolute hard limit to number of navmesh tiles. Be very, very careful while modifying it while
-	*	having big maps with navmesh. A single, empty tile takes 176 bytes and empty tiles are
-	*	allocated up front (subject to change, but that's where it's at now)
-	*	@note TileNumberHardLimit is always rounded up to the closest power of 2 */
-	UPROPERTY(EditAnywhere, Category = Generation, meta = (ClampMin = "1", UIMin = "1"), AdvancedDisplay)
-	int32 TileNumberHardLimit;
-
-	/** partitioning method for creating navmesh polys */
-	UPROPERTY(EditAnywhere, Category = Generation, AdvancedDisplay)
-	TEnumAsByte<ERecastPartitioning::Type> RegionPartitioning;
-
-	/** partitioning method for creating tile layers */
-	UPROPERTY(EditAnywhere, Category = Generation, AdvancedDisplay)
-	TEnumAsByte<ERecastPartitioning::Type> LayerPartitioning;
-
-	/** number of chunk splits (along single axis) used for region's partitioning: ChunkyMonotone */
-	UPROPERTY(EditAnywhere, Category = Generation, AdvancedDisplay)
-	int32 RegionChunkSplits;
-
-	/** number of chunk splits (along single axis) used for layer's partitioning: ChunkyMonotone */
-	UPROPERTY(EditAnywhere, Category = Generation, AdvancedDisplay)
-	int32 LayerChunkSplits;
-
-	/** Controls whether Navigation Areas will be sorted by cost before application
-	 *	to navmesh during navmesh generation. This is relevant when there are
-	 *	areas overlapping and we want to have area cost express area relevancy
-	 *	as well. Setting it to true will result in having area sorted by cost,
-	 *	but it will also increase navmesh generation cost a bit */
-	UPROPERTY(EditAnywhere, Category = Generation)
-	uint32 bSortNavigationAreasByCost : 1;
-
-	/** controls whether voxel filtering will be applied (via FRecastTileGenerator::ApplyVoxelFilter).
-	 *	Results in generated navmesh better fitting navigation bounds, but hits (a bit) generation performance */
-	UPROPERTY(EditAnywhere, Category = Generation, AdvancedDisplay)
-	uint32 bPerformVoxelFiltering : 1;
-
-	/** mark areas with insufficient free height above instead of cutting them out (accessible only for area modifiers using replace mode) */
-	UPROPERTY(EditAnywhere, Category = Generation, AdvancedDisplay)
-	uint32 bMarkLowHeightAreas : 1;
-
-	/** Expand the top of the area nav modifier's bounds by one cell height when applying to the navmesh. 
-	    If unset, navmesh on top of surfaces might not be marked by marking bounds flush with top surfaces (since navmesh is generated slightly above collision, depending on cell height). */
-	UPROPERTY(EditAnywhere, Category = Generation, AdvancedDisplay)
-	uint32 bUseExtraTopCellWhenMarkingAreas : 1;
-
-	/** if set, only single low height span will be allowed under valid one */
-	UPROPERTY(EditAnywhere, Category = Generation, AdvancedDisplay)
-	uint32 bFilterLowSpanSequences : 1;
-
-	/** if set, only low height spans with corresponding area modifier will be stored in tile cache (reduces memory, can't modify without full tile rebuild) */
-	UPROPERTY(EditAnywhere, Category = Generation, AdvancedDisplay)
-	uint32 bFilterLowSpanFromTileCache : 1;
-
-	/** if true, the NavMesh will allocate fixed size pool for tiles, should be enabled to support streaming */
-	UPROPERTY(EditAnywhere, Category = Generation)
-	uint32 bFixedTilePoolSize : 1;
-
-	/* In a world partitioned map, is this navmesh using world partitioning */
-	UPROPERTY(EditAnywhere, Category = Generation)
-	uint32 bIsWorldPartitioned : 1;
-	
-	NAVIGATIONSYSTEM_API FRecastNavMeshGenerationProperties();
-#if WITH_RECAST
-	NAVIGATIONSYSTEM_API FRecastNavMeshGenerationProperties(const ARecastNavMesh& RecastNavMesh);
-#endif // WITH_RECAST
-};
-
 UENUM()
 enum class EHeightFieldRenderMode : uint8
 {
