[NavMesh] Don't save tiles during undo redo operations

diff --git a/Engine/Source/Runtime/NavigationSystem/Private/NavMesh/PImplRecastNavMesh.cpp b/Engine/Source/Runtime/NavigationSystem/Private/NavMesh/PImplRecastNavMesh.cpp
--- a/Engine/Source/Runtime/NavigationSystem/Private/NavMesh/PImplRecastNavMesh.cpp
+++ b/Engine/Source/Runtime/NavigationSystem/Private/NavMesh/PImplRecastNavMesh.cpp
@@ -512,7 +512,7 @@ void FPImplRecastNavMesh::Serialize( FArchive& Ar, int32 NavMeshVersion )
 	int32 NumTiles = 0;
 	TArray<FNavTileRef> TilesToSave;
 
-	if (Ar.IsSaving())
+	if (Ar.IsSaving() && !Ar.IsTransacting()) // Do not save tiles during transactions(i.e. undo/redo)
 	{
 		TilesToSave.Reserve(DetourNavMesh->getMaxTiles());
 		
