Non-determinism fix in DDC data for NavCollision supplied by [REDACTED_PERSON] via UDN.

diff --git a/Engine/Source/Runtime/NavigationSystem/Private/NavCollision.cpp b/Engine/Source/Runtime/NavigationSystem/Private/NavCollision.cpp
--- a/Engine/Source/Runtime/NavigationSystem/Private/NavCollision.cpp
+++ b/Engine/Source/Runtime/NavigationSystem/Private/NavCollision.cpp
@@ -101,7 +101,7 @@ public:
 
 	virtual FString GetPluginSpecificCacheKeySuffix() const override
 	{
-		const uint16 Version = 14;
+		const uint16 Version = 15;
 
 		return FString::Printf( TEXT("%s_%s_%s_%hu")
 			, *Format.ToString()
@@ -312,7 +312,7 @@ void UNavCollision::ClearCollision()
 	ConvexCollision.VertexBuffer.Reset();
 	ConvexCollision.IndexBuffer.Reset();
 	ConvexShapeIndices.Reset();
-	Bounds = FBox();
+	Bounds = FBox(ForceInitToZero);
 
 	bHasConvexGeometry = false;
 }
