Fix IntCastChecked assert caused by invalid location, tiny tiles or giant invoker radius.

diff --git a/Engine/Source/Runtime/NavigationSystem/Private/NavMesh/RecastNavMesh.cpp b/Engine/Source/Runtime/NavigationSystem/Private/NavMesh/RecastNavMesh.cpp
--- a/Engine/Source/Runtime/NavigationSystem/Private/NavMesh/RecastNavMesh.cpp
+++ b/Engine/Source/Runtime/NavigationSystem/Private/NavMesh/RecastNavMesh.cpp
@@ -3893,10 +3893,10 @@ void ARecastNavMesh::UpdateActiveTiles(const TArray<FNavigationInvokerRaw>& Invo
 			const  FVector::FReal TileCenterDistanceToRemoveSq = FMath::Square(TileDim * UE_SQRT_2 / 2 + Invoker.RadiusMax);
 			const  FVector::FReal TileCenterDistanceToAddSq = FMath::Square(TileDim * UE_SQRT_2 / 2 + Invoker.RadiusMin);
 
-			const int32 MinTileX = IntCastChecked<int32>(FMath::FloorToInt((InvokerRelativeLocation.X - Invoker.RadiusMax) / TileDim));
-			const int32 MaxTileX = IntCastChecked<int32>(FMath::CeilToInt((InvokerRelativeLocation.X + Invoker.RadiusMax) / TileDim));
-			const int32 MinTileY = IntCastChecked<int32>(FMath::FloorToInt((InvokerRelativeLocation.Y - Invoker.RadiusMax) / TileDim));
-			const int32 MaxTileY = IntCastChecked<int32>(FMath::CeilToInt((InvokerRelativeLocation.Y + Invoker.RadiusMax) / TileDim));
+			const int32 MinTileX = static_cast<int32>(FMath::Clamp(FMath::FloorToInt((InvokerRelativeLocation.X - Invoker.RadiusMax) / TileDim), TNumericLimits<int32>::Min(), TNumericLimits<int32>::Max()));
+			const int32 MaxTileX = static_cast<int32>(FMath::Clamp(FMath::CeilToInt((InvokerRelativeLocation.X + Invoker.RadiusMax) / TileDim),  TNumericLimits<int32>::Min(), TNumericLimits<int32>::Max()));
+			const int32 MinTileY = static_cast<int32>(FMath::Clamp(FMath::FloorToInt((InvokerRelativeLocation.Y - Invoker.RadiusMax) / TileDim), TNumericLimits<int32>::Min(), TNumericLimits<int32>::Max()));
+			const int32 MaxTileY = static_cast<int32>(FMath::Clamp(FMath::CeilToInt((InvokerRelativeLocation.Y + Invoker.RadiusMax) / TileDim),  TNumericLimits<int32>::Min(), TNumericLimits<int32>::Max()));
 
 			for (int32 X = MinTileX; X <= MaxTileX; ++X)
 			{
diff --git a/Engine/Source/Runtime/NavigationSystem/Public/NavigationInvokerComponent.h b/Engine/Source/Runtime/NavigationSystem/Public/NavigationInvokerComponent.h
--- a/Engine/Source/Runtime/NavigationSystem/Public/NavigationInvokerComponent.h
+++ b/Engine/Source/Runtime/NavigationSystem/Public/NavigationInvokerComponent.h
@@ -21,10 +21,10 @@ class UNavigationInvokerComponent : public UActorComponent
 
 protected:
 
-	UPROPERTY(EditAnywhere, Category = Navigation, meta = (ClampMin = "0.1", UIMin = "0.1"))
+	UPROPERTY(EditAnywhere, Category = Navigation, meta = (ClampMin = "0.1", ClampMax = "6400000", UIMin = "0.1", UIMax = "6400000"))
 	float TileGenerationRadius;
 
-	UPROPERTY(EditAnywhere, Category = Navigation, meta = (ClampMin = "0.1", UIMin = "0.1"))
+	UPROPERTY(EditAnywhere, Category = Navigation, meta = (ClampMin = "0.1", ClampMax = "6400000", UIMin = "0.1", UIMax = "6400000"))
 	float TileRemovalRadius;
 
 	/** restrict navigation generation to specific agents */
