[NavMesh] Add ensure for TileSize != 0 and for the range of the tile coordinates

diff --git a/Engine/Source/Runtime/NavigationSystem/Private/NavMesh/RecastNavMesh.cpp b/Engine/Source/Runtime/NavigationSystem/Private/NavMesh/RecastNavMesh.cpp
--- a/Engine/Source/Runtime/NavigationSystem/Private/NavMesh/RecastNavMesh.cpp
+++ b/Engine/Source/Runtime/NavigationSystem/Private/NavMesh/RecastNavMesh.cpp
@@ -3843,6 +3843,10 @@ void ARecastNavMesh::UpdateActiveTiles(const TArray<FNavigationInvokerRaw>& Invo
 	const FRecastBuildConfig& Config = MyGenerator->GetConfig();
 	const FVector NavmeshOrigin = Recast2UnrealPoint(NavParams->orig);
 	const FVector::FReal TileDim = Config.GetTileSizeUU();
+	if (!ensureMsgf(TileDim > 0, TEXT("TileSize can't be 0. Check the navmesh configuration.")))
+	{
+		return;
+	}
 
 	TSet<FIntPoint>& OldActiveSet = UpdateActiveTilesWorkingMem.OldActiveSet;
 	TArray<FNavMeshDirtyTileElement>& TilesInMinDistance = UpdateActiveTilesWorkingMem.TilesInMinDistance;
@@ -3893,10 +3897,19 @@ void ARecastNavMesh::UpdateActiveTiles(const TArray<FNavigationInvokerRaw>& Invo
 			const  FVector::FReal TileCenterDistanceToRemoveSq = FMath::Square(TileDim * UE_SQRT_2 / 2 + Invoker.RadiusMax);
 			const  FVector::FReal TileCenterDistanceToAddSq = FMath::Square(TileDim * UE_SQRT_2 / 2 + Invoker.RadiusMin);
 
-			const int32 MinTileX = static_cast<int32>(FMath::Clamp(FMath::FloorToInt((InvokerRelativeLocation.X - Invoker.RadiusMax) / TileDim), TNumericLimits<int32>::Min(), TNumericLimits<int32>::Max()));
-			const int32 MaxTileX = static_cast<int32>(FMath::Clamp(FMath::CeilToInt((InvokerRelativeLocation.X + Invoker.RadiusMax) / TileDim),  TNumericLimits<int32>::Min(), TNumericLimits<int32>::Max()));
-			const int32 MinTileY = static_cast<int32>(FMath::Clamp(FMath::FloorToInt((InvokerRelativeLocation.Y - Invoker.RadiusMax) / TileDim), TNumericLimits<int32>::Min(), TNumericLimits<int32>::Max()));
-			const int32 MaxTileY = static_cast<int32>(FMath::Clamp(FMath::CeilToInt((InvokerRelativeLocation.Y + Invoker.RadiusMax) / TileDim),  TNumericLimits<int32>::Min(), TNumericLimits<int32>::Max()));
+			const int64 MinX = FMath::FloorToInt((InvokerRelativeLocation.X - Invoker.RadiusMax) / TileDim);
+			const int64 MaxX = FMath::CeilToInt((InvokerRelativeLocation.X + Invoker.RadiusMax) / TileDim);
+			const int64 MinY = FMath::FloorToInt((InvokerRelativeLocation.Y - Invoker.RadiusMax) / TileDim);
+			const int64 MaxY = FMath::CeilToInt((InvokerRelativeLocation.Y + Invoker.RadiusMax) / TileDim);
+			ensureMsgf(IntFitsIn<int32>(MinX), TEXT("MinX %lld is out of tile coordinate range. Location: %f, Radius %f, TileSize: %f"), MinX, InvokerRelativeLocation.X, Invoker.RadiusMax, TileDim);
+			ensureMsgf(IntFitsIn<int32>(MaxX), TEXT("MaxX %lld is out of tile coordinate range. Location: %f, Radius %f, TileSize: %f"), MaxX, InvokerRelativeLocation.X, Invoker.RadiusMax, TileDim);
+			ensureMsgf(IntFitsIn<int32>(MinY), TEXT("MinY %lld is out of tile coordinate range. Location: %f, Radius %f, TileSize: %f"), MinY, InvokerRelativeLocation.Y, Invoker.RadiusMax, TileDim);
+			ensureMsgf(IntFitsIn<int32>(MaxY), TEXT("MaxY %lld is out of tile coordinate range. Location: %f, Radius %f, TileSize: %f"), MaxY, InvokerRelativeLocation.Y, Invoker.RadiusMax, TileDim);
+			
+			const int32 MinTileX = static_cast<int32>(FMath::Clamp(MinX, TNumericLimits<int32>::Min(), TNumericLimits<int32>::Max()));
+			const int32 MaxTileX = static_cast<int32>(FMath::Clamp(MaxX, TNumericLimits<int32>::Min(), TNumericLimits<int32>::Max()));
+			const int32 MinTileY = static_cast<int32>(FMath::Clamp(MinY, TNumericLimits<int32>::Min(), TNumericLimits<int32>::Max()));
+			const int32 MaxTileY = static_cast<int32>(FMath::Clamp(MaxY, TNumericLimits<int32>::Min(), TNumericLimits<int32>::Max()));
 
 			for (int32 X = MinTileX; X <= MaxTileX; ++X)
 			{
