Change NavigationData::Raycast signature to know if the ray end is located on the ray corridor

diff --git a/Engine/Source/Runtime/AIModule/Private/BehaviorTree/Decorators/BTDecorator_DoesPathExist.cpp b/Engine/Source/Runtime/AIModule/Private/BehaviorTree/Decorators/BTDecorator_DoesPathExist.cpp
--- a/Engine/Source/Runtime/AIModule/Private/BehaviorTree/Decorators/BTDecorator_DoesPathExist.cpp
+++ b/Engine/Source/Runtime/AIModule/Private/BehaviorTree/Decorators/BTDecorator_DoesPathExist.cpp
@@ -5,7 +5,6 @@
 #include "GameFramework/Actor.h"
 #include "NavigationSystem.h"
 #include "NavFilters/NavigationQueryFilter.h"
-#include "NavMesh/RecastNavMesh.h"
 #include "BehaviorTree/BlackboardComponent.h"
 #include "AIController.h"
 
@@ -75,10 +74,10 @@ bool UBTDecorator_DoesPathExist::CalculateRawConditionValue(UBehaviorTreeCompone
 			const EPathExistanceQueryType::Type QueryType = PathQueryType.GetValue<EPathExistanceQueryType::Type>(OwnerComp);
 			if (QueryType == EPathExistanceQueryType::NavmeshRaycast2D)
 			{
-#if WITH_RECAST
-				const ARecastNavMesh* RecastNavMesh = Cast<const ARecastNavMesh>(NavData);
-				bHasPath = RecastNavMesh && RecastNavMesh->IsSegmentOnNavmesh(PointA, PointB, QueryFilter);
-#endif
+				FVector HitLocation;
+				FNavigationRaycastAdditionalResults AdditionalResults;
+				const bool bDidHit = NavData->Raycast(PointA, PointB, HitLocation, &AdditionalResults, QueryFilter);
+				bHasPath = !bDidHit && AdditionalResults.bIsRayEndInCorridor;
 			}
 			else
 			{
diff --git a/Engine/Source/Runtime/Engine/Classes/AI/Navigation/NavigationTypes.h b/Engine/Source/Runtime/Engine/Classes/AI/Navigation/NavigationTypes.h
--- a/Engine/Source/Runtime/Engine/Classes/AI/Navigation/NavigationTypes.h
+++ b/Engine/Source/Runtime/Engine/Classes/AI/Navigation/NavigationTypes.h
@@ -613,9 +613,12 @@ struct FNavigationRaycastWork : FRayStartEnd
 	/** depending on bDidHit HitLocation contains either actual hit location or RayEnd*/
 	FNavLocation HitLocation;
 	bool bDidHit;
+	/** when bDidHit is false, bIsRayEndInCorridor indicates if the projection of RayEnd is located in the corridor explored from the ray. 
+    *  When bIsRayEndInCorridor is false, it means that RayEnd failed to project to the NavigationData or on a navigation node that is not part of the explored corridor (e.g. different height) */
+	bool bIsRayEndInCorridor;
 
 	FNavigationRaycastWork(const FVector& InRayStart, const FVector& InRayEnd)
-		: FRayStartEnd(InRayStart, InRayEnd), HitLocation(InRayEnd), bDidHit(false)
+		: FRayStartEnd(InRayStart, InRayEnd), HitLocation(InRayEnd), bDidHit(false), bIsRayEndInCorridor(false)
 	{}
 };
 
diff --git a/Engine/Source/Runtime/NavigationSystem/Private/AbstractNavData.cpp b/Engine/Source/Runtime/NavigationSystem/Private/AbstractNavData.cpp
--- a/Engine/Source/Runtime/NavigationSystem/Private/AbstractNavData.cpp
+++ b/Engine/Source/Runtime/NavigationSystem/Private/AbstractNavData.cpp
@@ -35,7 +35,7 @@ AAbstractNavData::AAbstractNavData(const FObjectInitializer& ObjectInitializer)
 		TestPathImplementation = TestPathAbstract;
 		TestHierarchicalPathImplementation = TestPathAbstract;
 
-		RaycastImplementation = RaycastAbstract;
+		RaycastImplementationWithAdditionalResults = RaycastAbstract;
 
 		DefaultQueryFilter->SetFilterType<FAbstractQueryFilter>();
 	}
@@ -85,6 +85,11 @@ bool AAbstractNavData::TestPathAbstract(const FNavAgentProperties& AgentProperti
 }
 
 bool AAbstractNavData::RaycastAbstract(const ANavigationData* NavDataInstance, const FVector& RayStart, const FVector& RayEnd, FVector& HitLocation, FSharedConstNavQueryFilter QueryFilter, const UObject* Querier)
+{
+	return RaycastAbstract(NavDataInstance, RayStart, RayEnd, HitLocation, nullptr, QueryFilter, Querier);
+}
+
+bool AAbstractNavData::RaycastAbstract(const ANavigationData* NavDataInstance, const FVector& RayStart, const FVector& RayEnd, FVector& HitLocation, FNavigationRaycastAdditionalResults* AdditionalResults, FSharedConstNavQueryFilter QueryFilter, const UObject* Querier)
 {
 	return false;
 }
diff --git a/Engine/Source/Runtime/NavigationSystem/Private/NavMesh/NavTestRenderingComponent.cpp b/Engine/Source/Runtime/NavigationSystem/Private/NavMesh/NavTestRenderingComponent.cpp
--- a/Engine/Source/Runtime/NavigationSystem/Private/NavMesh/NavTestRenderingComponent.cpp
+++ b/Engine/Source/Runtime/NavigationSystem/Private/NavMesh/NavTestRenderingComponent.cpp
@@ -126,7 +126,7 @@ void FNavTestSceneProxy::GetDynamicMeshElements(const TArray<const FSceneView*>&
 				if (NavTestActor->bDrawRaycastToQueryTargetActor && NavTestActor->QueryTargetActor)
 				{
 					const FVector QueryTargetLocation = NavTestActor->QueryTargetActor->GetActorLocation();
-					PDI->DrawLine(ActorLocation, QueryTargetLocation, NavTestActor->bRaycastToQueryTargetActorResult ? FColor::Red : FColor::Green, SDPG_World, 2.5);
+					PDI->DrawLine(ActorLocation, QueryTargetLocation, NavTestActor->bRaycastToQueryTargetActorResult ? FColor::Red : NavTestActor->bRaycastToQueryTargetEndsInCorridor ? FColor::Green : FColor::Orange, SDPG_World, 2.5);
 
 					if (FNavigationSystem::IsValidLocation(NavTestActor->RaycastHitLocation))
 					{
diff --git a/Engine/Source/Runtime/NavigationSystem/Private/NavMesh/RecastNavMesh.cpp b/Engine/Source/Runtime/NavigationSystem/Private/NavMesh/RecastNavMesh.cpp
--- a/Engine/Source/Runtime/NavigationSystem/Private/NavMesh/RecastNavMesh.cpp
+++ b/Engine/Source/Runtime/NavigationSystem/Private/NavMesh/RecastNavMesh.cpp
@@ -512,7 +512,7 @@ ARecastNavMesh::ARecastNavMesh(const FObjectInitializer& ObjectInitializer)
 		TestPathImplementation = TestPath;
 		TestHierarchicalPathImplementation = TestHierarchicalPath;
 
-		RaycastImplementation = NavMeshRaycast;
+		RaycastImplementationWithAdditionalResults = NavMeshRaycast;
 
 		RecastNavMeshImpl = new FPImplRecastNavMesh(this);
 	
@@ -3305,7 +3305,12 @@ bool ARecastNavMesh::NavMeshRaycast(const ANavigationData* Self, const FVector&
 	return Result.HasHit();
 }
 
-bool ARecastNavMesh::NavMeshRaycast(const ANavigationData* Self, NavNodeRef RayStartNode, const FVector& RayStart, const FVector& RayEnd, FVector& HitLocation, FSharedConstNavQueryFilter QueryFilter, const UObject* QueryOwner)
+bool ARecastNavMesh::NavMeshRaycast(const ANavigationData* Self, NavNodeRef RayStartNode, const FVector& RayStart, const FVector& RayEnd, FVector& HitLocation, FSharedConstNavQueryFilter QueryFilter, const UObject* Querier)
+{
+	return NavMeshRaycast(Self, RayStartNode, RayStart, RayEnd, HitLocation, nullptr, QueryFilter, Querier);
+}
+
+bool ARecastNavMesh::NavMeshRaycast(const ANavigationData* Self, NavNodeRef RayStartNode, const FVector& RayStart, const FVector& RayEnd, FVector& HitLocation, FNavigationRaycastAdditionalResults* AdditionalResults, FSharedConstNavQueryFilter QueryFilter, const UObject* QueryOwner)
 {
 	check(Cast<const ARecastNavMesh>(Self));
 
@@ -3313,6 +3318,10 @@ bool ARecastNavMesh::NavMeshRaycast(const ANavigationData* Self, NavNodeRef RayS
 	if (Self == NULL || RecastNavMesh->RecastNavMeshImpl == NULL)
 	{
 		HitLocation = RayStart;
+		if (AdditionalResults)
+		{
+			AdditionalResults->bIsRayEndInCorridor = false;
+		}
 		return true;
 	}
 
@@ -3320,6 +3329,10 @@ bool ARecastNavMesh::NavMeshRaycast(const ANavigationData* Self, NavNodeRef RayS
 	RecastNavMesh->RecastNavMeshImpl->Raycast(RayStart, RayEnd, RecastNavMesh->GetRightFilterRef(QueryFilter), QueryOwner, Result, RayStartNode);
 
 	HitLocation = Result.HasHit() ? (RayStart + (RayEnd - RayStart) * Result.HitTime) : RayEnd;
+	if (AdditionalResults)
+	{
+		AdditionalResults->bIsRayEndInCorridor = Result.bIsRaycastEndInCorridor;
+	}
 	return Result.HasHit();
 }
 
@@ -3363,10 +3376,17 @@ void ARecastNavMesh::BatchRaycast(TArray<FNavigationRaycastWork>& Workload, FSha
 				, QueryFilter, &RaycastResult.HitTime, RecastHitNormal
 				, RaycastResult.CorridorPolys, &RaycastResult.CorridorPolysCount, RaycastResult.GetMaxCorridorSize());
 
-			if (dtStatusSucceed(RaycastStatus) && RaycastResult.HasHit())
+			if (dtStatusSucceed(RaycastStatus))
 			{
-				WorkItem.bDidHit = true;
-				WorkItem.HitLocation = FNavLocation(WorkItem.RayStart + (WorkItem.RayEnd - WorkItem.RayStart) * RaycastResult.HitTime, RaycastResult.GetLastNodeRef());
+				if (RaycastResult.HasHit())
+				{
+					WorkItem.bDidHit = true;
+					WorkItem.HitLocation = FNavLocation(WorkItem.RayStart + (WorkItem.RayEnd - WorkItem.RayStart) * RaycastResult.HitTime, RaycastResult.GetLastNodeRef());
+				}
+				else
+				{
+					WorkItem.bIsRayEndInCorridor = RaycastResult.bIsRaycastEndInCorridor;
+				}
 			}
 		}
 	}
@@ -3374,15 +3394,10 @@ void ARecastNavMesh::BatchRaycast(TArray<FNavigationRaycastWork>& Workload, FSha
 
 bool ARecastNavMesh::IsSegmentOnNavmesh(const FVector& SegmentStart, const FVector& SegmentEnd, FSharedConstNavQueryFilter Filter, const UObject* QueryOwner) const
 {
-	if (RecastNavMeshImpl == NULL)
-	{
-		return false;
-	}
-	
+	FVector HitLocation;
 	FRaycastResult Result;
-	RecastNavMeshImpl->Raycast(SegmentStart, SegmentEnd, GetRightFilterRef(Filter), QueryOwner, Result);
-
-	return Result.bIsRaycastEndInCorridor && !Result.HasHit();
+	const bool bDidHit = NavMeshRaycast(this, SegmentStart, SegmentEnd, HitLocation, Filter, QueryOwner, Result);
+	return Result.bIsRaycastEndInCorridor && !bDidHit;
 }
 
 bool ARecastNavMesh::FindStraightPath(const FVector& StartLoc, const FVector& EndLoc, const TArray<NavNodeRef>& PathCorridor, TArray<FNavPathPoint>& PathPoints, TArray<FNavLinkId>* CustomLinks) const
diff --git a/Engine/Source/Runtime/NavigationSystem/Private/NavigationSystem.cpp b/Engine/Source/Runtime/NavigationSystem/Private/NavigationSystem.cpp
--- a/Engine/Source/Runtime/NavigationSystem/Private/NavigationSystem.cpp
+++ b/Engine/Source/Runtime/NavigationSystem/Private/NavigationSystem.cpp
@@ -2066,6 +2066,11 @@ UNavigationPath* UNavigationSystemV1::FindPathToLocationSynchronously(UObject* W
 }
 
 bool UNavigationSystemV1::NavigationRaycast(UObject* WorldContextObject, const FVector& RayStart, const FVector& RayEnd, FVector& HitLocation, TSubclassOf<UNavigationQueryFilter> FilterClass, AController* Querier)
+{
+	return NavigationRaycastWithAdditionalResults(WorldContextObject, RayStart, RayEnd, HitLocation, nullptr, FilterClass, Querier);
+}
+
+bool UNavigationSystemV1::NavigationRaycastWithAdditionalResults(UObject* WorldContextObject, const FVector& RayStart, const FVector& RayEnd, FVector& HitLocation, FNavigationRaycastAdditionalResults* AdditionalResults, TSubclassOf<UNavigationQueryFilter> FilterClass, AController* Querier)
 {
 	UWorld* World = nullptr;
 
@@ -2081,6 +2086,10 @@ bool UNavigationSystemV1::NavigationRaycast(UObject* WorldContextObject, const F
 	// blocked, i.e. not traversable, by default
 	bool bRaycastBlocked = true;
 	HitLocation = RayStart;
+	if (AdditionalResults)
+	{
+		AdditionalResults->bIsRayEndInCorridor = false;
+	}
 
 	const UNavigationSystemV1* NavSys = FNavigationSystem::GetCurrent<UNavigationSystemV1>(World);
 
@@ -2101,7 +2110,7 @@ bool UNavigationSystemV1::NavigationRaycast(UObject* WorldContextObject, const F
 
 		if (NavData != nullptr)
 		{
-			bRaycastBlocked = NavData->Raycast(RayStart, RayEnd, HitLocation, UNavigationQueryFilter::GetQueryFilter(*NavData, Querier, FilterClass));
+			bRaycastBlocked = NavData->Raycast(RayStart, RayEnd, HitLocation, AdditionalResults, UNavigationQueryFilter::GetQueryFilter(*NavData, Querier, FilterClass));
 		}
 	}
 
diff --git a/Engine/Source/Runtime/NavigationSystem/Private/NavigationTestingActor.cpp b/Engine/Source/Runtime/NavigationSystem/Private/NavigationTestingActor.cpp
--- a/Engine/Source/Runtime/NavigationSystem/Private/NavigationTestingActor.cpp
+++ b/Engine/Source/Runtime/NavigationSystem/Private/NavigationTestingActor.cpp
@@ -59,6 +59,7 @@ ANavigationTestingActor::ANavigationTestingActor(const FObjectInitializer& Objec
 	bNavDataIsReadyInRadius = false;
 	bNavDataIsReadyToQueryTargetActor = false;
 	bRaycastToQueryTargetActorResult = false;
+	bRaycastToQueryTargetEndsInCorridor = false;
 	OffsetFromCornersDistance = 0.f;
 
 	QueryingExtent = FVector(DEFAULT_NAV_QUERY_EXTENT_HORIZONTAL, DEFAULT_NAV_QUERY_EXTENT_HORIZONTAL, DEFAULT_NAV_QUERY_EXTENT_VERTICAL);
@@ -428,7 +429,7 @@ void ANavigationTestingActor::UpdateTargetActorQueries()
 	}
 	if (bDrawRaycastToQueryTargetActor)
 	{
-		bRaycastToQueryTargetActorResult = CheckRaycastToActor(QueryTargetActor, RaycastHitLocation);
+		bRaycastToQueryTargetActorResult = CheckRaycastToActor(QueryTargetActor, RaycastHitLocation, bRaycastToQueryTargetEndsInCorridor);
 	}
 }
 
@@ -496,7 +497,7 @@ bool ANavigationTestingActor::CheckIfNavDataIsReadyToActor(const AActor* TargetA
 	return false;
 }
 
-bool ANavigationTestingActor::CheckRaycastToActor(const AActor* TargetActor, FVector& OutHitLocation)
+bool ANavigationTestingActor::CheckRaycastToActor(const AActor* TargetActor, FVector& OutHitLocation, bool& bOutIsRaycastEndInCorridor)
 {
 	OutHitLocation = FNavigationSystem::InvalidLocation;
 #if WITH_EDITORONLY_DATA
@@ -510,7 +511,10 @@ bool ANavigationTestingActor::CheckRaycastToActor(const AActor* TargetActor, FVe
 	if (MyNavData && TargetActor)
 	{
 		FSharedConstNavQueryFilter Filter = UNavigationQueryFilter::GetQueryFilter(*MyNavData, this, FilterClass);
-		return MyNavData->Raycast(GetActorLocation(), TargetActor->GetActorLocation(), OutHitLocation, Filter, this);
+		FNavigationRaycastAdditionalResults AdditionalResults;
+		const bool bDidHit = MyNavData->Raycast(GetActorLocation(), TargetActor->GetActorLocation(), OutHitLocation, &AdditionalResults, Filter, this);
+		bOutIsRaycastEndInCorridor = AdditionalResults.bIsRayEndInCorridor;
+		return bDidHit;
 	}
 
 	return false;
diff --git a/Engine/Source/Runtime/NavigationSystem/Public/AbstractNavData.h b/Engine/Source/Runtime/NavigationSystem/Public/AbstractNavData.h
--- a/Engine/Source/Runtime/NavigationSystem/Public/AbstractNavData.h
+++ b/Engine/Source/Runtime/NavigationSystem/Public/AbstractNavData.h
@@ -99,4 +99,5 @@ public:
 	static NAVIGATIONSYSTEM_API FPathFindingResult FindPathAbstract(const FNavAgentProperties& AgentProperties, const FPathFindingQuery& Query);
 	static NAVIGATIONSYSTEM_API bool TestPathAbstract(const FNavAgentProperties& AgentProperties, const FPathFindingQuery& Query, int32* NumVisitedNodes);
 	static NAVIGATIONSYSTEM_API bool RaycastAbstract(const ANavigationData* NavDataInstance, const FVector& RayStart, const FVector& RayEnd, FVector& HitLocation, FSharedConstNavQueryFilter QueryFilter, const UObject* Querier);
+	static NAVIGATIONSYSTEM_API bool RaycastAbstract(const ANavigationData* NavDataInstance, const FVector& RayStart, const FVector& RayEnd, FVector& HitLocation, FNavigationRaycastAdditionalResults* AdditionalResults, FSharedConstNavQueryFilter QueryFilter, const UObject* Querier);
 };
diff --git a/Engine/Source/Runtime/NavigationSystem/Public/NavMesh/RecastNavMesh.h b/Engine/Source/Runtime/NavigationSystem/Public/NavMesh/RecastNavMesh.h
--- a/Engine/Source/Runtime/NavigationSystem/Public/NavMesh/RecastNavMesh.h
+++ b/Engine/Source/Runtime/NavigationSystem/Public/NavMesh/RecastNavMesh.h
@@ -1368,13 +1368,15 @@ public:
 	NAVIGATIONSYSTEM_API bool ProjectPointMulti(const FVector& Point, TArray<FNavLocation>& OutLocations, const FVector& Extent,
 		FVector::FReal MinZ, FVector::FReal MaxZ, FSharedConstNavQueryFilter Filter = NULL, const UObject* Querier = NULL) const;
 	
-	// @todo docuement
+	// @todo document
 	static NAVIGATIONSYSTEM_API FPathFindingResult FindPath(const FNavAgentProperties& AgentProperties, const FPathFindingQuery& Query);
 	static NAVIGATIONSYSTEM_API bool TestPath(const FNavAgentProperties& AgentProperties, const FPathFindingQuery& Query, int32* NumVisitedNodes);
 	static NAVIGATIONSYSTEM_API bool TestHierarchicalPath(const FNavAgentProperties& AgentProperties, const FPathFindingQuery& Query, int32* NumVisitedNodes);
 	static NAVIGATIONSYSTEM_API bool NavMeshRaycast(const ANavigationData* Self, const FVector& RayStart, const FVector& RayEnd, FVector& HitLocation, FSharedConstNavQueryFilter QueryFilter, const UObject* Querier, FRaycastResult& Result);
 	static NAVIGATIONSYSTEM_API bool NavMeshRaycast(const ANavigationData* Self, const FVector& RayStart, const FVector& RayEnd, FVector& HitLocation, FSharedConstNavQueryFilter QueryFilter, const UObject* Querier = NULL);
+	static NAVIGATIONSYSTEM_API bool NavMeshRaycast(const ANavigationData* Self, const FVector& RayStart, const FVector& RayEnd, FVector& HitLocation, FNavigationRaycastAdditionalResults* AdditionalResults, FSharedConstNavQueryFilter QueryFilter, const UObject* Querier = NULL);
 	static NAVIGATIONSYSTEM_API bool NavMeshRaycast(const ANavigationData* Self, NavNodeRef RayStartNode, const FVector& RayStart, const FVector& RayEnd, FVector& HitLocation, FSharedConstNavQueryFilter QueryFilter, const UObject* Querier = NULL);
+	static NAVIGATIONSYSTEM_API bool NavMeshRaycast(const ANavigationData* Self, NavNodeRef RayStartNode, const FVector& RayStart, const FVector& RayEnd, FVector& HitLocation, FNavigationRaycastAdditionalResults* AdditionalResults, FSharedConstNavQueryFilter QueryFilter, const UObject* Querier = NULL);
 
 	NAVIGATIONSYSTEM_API virtual void BatchRaycast(TArray<FNavigationRaycastWork>& Workload, FSharedConstNavQueryFilter QueryFilter, const UObject* Querier = NULL) const override;
 
@@ -1544,9 +1546,20 @@ public:
 #if WITH_RECAST
 FORCEINLINE
 bool ARecastNavMesh::NavMeshRaycast(const ANavigationData* Self, const FVector& RayStart, const FVector& RayEnd, FVector& HitLocation, FSharedConstNavQueryFilter QueryFilter, const UObject* Querier)
+{
+	return NavMeshRaycast(Self, RayStart, RayEnd, HitLocation, nullptr, QueryFilter, Querier);
+}
+
+FORCEINLINE
+bool ARecastNavMesh::NavMeshRaycast(const ANavigationData* Self, const FVector& RayStart, const FVector& RayEnd, FVector& HitLocation, FNavigationRaycastAdditionalResults* AdditionalResults, FSharedConstNavQueryFilter QueryFilter, const UObject* Querier)
 {
 	FRaycastResult Result;
-	return NavMeshRaycast(Self, RayStart, RayEnd, HitLocation, QueryFilter, Querier, Result);
+	const bool bDidHit = NavMeshRaycast(Self, RayStart, RayEnd, HitLocation, QueryFilter, Querier, Result);
+	if (AdditionalResults)
+	{
+		AdditionalResults->bIsRayEndInCorridor = Result.bIsRaycastEndInCorridor;
+	}
+	return bDidHit;
 }
 
 
diff --git a/Engine/Source/Runtime/NavigationSystem/Public/NavigationData.h b/Engine/Source/Runtime/NavigationSystem/Public/NavigationData.h
--- a/Engine/Source/Runtime/NavigationSystem/Public/NavigationData.h
+++ b/Engine/Source/Runtime/NavigationSystem/Public/NavigationData.h
@@ -74,6 +74,14 @@ struct FPathFindingResult
 	FORCEINLINE bool IsPartial() const;
 };
 
+struct FNavigationRaycastAdditionalResults
+{
+	/** When the ray is not obstructed, indicates if the projection of RayEnd is located at the end of the explored corridor.
+	 *  When bIsRaytEndInCorridor is false, it means that RayEnd failed to project to the NavigationData or on a navigation node that is not part of the explored corridor (e.g. different height)
+	 */
+	bool bIsRayEndInCorridor = false;
+};
+
 struct FNavigationPath : public TSharedFromThis<FNavigationPath, ESPMode::ThreadSafe>
 {
 	//DECLARE_DELEGATE_OneParam(FPathObserverDelegate, FNavigationPath*);
@@ -836,9 +844,22 @@ public:
 	 */
 	FORCEINLINE bool Raycast(const FVector& RayStart, const FVector& RayEnd, FVector& HitLocation, FSharedConstNavQueryFilter QueryFilter, const UObject* Querier = NULL) const
 	{
-		check(RaycastImplementation);
+		return Raycast(RayStart, RayEnd, HitLocation, nullptr/*AdditionalResults*/, QueryFilter, Querier);
+	}
+
+	/** 
+	 *	Synchronously makes a raycast on navigation data using QueryFilter
+	 *	@param HitLocation if line was obstructed this will be set to hit location. Otherwise it contains SegmentEnd
+	 *	@param AdditionalResults contains more information about the result of the raycast query. See FNavigationRaycastAdditionalResults description for details
+	 *	@return true if line from RayStart to RayEnd is obstructed
+	 *
+	 *	@note don't make this function virtual! Look at implementation details and its comments for more info.
+	 */
+	FORCEINLINE bool Raycast(const FVector& RayStart, const FVector& RayEnd, FVector& HitLocation, FNavigationRaycastAdditionalResults* AdditionalResults, FSharedConstNavQueryFilter QueryFilter, const UObject* Querier = NULL) const
+	{
+		check(RaycastImplementationWithAdditionalResults);
 		// this awkward implementation avoids virtual call overhead - it's possible this function will be called a lot
-		return (*RaycastImplementation)(this, RayStart, RayEnd, HitLocation, QueryFilter, Querier);
+		return (*RaycastImplementationWithAdditionalResults)(this, RayStart, RayEnd, HitLocation, AdditionalResults, QueryFilter, Querier);
 	}
 
 	/** Raycasts batched for efficiency */
@@ -998,7 +1019,11 @@ protected:
 	FTestPathPtr TestHierarchicalPathImplementation; 
 
 	typedef bool(*FNavRaycastPtr)(const ANavigationData* NavDataInstance, const FVector& RayStart, const FVector& RayEnd, FVector& HitLocation, FSharedConstNavQueryFilter QueryFilter, const UObject* Querier);
-	FNavRaycastPtr RaycastImplementation; 
+	UE_DEPRECATED(5.6, "Please use RaycastImplementationWithAdditionalResults instead") 
+	FNavRaycastPtr RaycastImplementation;
+
+	typedef bool(*FNavRaycastWithAdditionalResultsPtr)(const ANavigationData* NavDataInstance, const FVector& RayStart, const FVector& RayEnd, FVector& HitLocation, FNavigationRaycastAdditionalResults* AdditionalResults, FSharedConstNavQueryFilter QueryFilter, const UObject* Querier);
+	FNavRaycastWithAdditionalResultsPtr RaycastImplementationWithAdditionalResults;
 
 protected:
 	TSharedPtr<FNavDataGenerator, ESPMode::ThreadSafe> NavDataGenerator;
diff --git a/Engine/Source/Runtime/NavigationSystem/Public/NavigationSystem.h b/Engine/Source/Runtime/NavigationSystem/Public/NavigationSystem.h
--- a/Engine/Source/Runtime/NavigationSystem/Public/NavigationSystem.h
+++ b/Engine/Source/Runtime/NavigationSystem/Public/NavigationSystem.h
@@ -528,6 +528,13 @@ public:
 	UFUNCTION(BlueprintCallable, Category="AI|Navigation", meta=(WorldContext="WorldContextObject" ))
 	static NAVIGATIONSYSTEM_API bool NavigationRaycast(UObject* WorldContextObject, const FVector& RayStart, const FVector& RayEnd, FVector& HitLocation, TSubclassOf<UNavigationQueryFilter> FilterClass = NULL, AController* Querier = NULL);
 
+	/** Performs navigation raycast on NavigationData appropriate for given Querier.
+	 *	@param Querier if not passed default navigation data will be used
+	 *	@param HitLocation if line was obstructed this will be set to hit location. Otherwise it contains SegmentEnd
+	 *	@param AdditionalResults contains more information about the result of the raycast query. See FNavigationRaycastAdditionalResults description for details
+	 *	@return true if line from RayStart to RayEnd was obstructed. Also, true when no navigation data present */
+	static NAVIGATIONSYSTEM_API bool NavigationRaycastWithAdditionalResults(UObject* WorldContextObject, const FVector& RayStart, const FVector& RayEnd, FVector& HitLocation, FNavigationRaycastAdditionalResults* AdditionalResults, TSubclassOf<UNavigationQueryFilter> FilterClass = NULL, AController* Querier = NULL);
+
 	/** will limit the number of simultaneously running navmesh tile generation jobs to specified number.
 	 *	@param MaxNumberOfJobs gets trimmed to be at least 1. You cannot use this function to pause navmesh generation */
 	UFUNCTION(BlueprintCallable, Category = "AI|Navigation")
diff --git a/Engine/Source/Runtime/NavigationSystem/Public/NavigationTestingActor.h b/Engine/Source/Runtime/NavigationSystem/Public/NavigationTestingActor.h
--- a/Engine/Source/Runtime/NavigationSystem/Public/NavigationTestingActor.h
+++ b/Engine/Source/Runtime/NavigationSystem/Public/NavigationTestingActor.h
@@ -124,7 +124,8 @@ public:
 	UPROPERTY(EditAnywhere, Category=Query)
 	uint32 bDrawIfNavDataIsReadyToQueryTargetActor : 1;
 
-	/** If set, a line is drawn to indicate to result of a raycast on the navigation data between the current actor and the QueryTargetActor location. */
+	/** If set, a line is drawn to indicate to result of a ray cast on the navigation data between the current actor and the QueryTargetActor location
+	 * (red when there is a hit, green when there is no hit and the ray end is on the explored corridor, orange otherwise). */
 	UPROPERTY(EditAnywhere, Category=Query)
 	uint32 bDrawRaycastToQueryTargetActor : 1;
 
@@ -198,6 +199,7 @@ public:
 	bool bNavDataIsReadyInRadius;
 	bool bNavDataIsReadyToQueryTargetActor;
 	bool bRaycastToQueryTargetActorResult;
+	bool bRaycastToQueryTargetEndsInCorridor;
 
 #if WITH_RECAST && WITH_EDITORONLY_DATA
 	/** detail data gathered from each step of regular A* algorithm */
@@ -263,6 +265,6 @@ protected:
 	NAVIGATIONSYSTEM_API FVector FindClosestWallLocation() const;
 	bool CheckIfNavDataIsReadyInRadius();
 	bool CheckIfNavDataIsReadyToActor(const AActor* TargetActor);
-	bool CheckRaycastToActor(const AActor* TargetActor, FVector& OutHitLocation);
+	bool CheckRaycastToActor(const AActor* TargetActor, FVector& OutHitLocation, bool& bOutIsRaycastEndInCorridor);
 	void OnQueryTargetActorTransformUpdated(USceneComponent* InRootComponent, EUpdateTransformFlags UpdateTransformFlags, ETeleportType Teleport);
 };
