[NavMesh] Bring back wrongly removed GetDebugGeometryForTile from CL 34671788

diff --git a/Engine/Source/Runtime/NavigationSystem/Private/NavMesh/RecastNavMesh.cpp b/Engine/Source/Runtime/NavigationSystem/Private/NavMesh/RecastNavMesh.cpp
--- a/Engine/Source/Runtime/NavigationSystem/Private/NavMesh/RecastNavMesh.cpp
+++ b/Engine/Source/Runtime/NavigationSystem/Private/NavMesh/RecastNavMesh.cpp
@@ -2646,6 +2646,23 @@ bool ARecastNavMesh::GetPolysWithinPathingDistance(FVector const& StartLoc, cons
 	return RecastNavMeshImpl && RecastNavMeshImpl->GetPolysWithinPathingDistance(StartLoc, PathingDistance, GetRightFilterRef(Filter), QueryOwner, FoundPolys, DebugData);
 }
 
+// Deprecated
+bool ARecastNavMesh::GetDebugGeometryForTile(FRecastDebugGeometry& OutGeometry, int32 TileIndex) const
+{
+	if (RecastNavMeshImpl)
+	{
+		FNavTileRef TileRef;
+		if (TileIndex != INDEX_NONE)
+		{
+			TArray<FNavTileRef> TileRefs;
+			FNavTileRef::DeprecatedMakeTileRefsFromTileIds(RecastNavMeshImpl, { static_cast<uint32>(TileIndex) }, TileRefs);
+			TileRef = TileRefs[0];
+		}
+		return GetDebugGeometryForTile(OutGeometry, TileRef);
+	}
+	return true;
+}
+
 bool ARecastNavMesh::GetDebugGeometryForTile(FRecastDebugGeometry& OutGeometry, FNavTileRef TileRef) const
 {
 	if (RecastNavMeshImpl)
