Improvements to NavLink comments

diff --git a/Engine/Source/Runtime/NavigationSystem/Public/NavMesh/LinkGenerationConfig.h b/Engine/Source/Runtime/NavigationSystem/Public/NavMesh/LinkGenerationConfig.h
--- a/Engine/Source/Runtime/NavigationSystem/Public/NavMesh/LinkGenerationConfig.h
+++ b/Engine/Source/Runtime/NavigationSystem/Public/NavMesh/LinkGenerationConfig.h
@@ -13,7 +13,7 @@ struct dtNavLinkBuilderJumpDownConfig;
 struct dtNavLinkBuilderJumpOverConfig;
 #endif //WITH_RECAST
 
-/** Experimental configuration for generated jump down links. */
+/** Experimental configuration to generate vertical links. */
 USTRUCT()
 struct FNavLinkGenerationJumpDownConfig
 {
@@ -23,39 +23,42 @@ struct FNavLinkGenerationJumpDownConfig
 	UPROPERTY(EditAnywhere, Config, Category = Settings)
 	bool bEnabled = true;
 	
-	/** Horizontal length of the jump. How far from the starting point we will look for ground. */ 
-	UPROPERTY(EditAnywhere, Config, Category = Settings, meta=(Units=cm, UIMin=0, ClampMin=0))
+	/** Horizontal length of the jump.
+	 * How far from the starting point we will look for ground. */ 
+	UPROPERTY(EditAnywhere, Config, Category = Settings, meta=(UIMin=0, ClampMin=0))
 	float JumpLength = 150.f; 
 
-	/** How far from the edge is the jump started. */
-	UPROPERTY(EditAnywhere, Config, Category = Settings, meta=(Units=cm, UIMin=0, ClampMin=0))
+	/** How far from the navmesh edge is the jump started. */
+	UPROPERTY(EditAnywhere, Config, Category = Settings, meta=(UIMin=0, ClampMin=0))
 	float JumpDistanceFromEdge = 10.f; 
 
 	/** How far below the starting height we want to look for landing ground. */
-	UPROPERTY(EditAnywhere, Config, Category = Settings, meta=(Units=cm, UIMin=0, ClampMin=0))
+	UPROPERTY(EditAnywhere, Config, Category = Settings, meta=(UIMin=0, ClampMin=0))
 	float JumpMaxDepth = 150.f;
 
 	/** Peak height relative to the height of the starting point. */
-	UPROPERTY(EditAnywhere, Config, Category = Settings, meta=(Units=cm, UIMin=0, ClampMin=0))
+	UPROPERTY(EditAnywhere, Config, Category = Settings, meta=(UIMin=0, ClampMin=0))
 	float JumpHeight = 50.f;
 
 	/** Tolerance at both ends of the jump to find ground. */
-	UPROPERTY(EditAnywhere, Config, Category = Settings, meta=(Units=cm, UIMin=0, ClampMin=0))
+	UPROPERTY(EditAnywhere, Config, Category = Settings, meta=(UIMin=0, ClampMin=0))
 	float JumpEndsHeightTolerance = 80.f;
 
-	/** Value multiplied by CellSize to find the distance between sampling trajectories. Default is 1. */
-    /*  Larger values improve generation speed but might introduce sampling errors.  */
+	/** Value multiplied by CellSize to find the distance between sampling trajectories. Default is 1.
+     *  Larger values improve generation speed but might introduce sampling errors.  */
 	UPROPERTY(EditAnywhere, Config, Category = Settings, meta=(UIMin=1, ClampMin=1))
 	float SamplingSeparationFactor = 1.f;
 	
-	/** When filtering similar links, distance used to compare between segment endpoints to match similar links. Use greater distance for more filtering (0 to deactivate filtering). */
-	UPROPERTY(EditAnywhere, Config, Category = Settings, meta=(Units=cm, UIMin=0, ClampMin=0))
+	/** When filtering similar links, it's the distance used to compare between segment endpoints to match similar links.
+	 * Use greater distance for more filtering (0 to deactivate filtering). */
+	UPROPERTY(EditAnywhere, Config, Category = Settings, meta=(UIMin=0, ClampMin=0))
 	float FilterDistanceThreshold = 80.f;
 
 	/** Class used to handle links made with this configuration.
 	 * Using this allows to implement custom behaviors when using navlinks, for example during the pathfollow.
-	 * Note that having a proxy is not required for successful navlink pathfinding, but it does allow for custom behavior at the start and the end of a given navlink.
-	 * This implies that using LinkProxyClass is optional and it can remain empty (the default value).
+	 * Note that having a proxy is not required for successful navlink pathfinding,
+	 * but it does allow for custom behavior at the start and the end of a given navlink.
+	 * This implies that using LinkProxyClass is optional, and it can remain empty (the default value).
 	 * @see INavLinkCustomInterface 
 	 * @see UGeneratedNavLinksProxy
 	 */
diff --git a/Engine/Source/Runtime/NavigationSystem/Public/NavMesh/RecastNavMesh.h b/Engine/Source/Runtime/NavigationSystem/Public/NavMesh/RecastNavMesh.h
--- a/Engine/Source/Runtime/NavigationSystem/Public/NavMesh/RecastNavMesh.h
+++ b/Engine/Source/Runtime/NavigationSystem/Public/NavMesh/RecastNavMesh.h
@@ -792,7 +792,8 @@ class ARecastNavMesh : public ANavigationData
 	UPROPERTY(EditAnywhere, Category=Generation, config, meta = (EditCondition = "bAllowWorldPartitionedNavMesh", HideEditConditionToggle, DisplayName = "IsWorldPartitionedNavMesh"))
 	uint32 bIsWorldPartitioned : 1;
 
-	/** Experimental: if set, navlinks will be automatically generated. */ 
+	/** Experimental: if set, navlinks will be automatically generated.
+	 * @see FNavLinkGenerationJumpDownConfig */ 
 	UPROPERTY(EditAnywhere, Category=Generation, config)
 	uint32 bGenerateNavLinks : 1;
 	
@@ -834,7 +835,8 @@ class ARecastNavMesh : public ANavigationData
 	UPROPERTY(config)
 	uint32 bUseVirtualFilters : 1;
 
-	/** Indicates whether use the virtual methods to check if an object should generate geometry or if we should call the normal method directly (i.e. FNavigationOctreeElement::ShouldUseGeometry).
+	/** Indicates whether to use the virtual methods to check if an object should generate geometry or
+	 *  if we should call the normal method directly (i.e. FNavigationOctreeElement::ShouldUseGeometry).
 	 *  If enabled, will also check if an object requesting an update on the navmesh is excluded to avoid dirtying the areas unnecessarily.
 	 *  Defaults to false. */
 	UPROPERTY(config)
@@ -867,7 +869,7 @@ protected:
 	uint32 bAllowWorldPartitionedNavMesh : 1;
 #endif // WITH_EDITORONLY_DATA
 
-	/** Experimental: paramters for building jump down links.  */
+	/** Experimental configuration to generate vertical links. */
 	UPROPERTY(EditAnywhere, Category=Generation, config)
 	FNavLinkGenerationJumpDownConfig NavLinkJumpDownConfig;
 	
