[NavigationSystem] updating API and various part of the navigation framework to support registration of non UObject elements.
- Added FNavigationRelevantElement as the main structure used to provide required information. The struct is meant to remove the dependencies on INavRelevantInterface and UObject. It currently provides constructors dedicated to handle current components and other UObjects implementing the interface.

diff --git a/Engine/Source/Developer/FunctionalTesting/Private/FunctionalAITest.cpp b/Engine/Source/Developer/FunctionalTesting/Private/FunctionalAITest.cpp
--- a/Engine/Source/Developer/FunctionalTesting/Private/FunctionalAITest.cpp
+++ b/Engine/Source/Developer/FunctionalTesting/Private/FunctionalAITest.cpp
@@ -7,6 +7,7 @@
 #include "FunctionalTestingManager.h"
 #include "NavigationSystem.h"
 #include "AI/Navigation/NavAreaBase.h"
+#include "AI/Navigation/NavigationElement.h"
 #include "AIController.h"
 #include "Blueprint/AIBlueprintHelperLibrary.h"
 #include "NavMesh/RecastNavMesh.h"
@@ -131,7 +132,7 @@ void AFunctionalAITestBase::OnTimeout()
 					DebugAreaNames += TEXT(',');
 				}
 
-				UE_LOG(LogFunctionalTest, Log, TEXT("> modifier, owner:%s areas:%s"), *GetNameSafe(Element.GetOwner()), *DebugAreaNames);
+				UE_LOG(LogFunctionalTest, Log, TEXT("> modifier, owner:%s areas:%s"), *Element.GetSourceElement()->GetName(), *DebugAreaNames);
 			}
 		});
 	}
diff --git a/Engine/Source/Runtime/Engine/Classes/AI/Navigation/NavRelevantInterface.h b/Engine/Source/Runtime/Engine/Classes/AI/Navigation/NavRelevantInterface.h
--- a/Engine/Source/Runtime/Engine/Classes/AI/Navigation/NavRelevantInterface.h
+++ b/Engine/Source/Runtime/Engine/Classes/AI/Navigation/NavRelevantInterface.h
@@ -19,22 +19,22 @@ struct FNavigableGeometryExport;
 struct FNavigationRelevantData;
 class UBodySetup;
 
-/** Determines if a NavRelevant object contains custom collision for navigation/AI */
+/** Determines if, and how, a navigation element should export collision for AI navigation */
 UENUM()
 namespace EHasCustomNavigableGeometry
 {
 enum Type : int
 {
-	/** Primitive doesn't have custom navigation geometry, if collision is enabled then its convex/trimesh collision will be used for generating the navmesh */
+	/** Element custom geometry export callback is not called, but the default collision export is performed using its convex/trimesh collision. */
 	No,
 
-	/** If primitive would normally affect navmesh, DoCustomNavigableGeometryExport() should be called to export this primitive's navigable geometry */
+	/** The custom geometry export callback is called and indicates if the default collision export should also be performed. */
 	Yes,
 
-	/** DoCustomNavigableGeometryExport() should be called even if the mesh is non-collidable and wouldn't normally affect the navmesh */
+	/** The custom geometry export callback is called even if the mesh is non-collidable and wouldn't normally affect the navigation data. */
 	EvenIfNotCollidable,
 
-	/** Don't export navigable geometry even if primitive is relevant for navigation (can still add modifiers) */
+	/** Neither the custom geometry export delegate nor the default export will be called (can still add modifiers through the NavigationData export callback). */
 	DontExport,
 };
 }
@@ -59,7 +59,7 @@ class INavRelevantInterface
 	virtual bool SupportsGatheringGeometrySlices() const { return false; }
 	
 	/** Indicates if the area covered by the navigation bounds of the object should not be dirtied when inserting, or removing, the object in the navigation octree.
-	 * In this case we expect that object to manually dirty areas (e.g. using OnObjectBoundsChanged).
+	 * In this case we expect that object to manually dirty areas (e.g. using UpdateNavigationElementBoundsDelegate).
 	 */
 	virtual bool ShouldSkipDirtyAreaOnAddOrRemove() const { return false; }
 
diff --git a/Engine/Source/Runtime/Engine/Classes/AI/Navigation/NavigationDirtyArea.h b/Engine/Source/Runtime/Engine/Classes/AI/Navigation/NavigationDirtyArea.h
--- a/Engine/Source/Runtime/Engine/Classes/AI/Navigation/NavigationDirtyArea.h
+++ b/Engine/Source/Runtime/Engine/Classes/AI/Navigation/NavigationDirtyArea.h
@@ -7,6 +7,7 @@
 #include "UObject/WeakObjectPtr.h"
 
 class UObject;
+struct FNavigationElement;
 
 enum class ENavigationDirtyFlag : uint8
 {
@@ -23,13 +24,17 @@ ENUM_CLASS_FLAGS(ENavigationDirtyFlag);
 struct FNavigationDirtyArea
 {
 	FBox Bounds = FBox(ForceInit);
+#if WITH_EDITORONLY_DATA
+	UE_DEPRECATED(5.5, "Use OptionalSourceElement instead.")
 	TWeakObjectPtr<UObject> OptionalSourceObject;
+#endif // WITH_EDITORONLY_DATA
+	TSharedPtr<const FNavigationElement> OptionalSourceElement;
 	ENavigationDirtyFlag Flags = ENavigationDirtyFlag::None;
 
 	FNavigationDirtyArea() = default;
-	ENGINE_API FNavigationDirtyArea(const FBox& InBounds, ENavigationDirtyFlag InFlags, UObject* const InOptionalSourceObject = nullptr);
+	ENGINE_API FNavigationDirtyArea(const FBox& InBounds, ENavigationDirtyFlag InFlags, const TSharedPtr<const FNavigationElement>& InOptionalSourceElement = nullptr);
 
-	UE_DEPRECATED(5.5, "Use constructor taking ENavigationDirtyFlag instead.")
+	UE_DEPRECATED(5.5, "Use the constructor taking ENavigationDirtyFlag and FNavigationElement instead.")
 	ENGINE_API FNavigationDirtyArea(const FBox& InBounds, int32 InFlags, UObject* const InOptionalSourceObject = nullptr);
 
 	bool HasFlag(const ENavigationDirtyFlag Flag) const
@@ -39,11 +44,13 @@ struct FNavigationDirtyArea
 
 	bool operator==(const FNavigationDirtyArea& Other) const
 	{ 
-		return Flags == Other.Flags && OptionalSourceObject == Other.OptionalSourceObject && Bounds.Equals(Other.Bounds);
+		return Flags == Other.Flags && OptionalSourceElement == Other.OptionalSourceElement && Bounds.Equals(Other.Bounds);
 	}
 	
 	bool operator!=( const FNavigationDirtyArea& Other) const
 	{
 		return !(*this == Other);
 	}
+
+	ENGINE_API FString GetSourceDescription() const;
 };
diff --git a/Engine/Source/Runtime/Engine/Classes/AI/Navigation/NavigationDirtyElement.h b/Engine/Source/Runtime/Engine/Classes/AI/Navigation/NavigationDirtyElement.h
--- a/Engine/Source/Runtime/Engine/Classes/AI/Navigation/NavigationDirtyElement.h
+++ b/Engine/Source/Runtime/Engine/Classes/AI/Navigation/NavigationDirtyElement.h
@@ -7,21 +7,27 @@
 
 enum class ENavigationDirtyFlag : uint8;
 class INavRelevantInterface;
+struct FNavigationElement;
 
 struct FNavigationDirtyElement
 {
 	/**
-	 * If not empty and the associated navigation relevant object controls the dirty areas explicitly (i.e. ShouldSkipDirtyAreaOnAddOrRemove returns true),
+	 * If not empty and the associated navigation element controls the dirty areas explicitly (i.e. DirtyAreasOnRegistration is 'false'),
 	 * the list will be used to indicate the areas that need rebuilding.
-	 * Otherwise, the default behavior, NavRelevant object's bounds will be used.
+	 * Otherwise, the default behavior, element's bounds will be used.
 	 */
 	TArray<FBox> ExplicitAreasToDirty;
 
-	/** object owning this element */
+	/** Navigation element associated with this dirty element */
+	TSharedRef<const FNavigationElement> NavigationElement;
+
+#if WITH_EDITORONLY_DATA
+	UE_DEPRECATED(5.5, "Use NavigationElement instead.")
 	FWeakObjectPtr Owner;
 
-	/** cached interface pointer */
+	UE_DEPRECATED(5.5, "Use NavigationElement instead.")
 	INavRelevantInterface* NavInterface = nullptr;
+#endif // WITH_EDITORONLY_DATA
 
 	/** bounds of already existing entry for this actor */
 	FBox PrevBounds = FBox(ForceInit);
@@ -44,26 +50,17 @@ struct FNavigationDirtyElement
 	/** part of the base navmesh */
 	uint8 bIsInBaseNavmesh : 1 = false;
 
+	UE_DEPRECATED(5.5, "The default constructor will be remove. Use the version with FNavigationElement instead.")
 	ENGINE_API FNavigationDirtyElement();
+	UE_DEPRECATED(5.5, "Use the version with FNavigationElement instead.")
 	ENGINE_API explicit FNavigationDirtyElement(UObject* InOwner);
-	ENGINE_API FNavigationDirtyElement(UObject* InOwner, INavRelevantInterface* InNavInterface, ENavigationDirtyFlag InFlagsOverride, const bool bUseWorldPartitionedDynamicMode = false);
-	ENGINE_API FNavigationDirtyElement(UObject* InOwner, INavRelevantInterface* InNavInterface, const bool bUseWorldPartitionedDynamicMode = false);
-
-	UE_DEPRECATED(5.5, "Use the version taking ENavigationDirtyFlag instead.")
+	UE_DEPRECATED(5.5, "Use the version with FNavigationElement instead.")
 	ENGINE_API FNavigationDirtyElement(UObject* InOwner, INavRelevantInterface* InNavInterface, int32 InFlagsOverride = 0, const bool bUseWorldPartitionedDynamicMode = false);
+	ENGINE_API FNavigationDirtyElement(const TSharedRef<const FNavigationElement>& InNavigationElement, ENavigationDirtyFlag InFlagsOverride, const bool bUseWorldPartitionedDynamicMode = false);
+	ENGINE_API explicit FNavigationDirtyElement(const TSharedRef<const FNavigationElement>& InNavigationElement, const bool bUseWorldPartitionedDynamicMode = false);
 
-	bool operator==(const FNavigationDirtyElement& Other) const 
-	{ 
-		return Owner == Other.Owner; 
-	}
-
-	bool operator==(const UObject*& OtherOwner) const 
-	{ 
-		return (Owner == OtherOwner);
-	}
+	UE_DEPRECATED(5.5, "This operator will no longer be used.")
+	ENGINE_API bool operator==(const UObject*& OtherOwner) const;
 
-	FORCEINLINE friend uint32 GetTypeHash(const FNavigationDirtyElement& Info)
-	{
-		return GetTypeHash(Info.Owner);
-	}
+	ENGINE_API friend uint32 GetTypeHash(const FNavigationDirtyElement& Info);
 };
diff --git a/Engine/Source/Runtime/Engine/Classes/AI/Navigation/NavigationElement.h b/Engine/Source/Runtime/Engine/Classes/AI/Navigation/NavigationElement.h
new file mode 100644
--- /dev/null
+++ b/Engine/Source/Runtime/Engine/Classes/AI/Navigation/NavigationElement.h
@@ -0,0 +1,451 @@
+// Copyright Epic Games, Inc. All Rights Reserved.
+
+#pragma once
+
+#include "HAL/Platform.h"
+#include "UObject/WeakObjectPtr.h"
+
+class INavRelevantInterface;
+class UBodySetup;
+class UObject;
+struct FNavigationElement;
+struct FNavigationRelevantData;
+struct FNavigableGeometryExport;
+enum class ENavDataGatheringMode : uint8;
+namespace EHasCustomNavigableGeometry
+{
+	enum Type : int;
+}
+
+DECLARE_DELEGATE_TwoParams(FNavigationDataExport, const FNavigationElement&, FNavigationRelevantData&);
+DECLARE_DELEGATE_ThreeParams(FCustomGeometryExport, const FNavigationElement&, FNavigableGeometryExport&, bool&);
+DECLARE_DELEGATE_ThreeParams(FGeometrySliceExport, const FNavigationElement&, FNavigableGeometryExport&, const FBox&);
+
+/**
+ * Structure used to identify a unique navigation element registered in the Navigation system.
+ * The handle can be created to represent two use cases:
+ *  1. Single UObject representing the navigation element
+ *     constructed from a UObject raw or weak pointer
+ *  2. Single UObject managing multiple non-UObject navigation elements
+ *     constructed from a UObject raw or weak pointer and using the optional constructor parameter to identity a unique sub-element
+ * @see FNavigationElement
+ */
+struct FNavigationElementHandle
+{
+	static ENGINE_API const FNavigationElementHandle Invalid;
+
+	FNavigationElementHandle() = default;
+
+	inline explicit FNavigationElementHandle(const UObject* Object, const uint64 SubElementId = INDEX_NONE);
+	inline explicit FNavigationElementHandle(const TWeakObjectPtr<const UObject>& WeakObject, const uint64 SubElementId = INDEX_NONE);
+
+	inline bool operator==(const FNavigationElementHandle& Other) const;
+	inline bool operator!=(const FNavigationElementHandle& Other) const;
+
+	/**
+	 * Conversion operator used to convert the handle to a boolean based on its validity.
+	 * @see IsValid
+	 */
+	inline explicit operator bool() const;
+
+	/** @return true if the handle has been properly assigned. */
+	inline bool IsValid() const;
+
+	/** Hash function to use TSet/TMap */
+	inline friend uint32 GetTypeHash(const FNavigationElementHandle& Handle);
+
+	/** Stringifies FNavigationElementHandle */
+	ENGINE_API friend FString LexToString(const FNavigationElementHandle& Handle);
+
+private:
+
+	/**
+	 * Main mandatory part of the handle to associate it with either a UObject with a 1:1 relation
+	 * with the registered navigation element, or a UObject that manages multiple non-UObject navigation elements.
+	 */
+	TWeakObjectPtr<const UObject> OwnerUObject;
+
+	/**
+	 * Second optional part of the handle used when the associated UObject manages multiple navigation elements.
+	 */
+	uint64 SubElementId = INDEX_NONE;
+};
+
+
+/**
+ * Structure registered in the navigation system that holds the required properties and delegates
+ * to gather navigation data (navigable geometry, NavArea modifiers, NavLinks, etc.) and be stored
+ * in the navigation octree.
+ *
+ * It represents a single element spatially located in a defined area in the level.
+ * That element can be created to represent two use cases:
+ *  1. Single UObject representing the navigation element
+ *     constructed from a UObject raw or weak pointer
+ *  2. Single UObject managing multiple non-UObject navigation elements
+ *     constructed from a UObject raw or weak pointer and using the optional constructor parameter to identity a unique sub-element
+ * @see FNavigationElementHandle
+ */
+struct FNavigationElement : TSharedFromThis<FNavigationElement>
+{
+private:
+	/** Temporary solution to provide access to deprecated constructors. */
+	friend struct FNavigationDirtyElement;
+    struct FPrivateToken { explicit FPrivateToken() = default; };
+	static TSharedRef<FNavigationElement> MakeFromUObject_DEPRECATED(UObject* InOwner)
+	{
+		return MakeShared<FNavigationElement>(FPrivateToken{}, InOwner, INDEX_NONE);
+	}
+
+public:
+
+    // Public, but can only be called by FNavigationElement and friends, because it needs access to FPrivateToken
+    ENGINE_API explicit FNavigationElement(FPrivateToken, const UObject* Object, uint64 SubElementId /*= INDEX_NONE*/);
+
+	FNavigationElement() = delete;
+	ENGINE_API explicit FNavigationElement(const UObject& Object, uint64 SubElementId = INDEX_NONE);
+	ENGINE_API explicit FNavigationElement(const INavRelevantInterface& NavRelevant, uint64 SubElementId = INDEX_NONE);
+
+	UE_DEPRECATED(5.5, "Temporary constructor to allow deprecation of other navigation types. Use the version taking an object reference instead.")
+	ENGINE_API explicit FNavigationElement(const UObject* Object, const uint64 SubElementId = INDEX_NONE)
+		: FNavigationElement(FPrivateToken{}, Object, SubElementId)
+	{
+	}
+
+	/** Factory helper function to create and initialize a sharable element from INavRelevantInterface. */
+	static ENGINE_API TSharedRef<const FNavigationElement> CreateFromNavRelevantInterface(const INavRelevantInterface& NavRelevantInterface);
+
+	/**
+	 * Returns a weak pointer to the UObject associated with this navigation element.
+	 * @return Weak pointer to the associated UObject if any.
+	 * @note If called from the Game thread all regular TWeakObjectPtr concerns and limitation applies.
+	 */
+	inline TWeakObjectPtr<const UObject> GetWeakUObject() const;
+
+	/**
+	 * Returns a handle that can be used to uniquely identify this element.
+	 * @return Handle representing the current element.
+	 */
+	ENGINE_API FNavigationElementHandle GetHandle() const;
+
+	/** @return Transform to use for the default geometry export. */
+	inline const FTransform& GetTransform() const;
+
+	/** Sets the Transform to use for the default geometry export. */
+	inline void SetTransform(const FTransform& InGeometryTransform);
+
+	/** @return The bounds to use when registering the element in the navigation octree. */
+	inline FBox GetBounds() const;
+
+	/** Sets the Bounds to use when registering the element in the navigation octree. */
+	inline void SetBounds(const FBox& InBounds);
+
+	/** @return The associated body setup (if any) to use for the default geometry export. */
+	inline UBodySetup* GetBodySetup() const;
+
+	/** Sets the body setup to use for the default geometry export. */
+	ENGINE_API void SetBodySetup(UBodySetup* InBodySetup);
+
+	/** @return The UObject for which the associated octree node will be used to hold the current element navigation data. */
+	inline const TWeakObjectPtr<const UObject>& GetNavigationParent() const;
+
+	/** Sets the UObject for which the associated octree node will be used to hold the current element navigation data. */
+	inline void SetNavigationParent(const UObject* InNavigationParent);
+
+	/**
+	 * @return The type of geometry export to use for the current element
+	 * @see EHasCustomNavigableGeometry
+	 */
+	inline EHasCustomNavigableGeometry::Type GetGeometryExportType() const;
+
+	/**
+	 * Sets the type of geometry export to use for the current element
+	 * @see GeometryExportType
+	 */
+	inline void SetGeometryExportType(const EHasCustomNavigableGeometry::Type InCustomNavigableGeometry);
+
+	/**
+	 * @return The mode that indicates when the geometry gathering must be executed.
+	 * @see GeometryGatheringMode
+	 */
+	inline ENavDataGatheringMode GetGeometryGatheringMode() const;
+
+	/**
+	 * Sets the mode that indicates when the geometry gathering must be executed.
+	 * @see GeometryGatheringMode
+	 */
+	inline void SetGeometryGatheringMode(const ENavDataGatheringMode InGeometryGatheringMode);
+
+	/**
+	 * @return True the area covered by the navigation bounds should be dirtied when inserting,
+	 * or removing, the element in the navigation octree (default behavior).
+	 * @see DirtyAreaOnRegistration
+	 */
+	inline bool GetDirtyAreaOnRegistration() const;
+
+	/**
+	 * Sets whether the area covered by the navigation bounds should be dirtied when inserting,
+	 * or removing, the element in the navigation octree (default behavior).
+	 * @see DirtyAreaOnRegistration
+	 */
+	inline void SetDirtyAreaOnRegistration(const bool bInDirtyAreaOnRegistration);
+
+	/**
+	 * @return True if the element was created from a UObject that is associated with a data layer
+	 * in the list of runtime data layers that should be included in the base navigation data (cooked)
+	 * or directly placed in the level.
+	 */
+	inline bool IsInBaseNavigationData() const;
+
+	/**
+	 * @return True if the element was created from a UObject while its level was pending
+	 * being made invisible or visible (i.e., loading/unloading).
+	 */
+	inline bool IsFromLevelVisibilityChange() const;
+
+	/** @return Name describing the element based on the associated UObject name and sub-element ID (if any). */
+	ENGINE_API FString GetName() const;
+
+	/** @return Name describing the element based on the associated UObject pathname and sub-element ID (if any). */
+	ENGINE_API FString GetPathName() const;
+
+	/** @return Name describing the element based on the associated UObject fullname and sub-element ID (if any). */
+	ENGINE_API FString GetFullName() const;
+
+	/** @return Name describing the specified element based on the name of its associated UObject and sub-element ID (if any). */
+	inline friend FString GetNameSafe(const FNavigationElement* Element);
+
+	/** @return Name describing the specified element based on the pathname of its associated UObject and sub-element ID (if any). */
+	inline friend FString GetPathNameSafe(const FNavigationElement* Element);
+
+	/** @return Name describing the specified element based on the fullname of its associated UObject and sub-element ID (if any). */
+	inline friend FString GetFullNameSafe(const FNavigationElement* Element);
+
+	/** Hash function to use TSet/TMap */
+	ENGINE_API friend uint32 GetTypeHash(const FNavigationElement& Element);
+
+	/** Stringifies FNavigationElementHandle */
+	ENGINE_API friend FString LexToString(const FNavigationElement& Element);
+
+	/** Delegate used to gather navigation data like NavLinks, NavAreaModifiers, etc. */
+	FNavigationDataExport NavigationDataExportDelegate;
+
+	/**
+	 * Delegate used during the geometry export based on the specified GeometryExportType
+	 * @see GeometryExportType
+	 */
+	FCustomGeometryExport CustomGeometryExportDelegate;
+
+	/** Delegate that can be used by very large elements to gather a limited piece of geometry for navigation generation in a given area. */
+	FGeometrySliceExport GeometrySliceExportDelegate;
+
+private:
+
+	ENGINE_API explicit FNavigationElement(const UObject& Object, uint64 SubElementId, bool bTryInitializeFromInterface);
+	void InitializeFromInterface(const INavRelevantInterface* NavRelevantInterface);
+	
+	/**
+	 * The associated UObject provided to create the element.
+	 * It is still currently used after creation for: ShouldSkipObjectPredicate (dirty areas), MetaModifier, logging and validation.
+	 * This is always set when creating the object and should be valid for the lifetime of the element if it is properly unregistered
+	 * based on the UObject lifetime (e.g., a Component will unregister its associated element when it gets unregistered from the world).
+	 */
+	TWeakObjectPtr<const UObject> OwnerUObject;
+
+	/**
+	 * Optional ID that can be used when a UObject owns and manages multiple sub-elements which should be represented individually
+	 * in the navigation system (i.e., each sub-element has its own bounds and should be registered in an octree node).
+	 */
+	uint64 SubElementId = INDEX_NONE;
+
+	/** Associated body setup (if any) used by the default geometry export. */
+	TWeakObjectPtr<UBodySetup> BodySetup = nullptr;
+
+	/** Transform used by the default geometry export. */
+	FTransform GeometryTransform = FTransform::Identity;
+
+	/** Bounds used to register the element in the navigation octree. */
+	FBox Bounds;
+
+	/**
+	 * Indicates that this element is not registered as a new node in the navigation octree, but instead
+	 * adds its navigation data to the parent octree node.
+	 */
+	TWeakObjectPtr<const UObject> NavigationParent = nullptr;
+
+	/**
+	 * Indicates the type of geometry export to use for the current element.
+	 * @see EHasCustomNavigableGeometry
+	 */
+	EHasCustomNavigableGeometry::Type GeometryExportType;
+
+	/**
+	 * Indicates when the geometry gathering must be executed:
+	 * - Instant: The geometry export is executed as soon as the element is registered in the navigation octree.
+	 * - Lazy: The geometry export is executed only when it is required to rebuild the navigation data intersecting with a given area (e.g., Tile generator).
+	 * - Default: Use the default gathering mode defined by the navigation system.
+	 */
+	ENavDataGatheringMode GeometryGatheringMode;
+
+	/**
+	 * Indicates if the area covered by the navigation bounds should be dirtied when inserting, or removing, the element in the navigation octree (default behavior).
+	 * When returning false it is expected that the element will explicitly dirty areas (e.g. using UpdateNavigationElementBoundsDelegate).
+	 */
+	bool bDirtyAreaOnRegistration = true;
+
+	/**
+	 * Indicates that the element was created from a UObject that is associated with a data layer
+	 * in the list of runtime data layers that should be included in the base navigation data (cooked)
+	 * or directly placed in the level.
+	 */
+	bool bIsInBaseNavigationData = false;
+
+	/** Indicates that the element was created from a UObject while its level was pending being made invisible or visible (i.e., loading/unloading). */
+	bool bIsFromLevelVisibilityChange = false;
+};
+
+
+//----------------------------------------------------------------------//
+// Inlines
+//----------------------------------------------------------------------//
+uint32 GetTypeHash(const FNavigationElementHandle& Key)
+{
+	return HashCombine(GetTypeHash(Key.OwnerUObject), GetTypeHash(Key.SubElementId));
+}
+
+FString GetNameSafe(const FNavigationElement* Element)
+{
+	return Element ? Element->GetName() : TEXT("None");
+}
+
+FString GetPathNameSafe(const FNavigationElement* Element)
+{
+	return Element ? Element->GetPathName() : TEXT("None");
+}
+
+FString GetFullNameSafe(const FNavigationElement* Element)
+{
+	return Element ? Element->GetFullName() : TEXT("None");
+}
+
+
+//----------------------------------------------------------------------//
+// Inlines (FNavigationElementHandle)
+//----------------------------------------------------------------------//
+
+FNavigationElementHandle::FNavigationElementHandle(const UObject* Object, const uint64 SubElementId /*= INDEX_NONE*/)
+	: OwnerUObject(Object)
+	, SubElementId(SubElementId)
+{
+}
+
+FNavigationElementHandle::FNavigationElementHandle(const TWeakObjectPtr<const UObject>& WeakObject, const uint64 SubElementId /*= INDEX_NONE*/)
+	: OwnerUObject(WeakObject)
+	, SubElementId(SubElementId)
+{
+}
+
+bool FNavigationElementHandle::operator==(const FNavigationElementHandle& Other) const
+{
+	return OwnerUObject.HasSameIndexAndSerialNumber(Other.OwnerUObject) && SubElementId == Other.SubElementId;
+}
+
+bool FNavigationElementHandle::operator!=(const FNavigationElementHandle& Other) const
+{
+	return !(*this == Other);
+}
+
+FNavigationElementHandle::operator bool() const
+{
+	return IsValid();
+}
+
+bool FNavigationElementHandle::IsValid() const
+{
+	return *this != Invalid;
+}
+
+
+//----------------------------------------------------------------------//
+// Inlines (FNavigationElement)
+//----------------------------------------------------------------------//
+
+TWeakObjectPtr<const UObject> FNavigationElement::GetWeakUObject() const
+{
+	checkf(IsInGameThread(), TEXT("Currenlty not allowing this to be called from non-game thread."));
+	return OwnerUObject;
+}
+
+bool FNavigationElement::IsFromLevelVisibilityChange() const
+{
+	return bIsFromLevelVisibilityChange;
+}
+
+bool FNavigationElement::IsInBaseNavigationData() const
+{
+	return bIsInBaseNavigationData;
+}
+
+const FTransform& FNavigationElement::GetTransform() const
+{
+	return GeometryTransform;
+}
+
+void FNavigationElement::SetTransform(const FTransform& InGeometryTransform)
+{
+	GeometryTransform = InGeometryTransform;
+}
+
+FBox FNavigationElement::GetBounds() const
+{
+	return Bounds;
+}
+
+void FNavigationElement::SetBounds(const FBox& InBounds)
+{
+	Bounds = InBounds;
+}
+
+UBodySetup* FNavigationElement::GetBodySetup() const
+{
+	return BodySetup.Get();
+}
+
+const TWeakObjectPtr<const UObject>& FNavigationElement::GetNavigationParent() const
+{
+	return NavigationParent;
+}
+
+void FNavigationElement::SetNavigationParent(const UObject* InNavigationParent)
+{
+	NavigationParent = InNavigationParent;
+}
+
+EHasCustomNavigableGeometry::Type FNavigationElement::GetGeometryExportType() const
+{
+	return GeometryExportType;
+}
+
+void FNavigationElement::SetGeometryExportType(const EHasCustomNavigableGeometry::Type InCustomNavigableGeometry)
+{
+	GeometryExportType = InCustomNavigableGeometry;
+}
+
+ENavDataGatheringMode FNavigationElement::GetGeometryGatheringMode() const
+{
+	return GeometryGatheringMode;
+}
+
+void FNavigationElement::SetGeometryGatheringMode(const ENavDataGatheringMode InGeometryGatheringMode)
+{
+	GeometryGatheringMode = InGeometryGatheringMode;
+}
+
+bool FNavigationElement::GetDirtyAreaOnRegistration() const
+{
+	return bDirtyAreaOnRegistration;
+}
+
+void FNavigationElement::SetDirtyAreaOnRegistration(const bool bInDirtyAreaOnRegistration)
+{
+	bDirtyAreaOnRegistration = bInDirtyAreaOnRegistration;
+}
diff --git a/Engine/Source/Runtime/Engine/Classes/AI/Navigation/NavigationRelevantData.h b/Engine/Source/Runtime/Engine/Classes/AI/Navigation/NavigationRelevantData.h
--- a/Engine/Source/Runtime/Engine/Classes/AI/Navigation/NavigationRelevantData.h
+++ b/Engine/Source/Runtime/Engine/Classes/AI/Navigation/NavigationRelevantData.h
@@ -2,12 +2,15 @@
 
 #pragma once
 
-#include "AI/Navigation/NavigationDirtyArea.h"
-#include "AI/Navigation/NavigationTypes.h"
-#include "AI/NavigationModifier.h"
 #if UE_ENABLE_INCLUDE_ORDER_DEPRECATED_IN_5_5
 #include "UObject/ObjectMacros.h"
-#endif //UE_ENABLE_INCLUDE_ORDER_DEPRECATED_IN_5_5
+#include "AI/Navigation/NavigationTypes.h"
+#endif // UE_ENABLE_INCLUDE_ORDER_DEPRECATED_IN_5_5
+
+#include "AI/NavigationModifier.h"
+#include "NavigationDirtyArea.h"
+
+struct FNavigationElement;
 
 struct FNavigationRelevantDataFilter 
 {
@@ -66,8 +69,14 @@ struct FNavigationRelevantData : public TSharedFromThis<FNavigationRelevantData,
 	/** additional modifiers: areas and external links */
 	FCompositeNavModifier Modifiers;
 
+#if WITH_EDITORONLY_DATA
 	/** UObject these data represents */
+	UE_DEPRECATED(5.5, "Use SourceElement instead.")
 	TWeakObjectPtr<UObject> SourceObject;
+#endif // WITH_EDITORONLY_DATA
+
+	/** Navigation element this data represents */
+	TSharedRef<const FNavigationElement> SourceElement;
 
 	/** get set to true when lazy navigation exporting is enabled and this navigation data has "potential" of
 	*	containing geometry data. First access will result in gathering the data and setting this flag back to false.
@@ -84,8 +93,13 @@ struct FNavigationRelevantData : public TSharedFromThis<FNavigationRelevantData,
 	/** From level loading (only valid in WP dynamic mode) */
 	uint32 bLoadedData : 1;
 
-	FNavigationRelevantData(UObject& Source)
-		: SourceObject(&Source)
+	FNavigationRelevantData() = delete;
+
+	UE_DEPRECATED(5.5, "Use the constructor using FNavigationElement instead.")
+	ENGINE_API FNavigationRelevantData(UObject& Source);
+
+	explicit FNavigationRelevantData(const TSharedRef<const FNavigationElement>& Source)
+		: SourceElement(Source)
 		, bPendingLazyGeometryGathering(false)
 		, bPendingLazyModifiersGathering(false)
 		, bPendingChildLazyModifiersGathering(false)
@@ -129,10 +143,7 @@ struct FNavigationRelevantData : public TSharedFromThis<FNavigationRelevantData,
 				: ENavigationDirtyFlag::None);
 	}
 
-	FORCEINLINE FCompositeNavModifier GetModifierForAgent(const struct FNavAgentProperties* NavAgent = nullptr) const
-	{
-		return Modifiers.HasMetaAreas() ? Modifiers.GetInstantiatedMetaModifier(NavAgent, SourceObject) : Modifiers;
-	}
+	ENGINE_API FCompositeNavModifier GetModifierForAgent(const FNavAgentProperties* NavAgent = nullptr) const;
 
 	ENGINE_API bool HasPerInstanceTransforms() const;
 	ENGINE_API bool IsMatchingFilter(const FNavigationRelevantDataFilter& Filter) const;
@@ -151,6 +162,9 @@ struct FNavigationRelevantData : public TSharedFromThis<FNavigationRelevantData,
 		}
 	}
 
-	FORCEINLINE UObject* GetOwner() const { return SourceObject.Get(); }
-	ENGINE_API FORCEINLINE decltype(SourceObject)& GetOwnerPtr() { return SourceObject; }
+	UE_DEPRECATED(5.5, "Use SourceElement instead.")
+	ENGINE_API const UObject* GetOwner() const;
+
+	UE_DEPRECATED(5.5, "Use SourceElement instead.")
+	ENGINE_API TWeakObjectPtr<UObject> GetOwnerPtr() const;
 };
diff --git a/Engine/Source/Runtime/Engine/Classes/AI/NavigationSystemBase.h b/Engine/Source/Runtime/Engine/Classes/AI/NavigationSystemBase.h
--- a/Engine/Source/Runtime/Engine/Classes/AI/NavigationSystemBase.h
+++ b/Engine/Source/Runtime/Engine/Classes/AI/NavigationSystemBase.h
@@ -19,6 +19,8 @@ class AWorldSettings;
 class ULevel;
 class AController;
 class UNavAreaBase;
+struct FNavigationElement;
+struct FNavigationElementHandle;
 
 ENGINE_API DECLARE_LOG_CATEGORY_EXTERN(LogNavigation, Warning, All);
 ENGINE_API DECLARE_LOG_CATEGORY_EXTERN(LogNavigationDataBuild, Log, All);
@@ -156,10 +158,15 @@ namespace FNavigationSystem
 	ENGINE_API void UnregisterNavRelevantObject(UObject& Object);
 	ENGINE_API void OnObjectBoundsChanged(UObject& Object, const FBox& NewBounds, TConstArrayView<FBox> DirtyAreas);
 
+	ENGINE_API FNavigationElementHandle AddNavigationElement(UWorld* World, FNavigationElement&& Element);
+	ENGINE_API void UpdateNavigationElement(UWorld* World, FNavigationElementHandle ElementHandle, FNavigationElement&& Element);
+	ENGINE_API void RemoveNavigationElement(UWorld* World, FNavigationElementHandle ElementHandle);
+	ENGINE_API void OnNavigationElementBoundsChanged(UWorld* World, FNavigationElementHandle ElementHandle, const FBox& NewBounds, TConstArrayView<FBox> DirtyAreas);
+
 	ENGINE_API void RemoveActorData(AActor& Actor);
 
 	ENGINE_API bool HasComponentData(UActorComponent& Comp);
-	
+
 	ENGINE_API const FNavDataConfig& GetDefaultSupportedAgent();
 	ENGINE_API const FNavDataConfig& GetBiggestSupportedAgent(const UWorld* World);
 	ENGINE_API double GetWorldPartitionNavigationDataBuilderOverlap(const UWorld& World);
@@ -197,6 +204,11 @@ namespace FNavigationSystem
 	ENGINE_API void StopMovement(const AController& Controller);
 	ENGINE_API IPathFollowingAgentInterface* FindPathFollowingAgentForActor(const AActor& Actor);
 
+	DECLARE_DELEGATE_RetVal_TwoParams(FNavigationElementHandle, FAddNavigationElementSignature, UWorld* /*World*/, FNavigationElement&& /*Element*/);
+	DECLARE_DELEGATE_TwoParams(FRemoveNavigationElementSignature, UWorld* /*World*/, FNavigationElementHandle /*ElementHandle*/);
+	DECLARE_DELEGATE_ThreeParams(FUpdateNavigationElementSignature, UWorld* /*World*/, FNavigationElementHandle /*ElementHandle*/, FNavigationElement&& /*Element*/)
+	DECLARE_DELEGATE_FourParams(FUpdateNavigationElementBoundsSignature, UWorld* /*World*/, FNavigationElementHandle /*ElementHandle*/, const FBox& /*NewBounds*/, TConstArrayView<FBox> /*DirtyAreas*/)
+
 	DECLARE_DELEGATE_OneParam(FObjectBasedSignature, UObject& /*Object*/);
 	DECLARE_DELEGATE_ThreeParams(FObjectBoundsChangedSignature, UObject& /*Object*/, const FBox& /*NewBounds*/, TConstArrayView<FBox> /*DirtyAreas*/)
 	DECLARE_DELEGATE_OneParam(FActorBasedSignature, AActor& /*Actor*/);
@@ -288,6 +300,11 @@ protected:
 
 	static ENGINE_API void ResetEventDelegates();
 
+	static ENGINE_API FNavigationSystem::FAddNavigationElementSignature& GetAddNavigationElementDelegate();
+	static ENGINE_API FNavigationSystem::FUpdateNavigationElementSignature& GetUpdateNavigationElementDelegate();
+	static ENGINE_API FNavigationSystem::FRemoveNavigationElementSignature& GetRemoveNavigationElementDelegate();
+	static ENGINE_API FNavigationSystem::FUpdateNavigationElementBoundsSignature& GetUpdateNavigationElementBoundsDelegate();
+
 	static ENGINE_API FNavigationSystem::FObjectBasedSignature& RegisterNavRelevantObjectDelegate();
 	static ENGINE_API FNavigationSystem::FObjectBasedSignature& UpdateNavRelevantObjectDelegate();
 	static ENGINE_API FNavigationSystem::FObjectBasedSignature& UnregisterNavRelevantObjectDelegate();
diff --git a/Engine/Source/Runtime/Engine/Private/AI/Navigation/NavigationDirtyArea.cpp b/Engine/Source/Runtime/Engine/Private/AI/Navigation/NavigationDirtyArea.cpp
--- a/Engine/Source/Runtime/Engine/Private/AI/Navigation/NavigationDirtyArea.cpp
+++ b/Engine/Source/Runtime/Engine/Private/AI/Navigation/NavigationDirtyArea.cpp
@@ -1,26 +1,34 @@
 // Copyright Epic Games, Inc. All Rights Reserved.
 
 #include "AI/Navigation/NavigationDirtyArea.h"
+#include "AI/Navigation/NavigationElement.h"
 #if !NO_LOGGING
 #include "AI/NavigationSystemBase.h"
 #endif
 
-FNavigationDirtyArea::FNavigationDirtyArea(const FBox& InBounds, const ENavigationDirtyFlag InFlags, UObject* const InOptionalSourceObject)
+FNavigationDirtyArea::FNavigationDirtyArea(const FBox& InBounds, const ENavigationDirtyFlag InFlags, const TSharedPtr<const FNavigationElement>& InOptionalSourceElement)
 	: Bounds(InBounds)
-	, OptionalSourceObject(InOptionalSourceObject)
+	, OptionalSourceElement(InOptionalSourceElement)
 	, Flags(InFlags)
 {
 #if !NO_LOGGING
 	if (!Bounds.IsValid || Bounds.ContainsNaN())
 	{
-		UE_LOG(LogNavigation, Warning, TEXT("Creation of FNavigationDirtyArea with invalid bounds%s. Bounds: %s, SourceObject: %s."),
-			Bounds.ContainsNaN() ? TEXT(" (contains NaN)") : TEXT(""), *Bounds.ToString(), *GetFullNameSafe(OptionalSourceObject.Get()));
+		const FNavigationElement* Element = OptionalSourceElement.Get();
+		UE_LOG(LogNavigation, Warning, TEXT("Creation of FNavigationDirtyArea with invalid bounds%s. Bounds: %s, SourceElement: %s."),
+			Bounds.ContainsNaN() ? TEXT(" (contains NaN)") : TEXT(""), *Bounds.ToString(), *GetFullNameSafe(Element));
 	}
 #endif //!NO_LOGGING
 }
 
+FString FNavigationDirtyArea::GetSourceDescription() const
+{
+	const FNavigationElement* SourceElement = OptionalSourceElement.Get();
+	return SourceElement ? SourceElement->GetFullName() : TEXT("");
+}
+
 // Deprecated
-FNavigationDirtyArea::FNavigationDirtyArea(const FBox& InBounds, int32 InFlags, UObject* const InOptionalSourceObject /*= nullptr*/)
-	: FNavigationDirtyArea(InBounds, static_cast<ENavigationDirtyFlag>(InFlags), InOptionalSourceObject)
+FNavigationDirtyArea::FNavigationDirtyArea(const FBox& InBounds, const int32 InFlags, UObject* const InOptionalSourceObject)
+	: FNavigationDirtyArea(InBounds, static_cast<ENavigationDirtyFlag>(InFlags), /*InOptionalSourceElement*/nullptr)
 {
-}
\ No newline at end of file
+}
diff --git a/Engine/Source/Runtime/Engine/Private/AI/Navigation/NavigationDirtyElement.cpp b/Engine/Source/Runtime/Engine/Private/AI/Navigation/NavigationDirtyElement.cpp
--- a/Engine/Source/Runtime/Engine/Private/AI/Navigation/NavigationDirtyElement.cpp
+++ b/Engine/Source/Runtime/Engine/Private/AI/Navigation/NavigationDirtyElement.cpp
@@ -2,47 +2,59 @@
 
 #include "AI/Navigation/NavigationDirtyElement.h"
 #include "AI/Navigation/NavigationDirtyArea.h"
+#include "AI/Navigation/NavigationElement.h"
 #include "AI/Navigation/NavigationTypes.h"
 
 FNavigationDirtyElement::FNavigationDirtyElement(
-	UObject* InOwner,
-	INavRelevantInterface* InNavInterface,
+	const TSharedRef<const FNavigationElement>& InNavigationElement,
 	const ENavigationDirtyFlag InFlagsOverride,
 	const bool bUseWorldPartitionedDynamicMode /*= false*/)
-	: Owner(InOwner)
-	, NavInterface(InNavInterface)
+	: NavigationElement(InNavigationElement)
 	, FlagsOverride(InFlagsOverride)
 	, PrevFlags(ENavigationDirtyFlag::None)
+	, bIsFromVisibilityChange(bUseWorldPartitionedDynamicMode && NavigationElement->IsFromLevelVisibilityChange())
+	, bIsInBaseNavmesh(bUseWorldPartitionedDynamicMode && NavigationElement->IsInBaseNavigationData())
 {
-	if (bUseWorldPartitionedDynamicMode)
-	{
-		bIsFromVisibilityChange = FNavigationSystem::IsLevelVisibilityChanging(InOwner);
-		bIsInBaseNavmesh = FNavigationSystem::IsInBaseNavmesh(InOwner);
-	}
-	else
-	{
-		bIsFromVisibilityChange = false;
-		bIsInBaseNavmesh = false;
-	}
 }
 
-FNavigationDirtyElement::FNavigationDirtyElement(UObject* InOwner, INavRelevantInterface* InNavInterface, const bool bUseWorldPartitionedDynamicMode /*= false*/)
-	: FNavigationDirtyElement(InOwner, InNavInterface, ENavigationDirtyFlag::None, bUseWorldPartitionedDynamicMode)
+FNavigationDirtyElement::FNavigationDirtyElement(const TSharedRef<const FNavigationElement>& InNavigationElement, const bool bUseWorldPartitionedDynamicMode /*= false*/)
+	: FNavigationDirtyElement(InNavigationElement, ENavigationDirtyFlag::None, bUseWorldPartitionedDynamicMode)
 {
 }
 
-FNavigationDirtyElement::FNavigationDirtyElement(UObject* InOwner)
-	: FNavigationDirtyElement(InOwner, /*InNavInterface*/nullptr, /*bUseWorldPartitionedDynamicMode*/false)
+uint32 GetTypeHash(const FNavigationDirtyElement& Info)
 {
+	return GetTypeHash(Info.NavigationElement.Get());
 }
 
+
+//----------------------------------------------------------------------//
+// Deprecated methods
+//----------------------------------------------------------------------//
+PRAGMA_DISABLE_DEPRECATION_WARNINGS
+
+// Deprecated
 FNavigationDirtyElement::FNavigationDirtyElement()
-	: FNavigationDirtyElement(/*InOwner*/nullptr, /*InNavInterface*/nullptr, /*bUseWorldPartitionedDynamicMode*/false)
+	: FNavigationDirtyElement(FNavigationElement::MakeFromUObject_DEPRECATED(nullptr))
 {
 }
 
 // Deprecated
-FNavigationDirtyElement::FNavigationDirtyElement(UObject* InOwner, INavRelevantInterface* InNavInterface, int32 InFlagsOverride /*= 0*/, const bool bUseWorldPartitionedDynamicMode /*= false*/)
-	: FNavigationDirtyElement(InOwner, InNavInterface, static_cast<ENavigationDirtyFlag>(InFlagsOverride), bUseWorldPartitionedDynamicMode)
+FNavigationDirtyElement::FNavigationDirtyElement(UObject* InOwner, INavRelevantInterface*, int32 InFlagsOverride /*= 0*/, const bool bUseWorldPartitionedDynamicMode /*= false*/)
+	: FNavigationDirtyElement(FNavigationElement::MakeFromUObject_DEPRECATED(InOwner), static_cast<ENavigationDirtyFlag>(InFlagsOverride), bUseWorldPartitionedDynamicMode)
 {
 }
+
+// Deprecated
+FNavigationDirtyElement::FNavigationDirtyElement(UObject* InOwner)
+	: FNavigationDirtyElement(FNavigationElement::MakeFromUObject_DEPRECATED(InOwner))
+{
+}
+
+// Deprecated
+bool FNavigationDirtyElement::operator==(const UObject*& OtherOwner) const
+{ 
+	return NavigationElement->GetHandle() == FNavigationElementHandle(OtherOwner);
+}
+
+PRAGMA_ENABLE_DEPRECATION_WARNINGS
diff --git a/Engine/Source/Runtime/Engine/Private/AI/Navigation/NavigationElement.cpp b/Engine/Source/Runtime/Engine/Private/AI/Navigation/NavigationElement.cpp
new file mode 100644
--- /dev/null
+++ b/Engine/Source/Runtime/Engine/Private/AI/Navigation/NavigationElement.cpp
@@ -0,0 +1,205 @@
+// Copyright Epic Games, Inc. All Rights Reserved.
+
+#include "AI/Navigation/NavigationElement.h"
+
+#include "AI/NavigationSystemBase.h"
+#include "AI/Navigation/NavigationTypes.h"
+#include "AI/Navigation/NavRelevantInterface.h"
+#include "PhysicsEngine/BodySetup.h"
+#include "UObject/Object.h"
+#include "Templates/TypeHash.h"
+
+namespace UE::NavigationElement::Private
+{
+#if !UE_BUILD_SHIPPING
+bool bValidateNavigationElementInitialization = false;
+
+FAutoConsoleVariableRef ConsoleVariables[] =
+{
+	FAutoConsoleVariableRef(
+		TEXT("ai.debug.nav.validateNavigationElementInitialization"),
+		bValidateNavigationElementInitialization,
+		TEXT("Used to validate that the values returned by INavRelevantInterface when initializing the FNavigationElement makes sense."
+			"Those warnings might not be critical since an update can be sent after and update then pending element but might still worth "
+			"investigating the use case to reduce redundant operations."))
+};
+#endif // !UE_BUILD_SHIPPING
+
+} // UE::NavigationOctree::Private
+
+const FNavigationElementHandle FNavigationElementHandle::Invalid = FNavigationElementHandle(nullptr, INDEX_NONE);
+
+FNavigationElement::FNavigationElement(const UObject& Object, const uint64 SubElementId /*= INDEX_NONE*/)
+	: OwnerUObject(&Object)
+	, SubElementId(SubElementId)
+	, GeometryExportType(EHasCustomNavigableGeometry::No)
+	, GeometryGatheringMode(ENavDataGatheringMode::Default)
+	, bIsInBaseNavigationData(FNavigationSystem::IsInBaseNavmesh(&Object))
+	, bIsFromLevelVisibilityChange(FNavigationSystem::IsLevelVisibilityChanging(&Object))
+{
+}
+
+FNavigationElement::FNavigationElement(FPrivateToken, const UObject* Object, uint64 SubElementId /*= INDEX_NONE*/)
+	: OwnerUObject(Object)
+	, SubElementId(SubElementId)
+	, GeometryExportType(EHasCustomNavigableGeometry::No)
+	, GeometryGatheringMode(ENavDataGatheringMode::Default)
+	, bIsInBaseNavigationData(FNavigationSystem::IsInBaseNavmesh(Object))
+	, bIsFromLevelVisibilityChange(FNavigationSystem::IsLevelVisibilityChanging(Object))
+{
+}
+
+FNavigationElement::FNavigationElement(const INavRelevantInterface& NavRelevant, const uint64 SubElementId)
+	: SubElementId(SubElementId)
+	, GeometryExportType(EHasCustomNavigableGeometry::No)
+	, GeometryGatheringMode(ENavDataGatheringMode::Default)
+
+{
+	if (const UObject* Object = Cast<UObject>(&NavRelevant))
+	{
+		OwnerUObject = Object;
+		bIsInBaseNavigationData = FNavigationSystem::IsInBaseNavmesh(Object);
+		bIsFromLevelVisibilityChange = FNavigationSystem::IsLevelVisibilityChanging(Object);
+	}
+	InitializeFromInterface(&NavRelevant);
+}
+
+FNavigationElement::FNavigationElement(const UObject& Object, const uint64 SubElementId, const bool bTryInitializeFromInterface)
+	: FNavigationElement(Object, SubElementId)
+{
+	if (bTryInitializeFromInterface)
+	{
+		InitializeFromInterface(Cast<INavRelevantInterface>(&Object));
+	}
+}
+
+void FNavigationElement::InitializeFromInterface(const INavRelevantInterface* NavRelevantInterface)
+{
+	if (NavRelevantInterface == nullptr)
+	{
+		return;
+	}
+
+	INavRelevantInterface* MutableInterface = const_cast<INavRelevantInterface*>(NavRelevantInterface);
+
+	bDirtyAreaOnRegistration = !NavRelevantInterface->ShouldSkipDirtyAreaOnAddOrRemove();
+
+	Bounds = NavRelevantInterface->GetNavigationBounds();
+	GeometryGatheringMode = NavRelevantInterface->GetGeometryGatheringMode();
+	BodySetup = MutableInterface->GetNavigableGeometryBodySetup();
+	GeometryTransform = NavRelevantInterface->GetNavigableGeometryTransform();
+	GeometryExportType = NavRelevantInterface->HasCustomNavigableGeometry();
+	NavigationParent = NavRelevantInterface->GetNavigationParent();
+
+	NavigationDataExportDelegate.BindWeakLambda(Cast<UObject>(NavRelevantInterface),
+		[NavRelevantInterface](const FNavigationElement&, FNavigationRelevantData& OutNavigationRelevantData)
+		{
+			NavRelevantInterface->GetNavigationData(OutNavigationRelevantData);
+		});
+
+	CustomGeometryExportDelegate.BindWeakLambda(Cast<UObject>(NavRelevantInterface),
+		[NavRelevantInterface](const FNavigationElement&, FNavigableGeometryExport& OutGeometry, bool& bOutShouldExportDefaultGeometry)
+		{
+			bOutShouldExportDefaultGeometry = NavRelevantInterface->DoCustomNavigableGeometryExport(OutGeometry);
+		});
+
+	if (NavRelevantInterface->SupportsGatheringGeometrySlices())
+	{
+		GeometrySliceExportDelegate.BindWeakLambda(Cast<UObject>(NavRelevantInterface),
+			[MutableInterface](const FNavigationElement&, FNavigableGeometryExport& OutGeometryExport, const FBox& InSliceBox)
+			{
+				MutableInterface->PrepareGeometryExportSync();
+				MutableInterface->GatherGeometrySlice(OutGeometryExport, InSliceBox);
+			});
+	}
+
+#if !UE_BUILD_SHIPPING
+	if (UE::NavigationElement::Private::bValidateNavigationElementInitialization)
+	{
+		if (!Bounds.IsValid)
+		{
+			UE_LOG(LogNavigation, Warning,
+			   TEXT("Initializing a FNavigationElement from '%s' that provides invalid navigation bounds."),
+			   *GetPathNameSafe(Cast<UObject>(NavRelevantInterface)));
+		}
+		else if (bDirtyAreaOnRegistration && FVector2d(Bounds.GetSize()).IsNearlyZero())
+		{
+			UE_LOG(LogNavigation, Warning,
+			   TEXT("Initializing a FNavigationElement from '%s' that provides empty navigation bounds."),
+			   *GetPathNameSafe(Cast<UObject>(NavRelevantInterface)));
+		}
+
+		if (!NavRelevantInterface->IsNavigationRelevant())
+		{
+			UE_LOG(LogNavigation, Warning,
+			   TEXT("Initializing a FNavigationElement from '%s' for which 'IsNavigationRelevant()' returns 'false'."),
+			   *GetPathNameSafe(Cast<UObject>(NavRelevantInterface)));
+		}
+	}
+#endif // !UE_BUILD_SHIPPING
+}
+
+TSharedRef<const FNavigationElement> FNavigationElement::CreateFromNavRelevantInterface(const INavRelevantInterface& NavRelevantInterface)
+{
+	return MakeShared<const FNavigationElement>(NavRelevantInterface, /*SubElementId*/INDEX_NONE);
+}
+
+void FNavigationElement::SetBodySetup(UBodySetup* InBodySetup)
+{
+	BodySetup = InBodySetup;
+}
+
+FNavigationElementHandle FNavigationElement::GetHandle() const
+{
+	return FNavigationElementHandle(OwnerUObject, SubElementId);
+}
+
+FString FNavigationElement::GetName() const
+{
+	const UObject* Owner = OwnerUObject.Get();
+
+	if (SubElementId != INDEX_NONE)
+	{
+		return FString::Printf(TEXT("%s - %llu"), *GetNameSafe(Owner), SubElementId);
+	}
+	return FString::Printf(TEXT("%s"), *GetNameSafe(Owner));
+}
+
+FString FNavigationElement::GetPathName() const
+{
+	const UObject* Owner = OwnerUObject.Get();
+	if (SubElementId != INDEX_NONE)
+	{
+		return FString::Printf(TEXT("%s - %llu"), *GetPathNameSafe(Owner), SubElementId);
+	}
+	return FString::Printf(TEXT("%s"), *GetPathNameSafe(Owner));
+}
+
+FString FNavigationElement::GetFullName() const
+{
+	const UObject* Owner = OwnerUObject.Get();
+	if (SubElementId != INDEX_NONE)
+	{
+		return FString::Printf(TEXT("%s - %llu"), *GetFullNameSafe(Owner), SubElementId);
+	}
+	return FString::Printf(TEXT("%s"), *GetFullNameSafe(Owner));
+}
+
+uint32 GetTypeHash(const FNavigationElement& Element)
+{
+	return HashCombine(GetTypeHash(Element.OwnerUObject), GetTypeHash(Element.SubElementId));
+}
+
+FString LexToString(const FNavigationElement& Element)
+{
+	return Element.GetName();
+}
+
+FString LexToString(const FNavigationElementHandle& Handle)
+{
+	if (Handle.SubElementId != INDEX_NONE)
+	{
+		return FString::Printf(TEXT("%s - %llu"), *GetNameSafe(Handle.OwnerUObject.Get()), Handle.SubElementId);
+	}
+	return FString::Printf(TEXT("%s"), *GetNameSafe(Handle.OwnerUObject.Get()));
+}
\ No newline at end of file
diff --git a/Engine/Source/Runtime/Engine/Private/AI/Navigation/NavigationModifier.cpp b/Engine/Source/Runtime/Engine/Private/AI/Navigation/NavigationModifier.cpp
--- a/Engine/Source/Runtime/Engine/Private/AI/Navigation/NavigationModifier.cpp
+++ b/Engine/Source/Runtime/Engine/Private/AI/Navigation/NavigationModifier.cpp
@@ -792,7 +792,12 @@ void FCompositeNavModifier::Empty()
 	NavMeshResolution = ENavigationDataResolution::Invalid;
 }
 
-FCompositeNavModifier FCompositeNavModifier::GetInstantiatedMetaModifier(const FNavAgentProperties* NavAgent, TWeakObjectPtr<UObject> WeakOwnerPtr) const
+FCompositeNavModifier FCompositeNavModifier::GetInstantiatedMetaModifier(const FNavAgentProperties* NavAgent, const TWeakObjectPtr<UObject> WeakOwnerPtr) const
+{
+	return GetInstantiatedMetaModifier(NavAgent, static_cast<const TWeakObjectPtr<const UObject>&>(WeakOwnerPtr));
+}
+
+FCompositeNavModifier FCompositeNavModifier::GetInstantiatedMetaModifier(const FNavAgentProperties* NavAgent, const TWeakObjectPtr<const UObject>& WeakOwnerPtr) const
 {
 	SCOPE_CYCLE_COUNTER(STAT_Navigation_MetaAreaTranslation);
 	FCompositeNavModifier Result;
@@ -801,13 +806,13 @@ FCompositeNavModifier FCompositeNavModifier::GetInstantiatedMetaModifier(const F
 	// should not be called when HasMetaAreas == false since it's a waste of performance
 	ensure(HasMetaAreas() == true);
 
-	UObject* ObjectOwner = WeakOwnerPtr.Get();
-	if (ObjectOwner == NULL)
+	const UObject* ObjectOwner = WeakOwnerPtr.Get();
+	if (ObjectOwner == nullptr)
 	{
 		return Result;
 	}
 	
-	auto FindActorOwner = [](UObject* Obj) -> const AActor*
+	auto FindActorOwner = [](const UObject* Obj) -> const AActor*
 	{
 		while (Obj)
 		{
@@ -820,8 +825,13 @@ FCompositeNavModifier FCompositeNavModifier::GetInstantiatedMetaModifier(const F
 		return nullptr;
 	};
 
-	const AActor* ActorOwner = Cast<AActor>(ObjectOwner) ? (AActor*)ObjectOwner : FindActorOwner(ObjectOwner);
-	if (ActorOwner == NULL)
+	const AActor* ActorOwner = Cast<AActor>(ObjectOwner);
+	if (ActorOwner == nullptr)
+	{
+		ActorOwner = FindActorOwner(ObjectOwner);
+	}
+
+	if (ActorOwner == nullptr)
 	{
 		return Result;
 	}
@@ -1005,10 +1015,3 @@ uint32 FCompositeNavModifier::GetAllocatedSize() const
 
 	return IntCastChecked<uint32>(MemUsed);
 }
-
-bool FCompositeNavModifier::HasPerInstanceTransforms() const
-{
-	PRAGMA_DISABLE_DEPRECATION_WARNINGS
-	return NavDataPerInstanceTransformDelegate.IsBound();
-	PRAGMA_ENABLE_DEPRECATION_WARNINGS
-}
diff --git a/Engine/Source/Runtime/Engine/Private/AI/Navigation/NavigationTypes.cpp b/Engine/Source/Runtime/Engine/Private/AI/Navigation/NavigationTypes.cpp
--- a/Engine/Source/Runtime/Engine/Private/AI/Navigation/NavigationTypes.cpp
+++ b/Engine/Source/Runtime/Engine/Private/AI/Navigation/NavigationTypes.cpp
@@ -1,13 +1,14 @@
 // Copyright Epic Games, Inc. All Rights Reserved.
 
 #include "AI/Navigation/NavigationTypes.h"
-#include "AI/NavigationSystemBase.h"
-#include "AI/Navigation/NavQueryFilter.h"
+#include "AI/Navigation/NavAreaBase.h"
 #include "AI/Navigation/NavigationRelevantData.h"
+#include "AI/Navigation/NavigationElement.h"
+#include "AI/Navigation/NavQueryFilter.h"
+#include "AI/NavigationSystemBase.h"
+#include "Components/ShapeComponent.h"
 #include "Engine/Level.h"
 #include "EngineStats.h"
-#include "Components/ShapeComponent.h"
-#include "AI/Navigation/NavAreaBase.h"
 #include "GameFramework/WorldSettings.h"
 #include "WorldPartition/DataLayer/DataLayerAsset.h"
 
@@ -137,6 +138,11 @@ bool FNavigationRelevantData::FCollisionDataHeader::IsValid(const uint8* RawData
 	return (RawDataSize == 0) || ((RawDataSize >= HeaderSize) && (((const FCollisionDataHeader*)RawData)->DataSize == RawDataSize));
 }
 
+FCompositeNavModifier FNavigationRelevantData::GetModifierForAgent(const FNavAgentProperties* NavAgent) const
+{
+	return Modifiers.HasMetaAreas() ? Modifiers.GetInstantiatedMetaModifier(NavAgent, SourceElement->GetWeakUObject()) : Modifiers;
+}
+
 bool FNavigationRelevantData::HasPerInstanceTransforms() const
 {
 	return NavDataPerInstanceTransformDelegate.IsBound();
@@ -166,7 +172,7 @@ bool FNavigationRelevantData::IsCollisionDataValid() const
 	const bool bIsValid = FCollisionDataHeader::IsValid(CollisionData.GetData(), CollisionData.Num());
 	if (!ensure(bIsValid))
 	{
-		UE_LOG(LogNavigation, Error, TEXT("NavOctree element has corrupted collision data! Owner:%s Bounds:%s"), *GetNameSafe(GetOwner()), *Bounds.ToString());
+		UE_LOG(LogNavigation, Error, TEXT("NavOctree element has corrupted collision data! Owner:%s Bounds:%s"), *SourceElement->GetName(), *Bounds.ToString());
 		return false;
 	}
 
@@ -402,3 +408,34 @@ TSubclassOf<UNavAreaBase> UNavAreaBase::PickAreaClassForAgent(const AActor& Acto
 	return GetClass();
 }
 
+//----------------------------------------------------------------------//
+// Deprecated methods
+//----------------------------------------------------------------------//
+
+PRAGMA_DISABLE_DEPRECATION_WARNINGS
+
+// Deprecated
+FNavigationRelevantData::FNavigationRelevantData(UObject& Source)
+	: SourceElement(MakeShared<const FNavigationElement>(Source))
+	, bPendingLazyGeometryGathering(false)
+	, bPendingLazyModifiersGathering(false)
+	, bPendingChildLazyModifiersGathering(false)
+	, bSupportsGatheringGeometrySlices(false)
+	, bShouldSkipDirtyAreaOnAddOrRemove(false)
+	, bLoadedData(false)
+{
+}
+
+// Deprecated
+const UObject* FNavigationRelevantData::GetOwner() const
+{
+	return SourceElement->GetWeakUObject().Get();
+}
+
+// Deprecated
+TWeakObjectPtr<UObject> FNavigationRelevantData::GetOwnerPtr() const
+{
+	return TWeakObjectPtr(const_cast<UObject*>(GetOwner()));
+}
+
+PRAGMA_ENABLE_DEPRECATION_WARNINGS
diff --git a/Engine/Source/Runtime/Engine/Private/AI/NavigationSystemBase.cpp b/Engine/Source/Runtime/Engine/Private/AI/NavigationSystemBase.cpp
--- a/Engine/Source/Runtime/Engine/Private/AI/NavigationSystemBase.cpp
+++ b/Engine/Source/Runtime/Engine/Private/AI/NavigationSystemBase.cpp
@@ -8,6 +8,9 @@
 #include "AI/Navigation/PathFollowingAgentInterface.h"
 #include "AI/NavigationSystemConfig.h"
 #include "AI/Navigation/NavigationDataChunk.h"
+#include "AI/Navigation/NavRelevantInterface.h"
+#include "AI/Navigation/NavigationElement.h"
+
 
 #include UE_INLINE_GENERATED_CPP_BY_NAME(NavigationSystemBase)
 
@@ -28,7 +31,7 @@ namespace FNavigationSystem
 		const auto& Levels = InWorld.GetLevels();
 		for (ULevel* Level : Levels)
 		{
-			UE_LOG(LogNavigation, Verbose, TEXT("%s for %s"), ANSI_TO_TCHAR(__FUNCTION__), *GetFullNameSafe(Level));
+			UE_LOG(LogNavigation, Verbose, TEXT("%hs for %s"), __FUNCTION__, *GetFullNameSafe(Level));
 			
 			for (UNavigationDataChunk* NavChunk : Level->NavDataChunks)
 			{
@@ -62,7 +65,7 @@ namespace FNavigationSystem
 			return FNavigationSystemRunMode::InvalidMode;
 
 		default:
-			UE_LOG(LogNavigation, Warning, TEXT("%s Unhandled world type, defaulting to FNavigationSystemRunMode::InvalidMode."), ANSI_TO_TCHAR(__FUNCTION__));
+			UE_LOG(LogNavigation, Warning, TEXT("%hs Unhandled world type, defaulting to FNavigationSystemRunMode::InvalidMode."), __FUNCTION__);
 			return FNavigationSystemRunMode::InvalidMode;
 		}
 	}
@@ -75,7 +78,7 @@ namespace FNavigationSystem
 	void AddNavigationSystemToWorld(UWorld& WorldOwner, const FNavigationSystemRunMode RunMode, UNavigationSystemConfig* NavigationSystemConfig, const bool bInitializeForWorld, const bool bOverridePreviousNavSys)
 	{
 		TRACE_CPUPROFILER_EVENT_SCOPE(FNavigationSystem::AddNavigationSystemToWorld);
-		UE_LOG(LogNavigation, Verbose, TEXT("%s bOverridePreviousNavSys=%i (WorldOwner: %s)"), ANSI_TO_TCHAR(__FUNCTION__), bOverridePreviousNavSys, *WorldOwner.GetOuter()->GetName());
+		UE_LOG(LogNavigation, Verbose, TEXT("%hs bOverridePreviousNavSys=%i (WorldOwner: %s)"), __FUNCTION__, bOverridePreviousNavSys, *WorldOwner.GetOuter()->GetName());
 
 		const FNavigationSystemRunMode ResolvedRunMode = (RunMode == FNavigationSystemRunMode::InferFromWorldMode) ? FindRunModeFromWorldType(WorldOwner) : RunMode;
 
@@ -127,6 +130,11 @@ namespace FNavigationSystem
 	class FDelegates
 	{
 	public:
+		FAddNavigationElementSignature AddNavigationElementDelegate;
+		FUpdateNavigationElementSignature UpdateNavigationElementDelegate;
+		FRemoveNavigationElementSignature RemoveNavigationElementDelegate;
+		FUpdateNavigationElementBoundsSignature UpdateNavigationElementBoundsDelegate;
+
 		FObjectBasedSignature RegisterNavRelevantObject;
 		FObjectBasedSignature UpdateNavRelevantObject;
 		FObjectBasedSignature UnregisterNavRelevantObject;
@@ -176,6 +184,11 @@ PRAGMA_ENABLE_DEPRECATION_WARNINGS
 
 		FDelegates()
 		{
+			AddNavigationElementDelegate.BindLambda([](UWorld* World, FNavigationElement&&) { return FNavigationElementHandle::Invalid; });
+			UpdateNavigationElementDelegate.BindLambda([](UWorld* World, FNavigationElementHandle, FNavigationElement&&) {});
+			RemoveNavigationElementDelegate.BindLambda([](UWorld* World, FNavigationElementHandle) {});
+			UpdateNavigationElementBoundsDelegate.BindLambda([](UWorld* World, FNavigationElementHandle, const FBox&, TConstArrayView<FBox>) {});
+
 			RegisterNavRelevantObject.BindLambda([](UObject&) {});
 			UpdateNavRelevantObject.BindLambda([](UObject&) {});
 			UnregisterNavRelevantObject.BindLambda([](UObject&) {});
@@ -198,7 +211,7 @@ PRAGMA_ENABLE_DEPRECATION_WARNINGS
 			IsFollowingAPath.BindLambda([](const AController&) { return false; });
 			HasComponentData.BindLambda([](UActorComponent&) { return false; });
 			GetDefaultSupportedAgent.BindStatic(&GetFallbackSupportedAgent);
-			GetBiggestSupportedAgent.BindStatic(&GetFallbackSupportedAgent);			
+			GetBiggestSupportedAgent.BindStatic(&GetFallbackSupportedAgent);
 			UpdateActorAndComponentData.BindLambda([](AActor&, bool) {});
 			OnComponentBoundsChanged.BindLambda([](UActorComponent&, const FBox&, const FBox&) {});
 			GetNavDataForActor.BindLambda([](const AActor&) { return nullptr; });
@@ -219,7 +232,30 @@ PRAGMA_ENABLE_DEPRECATION_WARNINGS
 
 	FDelegates Delegates;
 
-	void ResetDelegates() { new(&Delegates)FDelegates(); }
+	void ResetDelegates()
+	{
+		new(&Delegates)FDelegates();
+	}
+
+	FNavigationElementHandle AddNavigationElement(UWorld* World, FNavigationElement&& Element)
+	{
+		return Delegates.AddNavigationElementDelegate.Execute(World, MoveTemp(Element));
+	}
+
+	void UpdateNavigationElement(UWorld* World, const FNavigationElementHandle ElementHandle, FNavigationElement&& Element)
+	{
+		Delegates.UpdateNavigationElementDelegate.Execute(World, ElementHandle, MoveTemp(Element));
+	}
+
+	void RemoveNavigationElement(UWorld* World, const FNavigationElementHandle ElementHandle)
+	{
+		Delegates.RemoveNavigationElementDelegate.Execute(World, ElementHandle);
+	}
+
+	void OnNavigationElementBoundsChanged(UWorld* World, const FNavigationElementHandle ElementHandle, const FBox& NewBounds, const TConstArrayView<FBox> DirtyAreas)
+	{
+		Delegates.UpdateNavigationElementBoundsDelegate.Execute(World, ElementHandle, NewBounds, DirtyAreas);
+	}
 
 	void RegisterNavRelevantObject(UObject& Object) { Delegates.RegisterNavRelevantObject.Execute(Object); }
 	void UpdateNavRelevantObject(UObject& Object) { Delegates.UpdateNavRelevantObject.Execute(Object); }
@@ -407,6 +443,26 @@ void UNavigationSystemBase::SetDefaultObstacleArea(TSubclassOf<UNavAreaBase> InA
 
 void UNavigationSystemBase::ResetEventDelegates() { FNavigationSystem::ResetDelegates(); }
 
+FNavigationSystem::FAddNavigationElementSignature& UNavigationSystemBase::GetAddNavigationElementDelegate()
+{
+	return FNavigationSystem::Delegates.AddNavigationElementDelegate;
+}
+
+FNavigationSystem::FUpdateNavigationElementSignature& UNavigationSystemBase::GetUpdateNavigationElementDelegate()
+{
+	return FNavigationSystem::Delegates.UpdateNavigationElementDelegate;
+}
+
+FNavigationSystem::FRemoveNavigationElementSignature& UNavigationSystemBase::GetRemoveNavigationElementDelegate()
+{
+	return FNavigationSystem::Delegates.RemoveNavigationElementDelegate;
+}
+
+FNavigationSystem::FUpdateNavigationElementBoundsSignature& UNavigationSystemBase::GetUpdateNavigationElementBoundsDelegate()
+{
+	return FNavigationSystem::Delegates.UpdateNavigationElementBoundsDelegate;
+}
+
 FNavigationSystem::FObjectBasedSignature& UNavigationSystemBase::RegisterNavRelevantObjectDelegate() { return FNavigationSystem::Delegates.RegisterNavRelevantObject; }
 FNavigationSystem::FObjectBasedSignature& UNavigationSystemBase::UpdateNavRelevantObjectDelegate() { return FNavigationSystem::Delegates.UpdateNavRelevantObject; }
 FNavigationSystem::FObjectBasedSignature& UNavigationSystemBase::UnregisterNavRelevantObjectDelegate() { return FNavigationSystem::Delegates.UnregisterNavRelevantObject; }
diff --git a/Engine/Source/Runtime/Engine/Public/AI/NavigationModifier.h b/Engine/Source/Runtime/Engine/Public/AI/NavigationModifier.h
--- a/Engine/Source/Runtime/Engine/Public/AI/NavigationModifier.h
+++ b/Engine/Source/Runtime/Engine/Public/AI/NavigationModifier.h
@@ -362,11 +362,11 @@ struct FCompositeNavModifier : public FNavigationModifier
     void MarkAsPerInstanceModifier() { bIsPerInstanceModifier = true; }
 
 	/** returns a copy of Modifier */
-	ENGINE_API FCompositeNavModifier GetInstantiatedMetaModifier(const struct FNavAgentProperties* NavAgent, TWeakObjectPtr<UObject> WeakOwnerPtr) const;
-	ENGINE_API uint32 GetAllocatedSize() const;
+	ENGINE_API FCompositeNavModifier GetInstantiatedMetaModifier(const FNavAgentProperties* NavAgent, const TWeakObjectPtr<const UObject>& WeakOwnerPtr) const;
+	UE_DEPRECATED(5.5, "Use the version taking a const reference instead.")
+	ENGINE_API FCompositeNavModifier GetInstantiatedMetaModifier(const FNavAgentProperties* NavAgent, TWeakObjectPtr<UObject> WeakOwnerPtr) const;
 
-	UE_DEPRECATED(4.24, "This method will be removed in future versions. Use FNavigationRelevantData::HasPerInstanceTransforms instead.")
-	ENGINE_API bool HasPerInstanceTransforms() const;
+	ENGINE_API uint32 GetAllocatedSize() const;
 
 	TArray<FAreaNavModifier>& GetMutableAreas() { return Areas; }
 	TArray<FSimpleLinkNavModifier>& GetSimpleLinks() { return SimpleLinks; }
diff --git a/Engine/Source/Runtime/NavigationSystem/Private/NavMesh/RecastNavMesh.cpp b/Engine/Source/Runtime/NavigationSystem/Private/NavMesh/RecastNavMesh.cpp
--- a/Engine/Source/Runtime/NavigationSystem/Private/NavMesh/RecastNavMesh.cpp
+++ b/Engine/Source/Runtime/NavigationSystem/Private/NavMesh/RecastNavMesh.cpp
@@ -3051,7 +3051,14 @@ void ARecastNavMesh::OnStreamingNavDataAdded(ANavigationDataChunkActor& InActor)
 				{
 					UE_VLOG_BOX(this, LogNavigation, Verbose, NavElement.Bounds.GetBox(), FColor::Orange, TEXT(""));
 
-					NavSys->AddDirtyArea(NavElement.Bounds.GetBox(), ENavigationDirtyFlag::All, [&NavElement] { return NavElement.Data->SourceObject.Get(); }, "Streaming data added");
+					NavSys->AddDirtyArea(
+						NavElement.Bounds.GetBox(),
+						ENavigationDirtyFlag::All,
+						[SourceElement = NavElement.Data->SourceElement]
+						{
+							return SourceElement;
+						},
+						"Streaming data added");
 				}
 			}
 		}
diff --git a/Engine/Source/Runtime/NavigationSystem/Private/NavMesh/RecastNavMeshGenerator.cpp b/Engine/Source/Runtime/NavigationSystem/Private/NavMesh/RecastNavMeshGenerator.cpp
--- a/Engine/Source/Runtime/NavigationSystem/Private/NavMesh/RecastNavMeshGenerator.cpp
+++ b/Engine/Source/Runtime/NavigationSystem/Private/NavMesh/RecastNavMeshGenerator.cpp
@@ -8,7 +8,6 @@
 #include "HAL/FileManager.h"
 #include "Misc/Paths.h"
 #include "Serialization/MemoryWriter.h"
-#include "EngineGlobals.h"
 #include "GameFramework/PlayerController.h"
 #include "Engine/Engine.h"
 #include "Logging/LogScopedCategoryAndVerbosityOverride.h"
@@ -127,8 +126,8 @@ int32 GetTilesCountHelper(const dtNavMesh* DetourMesh)
 
 /**
  * Exports geometry to OBJ file. Can be used to verify NavMesh generation in RecastDemo app
- * @param FileName - full name of OBJ file with extension
- * @param GeomVerts - list of vertices
+ * @param InFileName - full name of OBJ file with extension
+ * @param GeomCoords - list of vertices
  * @param GeomFaces - list of triangles (3 vert indices for each)
  */
 static void ExportGeomToOBJFile(const FString& InFileName, const TNavStatArray<FVector::FReal>& GeomCoords, const TNavStatArray<int32>& GeomFaces, const FString& AdditionalData)
@@ -972,56 +971,50 @@ FORCEINLINE_DEBUGGABLE void ExportRigidBodySetup(UBodySetup& BodySetup, TNavStat
 	TemporaryShapeBuffer.Reset();
 }
 
-void ExportObject(INavRelevantInterface& NavRelevantInterface, FRecastGeometryExport& GeomExport)
+void ExportObject(const FNavigationElement& Element, FRecastGeometryExport& GeomExport)
 {
-	if (NavRelevantInterface.IsNavigationRelevant() && (NavRelevantInterface.HasCustomNavigableGeometry() != EHasCustomNavigableGeometry::DontExport))
+	const EHasCustomNavigableGeometry::Type GeometryExportType = Element.GetGeometryExportType();
+	if (GeometryExportType == EHasCustomNavigableGeometry::DontExport)
 	{
-		const bool bHasData = (NavRelevantInterface.HasCustomNavigableGeometry() != EHasCustomNavigableGeometry::Type::No)
-			&& !NavRelevantInterface.DoCustomNavigableGeometryExport(GeomExport);
-
-		UBodySetup* BodySetup = nullptr;
-		{
-			// Might need to create the BodySetup outside of the main thread so garbage collection guard is required.
-			FGCScopeGuard GCGuard;
-			BodySetup = NavRelevantInterface.GetNavigableGeometryBodySetup();
+		return;
+	}
 
-			if (BodySetup)
-			{
-				// Async flag need to be cleared to allow garbage collection.
-				BodySetup->AtomicallyClearInternalFlags(EInternalObjectFlags::Async);
-			}
-		}
+	bool bDefaultGeometryExportRequired = true;
+	if (GeometryExportType != EHasCustomNavigableGeometry::No)
+	{
+		Element.CustomGeometryExportDelegate.ExecuteIfBound(Element, GeomExport, bDefaultGeometryExportRequired);
+	}
 
-		if (BodySetup)
+	if (UBodySetup* BodySetup = Element.GetBodySetup())
+	{
+		if (bDefaultGeometryExportRequired)
 		{
-			if (!bHasData)
-			{
-				ExportRigidBodySetup(*BodySetup, GeomExport.VertexBuffer, GeomExport.IndexBuffer, GeomExport.Data->Bounds, NavRelevantInterface.GetNavigableGeometryTransform());
-			}
-
-			GeomExport.SlopeOverride = BodySetup->WalkableSlopeOverride;
+			ExportRigidBodySetup(*BodySetup, GeomExport.VertexBuffer, GeomExport.IndexBuffer, GeomExport.Data->Bounds, Element.GetTransform());
 		}
+		GeomExport.SlopeOverride = BodySetup->WalkableSlopeOverride;
 	}
 }
 
 #if !UE_BUILD_SHIPPING
 FORCEINLINE_DEBUGGABLE void ValidateGeometryExport(const FRecastGeometryExport& GeomExport)
 {
-	if (const UObject* Owner = GeomExport.Data->GetOwner())
+	if (const UObject* Owner = GeomExport.Data->SourceElement->GetWeakUObject().Get())
 	{
 		if (const UNavigationSystemV1* NavSys = FNavigationSystem::GetCurrent<UNavigationSystemV1>(Owner->GetWorld()))
 		{
 			constexpr int32 CoordinatePerTriangle = 9;
-			if (NavSys->GeometryExportTriangleCountWarningThreshold > 0 && (GeomExport.VertexBuffer.Num() / CoordinatePerTriangle) > NavSys->GeometryExportTriangleCountWarningThreshold)
+			if (NavSys->GeometryExportTriangleCountWarningThreshold > 0
+				&& (GeomExport.VertexBuffer.Num() / CoordinatePerTriangle) > NavSys->GeometryExportTriangleCountWarningThreshold)
 			{
 				static uint32 LastNameHash = 0;
-				const FString FullName = GetFullNameSafe(Owner);
+				const FString FullName = GeomExport.Data->SourceElement.Get().GetFullName();
 				const uint32 NameHash = GetTypeHash(FullName);
 				if (NameHash != LastNameHash)
 				{
 					UE_LOG(LogNavigation,
-						Warning, 
-						TEXT("Exporting collision geometry with too many triangles (%i). This might cause performance and memory issues. Add a simple collision or change GeometryExportVertexCountWarningThreshold. See '%s'."), 
+						Warning,
+						TEXT("Exporting collision geometry with too many triangles (%i). This might cause performance and memory issues."
+						" Add a simple collision or change GeometryExportVertexCountWarningThreshold. See '%s'."),
 						GeomExport.VertexBuffer.Num() / CoordinatePerTriangle, *FullName);
 				}
 				LastNameHash = NameHash;
@@ -1188,7 +1181,7 @@ void FRecastGeometryExport::TransformVertexSoupToRecast(const TArray<FVector>& V
 
 FString FRecastGeometryExport::GetDataOwnerName() const
 {
-	return Data ? GetNameSafe(Data->GetOwner()) : TEXT("No Data");
+	return Data ? Data->SourceElement.Get().GetName() : TEXT("No Data");
 }
 
 FORCEINLINE void GrowConvexHull(const FVector::FReal ExpandBy, const TArray<FVector>& Verts, TArray<FVector>& OutResult)
@@ -1455,27 +1448,27 @@ public:
 protected:
 	/// Logs a message.
 	///  @param[in]		category	The category of the message.
-	///  @param[in]		msg			The formatted message.
+	///  @param[in]		Msg			The formatted message.
 	///  @param[in]		len			The length of the formatted message.
 	virtual void doLog(const rcLogCategory category, const char* Msg, const int32 /*len*/) 
 	{
 		switch (category) 
 		{
 		case RC_LOG_ERROR:
-			UE_LOG(LogNavigation, Error, TEXT("Recast: %s"), ANSI_TO_TCHAR( Msg ) );
+			UE_LOG(LogNavigation, Error, TEXT("Recast: %hs"),  Msg);
 			break;
 		case RC_LOG_WARNING:
-			UE_LOG(LogNavigation, Log, TEXT("Recast: %s"), ANSI_TO_TCHAR( Msg ) );
+			UE_LOG(LogNavigation, Log, TEXT("Recast: %hs"),  Msg);
 			break;
 		default:
-			UE_LOG(LogNavigation, VeryVerbose, TEXT("Recast: %s"), ANSI_TO_TCHAR( Msg ) );
+			UE_LOG(LogNavigation, VeryVerbose, TEXT("Recast: %hs"), Msg);
 			break;
 		}
 	}
 
 	virtual void doDtLog(const char* Msg, const int32 /*len*/)
 	{
-		UE_LOG(LogNavigation, Error, TEXT("Recast: %s"), ANSI_TO_TCHAR(Msg));
+		UE_LOG(LogNavigation, Error, TEXT("Recast: %hs"), Msg);
 	}
 };
 
@@ -1975,12 +1968,6 @@ void FRecastTileGenerator::GatherGeometryFromSources()
 
 	for (TSharedRef<FNavigationRelevantData, ESPMode::ThreadSafe>& ElementData : NavigationRelevantData)
 	{
-		if (ElementData->GetOwner() == nullptr)
-		{
-			UE_LOG(LogNavigation, Warning, TEXT("%s: skipping an element with no longer valid Owner"), ANSI_TO_TCHAR(__FUNCTION__));
-			continue;
-		}
-
 		GatherNavigationDataGeometry(ElementData, *NavSys, NavDataConfig, bUpdateGeometry);
 	}
 }
@@ -1997,15 +1984,8 @@ ETimeSliceWorkResult FRecastTileGenerator::GatherGeometryFromSourcesTimeSliced()
 
 	while(NavigationRelevantData.Num())
 	{
-		TSharedRef<FNavigationRelevantData, ESPMode::ThreadSafe> ElementData = NavigationRelevantData.Pop(EAllowShrinking::No);
-		if (ElementData->GetOwner() == nullptr)
-		{
-			UE_LOG(LogNavigation, Warning, TEXT("%s: skipping an element with no longer valid Owner"), ANSI_TO_TCHAR(__FUNCTION__));
-			continue;
-		}
+		GatherNavigationDataGeometry(NavigationRelevantData.Pop(EAllowShrinking::No), *NavSys, NavDataConfig, bUpdateGeometry);
 
-		GatherNavigationDataGeometry(ElementData, *NavSys, NavDataConfig, bUpdateGeometry);
-		
 		MARK_TIMESLICE_SECTION_DEBUG(TimeSliceManager->GetTimeSlicer(), GatherGeometryFromSources);
 
 		if (TimeSliceManager->GetTimeSlicer().TestTimeSliceFinished())
@@ -2107,59 +2087,57 @@ void FRecastTileGenerator::GatherGeometry(const FRecastNavMeshGenerator& ParentG
 	}
 }
 
-void FRecastTileGenerator::GatherNavigationDataGeometry(const TSharedRef<FNavigationRelevantData, ESPMode::ThreadSafe>& ElementData, UNavigationSystemV1& NavSys, const FNavDataConfig& OwnerNavDataConfig, const bool bGeometryChanged)
+void FRecastTileGenerator::GatherNavigationDataGeometry(const TSharedRef<FNavigationRelevantData, ESPMode::ThreadSafe>& ElementDataRef, UNavigationSystemV1& NavSys, const FNavDataConfig& OwnerNavDataConfig, const bool bGeometryChanged)
 {
 	bool bDumpGeometryData = false;
+	FNavigationRelevantData& ElementData = ElementDataRef.Get();
 
 #if RECAST_INTERNAL_DEBUG_DATA
 	if (!IsTileDebugAllowingGeneration())
 	{
 		return;
 	}
+
 	if (IsTileDebugActive())
 	{
-		UE_LOG(LogNavigation, Log, TEXT("Gathering geometry for tile (%i,%i): %s."), TileX, TileY, *GetFullNameSafe(ElementData->GetOwner()));
-		UE_LOG(LogNavigation, Log, TEXT("                       Bounds: %s"), *ElementData->Bounds.ToString());
-		UE_LOG(LogNavigation, Log, TEXT("                       Geometry: Has=%s Pending=%s Slice=%s"), ElementData->HasGeometry() ? TEXT("true") : TEXT("false"), ElementData->IsPendingLazyGeometryGathering() ? TEXT("true") : TEXT("false"), ElementData->SupportsGatheringGeometrySlices() ? TEXT("true") : TEXT("false"));
-		UE_LOG(LogNavigation, Log, TEXT("                       Modifier: Has=%s Pending=%s"), ElementData->HasModifiers() ? TEXT("true") : TEXT("false"), ElementData->NeedAnyPendingLazyModifiersGathering() ? TEXT("true") : TEXT("false"));
+		UE_LOG(LogNavigation, Log, TEXT("Gathering geometry for tile (%i,%i): %s.\n"
+										"                       Bounds: %s\n"
+										"                       Geometry: Has=%s Pending=%s Slice=%s\n"
+										"                       Modifier: Has=%s Pending=%s"),
+			TileX, TileY, *ElementData.SourceElement.Get().GetFullName(),
+			*ElementData.Bounds.ToString(),
+			*LexToString(ElementData.HasGeometry()), *LexToString(ElementData.IsPendingLazyGeometryGathering()), *LexToString(ElementData.SupportsGatheringGeometrySlices()),
+			*LexToString(ElementData.HasModifiers()), *LexToString(ElementData.NeedAnyPendingLazyModifiersGathering()));
 	}
 #endif // RECAST_INTERNAL_DEBUG_DATA
 
-	if (ElementData->IsPendingLazyGeometryGathering() || ElementData->NeedAnyPendingLazyModifiersGathering())
+	if (ElementData.IsPendingLazyGeometryGathering() || ElementData.NeedAnyPendingLazyModifiersGathering())
 	{
 		QUICK_SCOPE_CYCLE_COUNTER(STAT_RecastNavMeshGenerator_LazyGeometryExport);
-		NavSys.DemandLazyDataGathering(*ElementData);
+		NavSys.DemandLazyDataGathering(ElementData);
 	}
 				
-	if (ElementData->IsPendingLazyGeometryGathering() && ElementData->SupportsGatheringGeometrySlices())
+	if (ElementData.IsPendingLazyGeometryGathering() && ElementData.SupportsGatheringGeometrySlices())
 	{
 		QUICK_SCOPE_CYCLE_COUNTER(STAT_RecastNavMeshGenerator_LandscapeSlicesExporting);
 
-		FRecastGeometryExport GeomExport(const_cast<FNavigationRelevantData&>(*ElementData));
-
-		INavRelevantInterface* NavRelevant = const_cast<INavRelevantInterface*>(Cast<const INavRelevantInterface>(ElementData->GetOwner()));
-		if (NavRelevant)
-		{
-			NavRelevant->PrepareGeometryExportSync();
-			// adding a small bump to avoid special case of zero-expansion when tile bounds
-			// overlap landscape's tile bounds
-			NavRelevant->GatherGeometrySlice(GeomExport, TileBBExpandedForAgent);
+		FRecastGeometryExport GeomExport(ElementData);
 
-			RecastGeometryExport::ConvertCoordDataToRecast(GeomExport.VertexBuffer);
-			RecastGeometryExport::StoreCollisionCache(GeomExport);
-			bDumpGeometryData = true;
-		}
-		else
-		{
-			UE_LOG(LogNavigation, Error, TEXT("GatherGeometry: got an invalid NavRelevant instance!"));
-		}
+		// adding a small bump to avoid special case of zero-expansion when tile bounds
+		// overlap landscape's tile bounds
+		ElementData.SourceElement->GeometrySliceExportDelegate.ExecuteIfBound(ElementData.SourceElement.Get(), GeomExport, TileBBExpandedForAgent);
+		RecastGeometryExport::ConvertCoordDataToRecast(GeomExport.VertexBuffer);
+		RecastGeometryExport::StoreCollisionCache(GeomExport);
+		bDumpGeometryData = true;
 	}
 
 	// Temporary change to help narrow down a rare crash:
 	// Was: const FCompositeNavModifier ModifierInstance = ElementData->GetModifierForAgent(&OwnerNavDataConfig);
-	const FCompositeNavModifier& ModifierInstance = ElementData->Modifiers.HasMetaAreas() ? ElementData->Modifiers.GetInstantiatedMetaModifier(&OwnerNavDataConfig, ElementData->SourceObject) : ElementData->Modifiers;
+	const FCompositeNavModifier& ModifierInstance = ElementData.Modifiers.HasMetaAreas()
+		? ElementData.Modifiers.GetInstantiatedMetaModifier(&OwnerNavDataConfig, ElementData.SourceElement->GetWeakUObject())
+		: ElementData.Modifiers;
 
-	const bool bExportGeometry = bGeometryChanged && ElementData->HasGeometry();
+	const bool bExportGeometry = bGeometryChanged && ElementData.HasGeometry();
 	if (bExportGeometry)
 	{
 		if (ARecastNavMesh::IsVoxelCacheEnabled())
@@ -2170,11 +2148,11 @@ void FRecastTileGenerator::GatherNavigationDataGeometry(const TSharedRef<FNaviga
 
 			DECLARE_SCOPE_CYCLE_COUNTER(TEXT("Rasterization: prepare voxel cache"), Stat_RecastRasterCachePrep, STATGROUP_Navigation);
 
-			if (!HasVoxelCache(ElementData->VoxelData, CachedVoxels, NumCachedVoxels))
+			if (!HasVoxelCache(ElementData.VoxelData, CachedVoxels, NumCachedVoxels))
 			{
 
 				// rasterize
-				PrepareVoxelCache(ElementData->CollisionData, ModifierInstance, SpanData);
+				PrepareVoxelCache(ElementData.CollisionData, ModifierInstance, SpanData);
 				CachedVoxels = SpanData.GetData();
 				NumCachedVoxels = SpanData.Num();
 
@@ -2182,11 +2160,10 @@ void FRecastTileGenerator::GatherNavigationDataGeometry(const TSharedRef<FNaviga
 				{
 					LLM_SCOPE_BYTAG(NavigationOctree);
 
-					const SIZE_T PrevElementMemory = ElementData->GetAllocatedSize();
-					FNavigationRelevantData* ModData = (FNavigationRelevantData*)&ElementData;
-					AddVoxelCache(ModData->VoxelData, CachedVoxels, NumCachedVoxels);
+					const SIZE_T PrevElementMemory = ElementData.GetAllocatedSize();
+					AddVoxelCache(ElementData.VoxelData, CachedVoxels, NumCachedVoxels);
 
-					const SIZE_T NewElementMemory = ElementData->GetAllocatedSize();
+					const SIZE_T NewElementMemory = ElementData.GetAllocatedSize();
 					const SIZE_T ElementMemoryDelta = NewElementMemory - PrevElementMemory;
 					INC_MEMORY_STAT_BY(STAT_Navigation_CollisionTreeMemory, ElementMemoryDelta);
 				}
@@ -2199,13 +2176,13 @@ void FRecastTileGenerator::GatherNavigationDataGeometry(const TSharedRef<FNaviga
 
 		if (bDumpGeometryData)
 		{
-			const_cast<FNavigationRelevantData&>(*ElementData).CollisionData.Empty();
+			ElementData.CollisionData.Empty();
 		}
 	}
 				
 	if (ModifierInstance.IsEmpty() == false)
 	{
-		AppendModifier(ModifierInstance, ElementData->NavDataPerInstanceTransformDelegate);
+		AppendModifier(ModifierInstance, ElementData.NavDataPerInstanceTransformDelegate);
 	}
 }
 
@@ -2488,18 +2465,22 @@ void FRecastTileGenerator::AppendModifier(const FCompositeNavModifier& Modifier,
 	Modifiers.Add(MoveTemp(ModifierElement));
 }
 
-void FRecastTileGenerator::ValidateAndAppendGeometry(const TSharedRef<FNavigationRelevantData, ESPMode::ThreadSafe>& ElementData, const FCompositeNavModifier& InModifier)
+void FRecastTileGenerator::ValidateAndAppendGeometry(const FNavigationRelevantData& ElementData, const FCompositeNavModifier& InModifier)
 {
-	const FNavigationRelevantData& DataRef = ElementData.Get();
-	if (DataRef.IsCollisionDataValid())
+	if (ElementData.IsCollisionDataValid())
 	{
-		AppendGeometry(DataRef, InModifier, DataRef.NavDataPerInstanceTransformDelegate);
+		AppendGeometry(ElementData, InModifier, ElementData.NavDataPerInstanceTransformDelegate);
 	}
 }
 
-void FRecastTileGenerator::AppendGeometry(const FNavigationRelevantData& DataRef, const FCompositeNavModifier& InModifier, const FNavDataPerInstanceTransformDelegate& InTransformsDelegate)
+void FRecastTileGenerator::ValidateAndAppendGeometry(const TSharedRef<FNavigationRelevantData, ESPMode::ThreadSafe>& ElementData, const FCompositeNavModifier& InModifier)
+{
+	ValidateAndAppendGeometry(ElementData.Get(), InModifier);
+}
+
+void FRecastTileGenerator::AppendGeometry(const FNavigationRelevantData& ElementData, const FCompositeNavModifier& InModifier, const FNavDataPerInstanceTransformDelegate& InTransformsDelegate)
 {	
-	const TNavStatArray<uint8>& RawCollisionCache = DataRef.CollisionData;
+	const TNavStatArray<uint8>& RawCollisionCache = ElementData.CollisionData;
 	if (RawCollisionCache.Num() == 0)
 	{
 		return;
@@ -2527,7 +2508,8 @@ void FRecastTileGenerator::AppendGeometry(const FNavigationRelevantData& DataRef
 	const int32 NumIndices = CollisionCache.Header.NumFaces * 3;
 	if (NumIndices > 0)
 	{
-		UE_LOG(LogNavigationDataBuild, VeryVerbose, TEXT("%s adding %i vertices from %s."), ANSI_TO_TCHAR(__FUNCTION__), CollisionCache.Header.NumVerts, *GetFullNameSafe(DataRef.GetOwner()));
+		UE_LOG(LogNavigationDataBuild, VeryVerbose, TEXT("%hs adding %i vertices from %s."),
+			__FUNCTION__, CollisionCache.Header.NumVerts, *ElementData.SourceElement.Get().GetFullName());
 
 		GeometryElement.GeomCoords.SetNumUninitialized(NumCoords);
 		GeometryElement.GeomIndices.SetNumUninitialized(NumIndices);
@@ -4788,13 +4770,11 @@ uint32 FRecastTileGenerator::GetUsedMemCount() const
 
 void FRecastTileGenerator::AddReferencedObjects(FReferenceCollector& Collector)
 {
-	for (auto& RelevantData : NavigationRelevantData)
+	for (const TSharedRef<FNavigationRelevantData>& RelevantData : NavigationRelevantData)
 	{
-		auto& Owner = RelevantData->GetOwnerPtr();
-		if (Owner.Get())
-		{
-			Collector.AddReferencedObject(Owner);
-		}
+		// Local variable since 'AddReferencedObject' requires a non-const& parameter
+		TWeakObjectPtr<const UObject> WeakObject = RelevantData->SourceElement->GetWeakUObject();
+		Collector.AddReferencedObject(WeakObject);
 	}
 }
 
@@ -5995,21 +5975,17 @@ void FRecastNavMeshGenerator::LogDirtyAreas(const UObject& OwnerNav,
 			DirtyResultsTuple->TotalDirtyTiles++;
 		}
 	}
-		
-	for (const FNavigationDirtyAreaDebugInformation& DirtyResultsTuple : DirtyAreaToDirtyTilesCount)
+
+	for (const auto& [DirtyArea, NewlyAddedDirtyTiles, TotalDirtyTiles] : DirtyAreaToDirtyTilesCount)
 	{
-		const UObject* const SourceObject = DirtyResultsTuple.DirtyArea.OptionalSourceObject.Get();
-		const UActorComponent* const ObjectAsComponent = Cast<UActorComponent>(SourceObject);
-		const AActor* const ComponentOwner = ObjectAsComponent ? ObjectAsComponent->GetOwner() : nullptr;
-		const FVector2D BoundsSize(DirtyResultsTuple.DirtyArea.Bounds.GetSize());
+		const FVector2D BoundsSize(DirtyArea.Bounds.GetSize());
 		
 		UE_LOG(LogNavigationDirtyArea, VeryVerbose,
-			TEXT("(navmesh: %-30s) Dirty area trying to dirt %2d tiles (out of which %2d are newly added/not pending) | Source Object = %s | Potential component's owner = %s | Bounds size = %s)"),
-			*GetNameSafe(GetOwner()), DirtyResultsTuple.TotalDirtyTiles, DirtyResultsTuple.NewlyAddedDirtyTiles, *GetFullNameSafe(SourceObject),
-			*GetFullNameSafe(ComponentOwner), *BoundsSize.ToString());
+			TEXT("(navmesh: %-30s) Dirty area trying to dirty %2d tiles (out of which %2d are newly added/not pending) | Source = %s | Bounds size = %s)"),
+			*GetNameSafe(GetOwner()), TotalDirtyTiles, NewlyAddedDirtyTiles, *DirtyArea.GetSourceDescription(), *BoundsSize.ToString());
 
-		UE_VLOG_BOX(&OwnerNav, LogNavigationDirtyArea, VeryVerbose, DirtyResultsTuple.DirtyArea.Bounds, FColor::Purple,
-			TEXT("Tiles %d (new: %d), Source: %s"), DirtyResultsTuple.TotalDirtyTiles, DirtyResultsTuple.NewlyAddedDirtyTiles,*GetFullNameSafe(SourceObject));
+		UE_VLOG_BOX(&OwnerNav, LogNavigationDirtyArea, VeryVerbose, DirtyArea.Bounds, FColor::Purple,
+			TEXT("Tiles %d (new: %d), Source: %s"), TotalDirtyTiles, NewlyAddedDirtyTiles, *DirtyArea.GetSourceDescription());
 	}
 }
 #endif
@@ -6309,8 +6285,9 @@ void FRecastNavMeshGenerator::MarkDirtyTiles(const TArray<FNavigationDirtyArea>&
 		for (const FNavigationDirtyArea& DirtyArea : DirtyAreas)
 		{
 			QUICK_SCOPE_CYCLE_COUNTER(STAT_RecastNavMeshGenerator_DirtyArea);
-		
-			if (!ensureMsgf(DirtyArea.Bounds.IsValid, TEXT("%hs Attempting to use DirtyArea.Bounds which are not valid. SourceObject: %s"), __FUNCTION__, *GetFullNameSafe(DirtyArea.OptionalSourceObject.Get())))
+
+			if (!ensureMsgf(DirtyArea.Bounds.IsValid, TEXT("%hs Attempting to use DirtyArea.Bounds which are not valid. Source: %s"),
+					__FUNCTION__, *DirtyArea.GetSourceDescription()))
 			{
 				continue;
 			}
@@ -6321,17 +6298,16 @@ void FRecastNavMeshGenerator::MarkDirtyTiles(const TArray<FNavigationDirtyArea>&
 				continue;
 			}
 
-			UE_VLOG_BOX(OwnerNav, LogNavigation, VeryVerbose, DirtyArea.Bounds, FColor::Blue, TEXT("DirtyArea %s"), *GetNameSafe(DirtyArea.OptionalSourceObject.Get()));
-		
+			UE_VLOG_BOX(OwnerNav, LogNavigation, VeryVerbose, DirtyArea.Bounds, FColor::Blue, TEXT("DirtyArea %s"), *DirtyArea.GetSourceDescription());
+
 			// (if bUseVirtualGeometryFilteringAndDirtying is true) Ignore dirty areas flagged by a source object that is not supposed to apply to this navmesh
 			if (bUseVirtualGeometryFilteringAndDirtying && NavSys && NavOctreeInstance && NavDataConfig)
 			{
-				if (const UObject* const SourceObject = DirtyArea.OptionalSourceObject.Get())
+				const FNavigationElement* SourceElement = DirtyArea.OptionalSourceElement.Get();
+				if (SourceElement
+					&& !ShouldDirtyTilesRequestedByElement(*NavSys, *NavOctreeInstance, SourceElement->GetHandle(), *NavDataConfig))
 				{
-					if (!ShouldDirtyTilesRequestedByObject(*NavSys, *NavOctreeInstance, *SourceObject, *NavDataConfig))
-					{
-						continue;
-					}
+					continue;
 				}
 			}
 		
@@ -6367,7 +6343,7 @@ void FRecastNavMeshGenerator::MarkDirtyTiles(const TArray<FNavigationDirtyArea>&
 				QUICK_SCOPE_CYCLE_COUNTER(STAT_RecastNavMeshGenerator_CheckTilesInBounds);
 
 				uint32 PendingTilesMarked = 0;
-			
+
 				const FRcTileBox TileBox(AdjustedAreaBounds, RcNavMeshOrigin, TileSizeInWorldUnits);
 
 				for (int32 TileY = TileBox.YMin; TileY <= TileBox.YMax; ++TileY)
@@ -6439,7 +6415,7 @@ void FRecastNavMeshGenerator::MarkDirtyTiles(const TArray<FNavigationDirtyArea>&
 
 #if !UE_BUILD_SHIPPING
 				// Warn if this is from a big dirty area
-				UE_SUPPRESS(LogNavigationDirtyArea, Warning, 
+				UE_SUPPRESS(LogNavigationDirtyArea, Warning,
 				{
 					if (PendingTilesMarked > 0)
 					{
@@ -6453,14 +6429,11 @@ void FRecastNavMeshGenerator::MarkDirtyTiles(const TArray<FNavigationDirtyArea>&
 						if (NavSys && NavSys->GetOperationMode() == FNavigationSystemRunMode::GameMode && NavSys->IsActiveTilesGenerationEnabled() &&
 							AdjustedAreaBounds.GetSize().GetMax() > NavSys->GetDirtyAreaWarningSizeThreshold())
 						{
-							const UObject* const SourceObject = DirtyArea.OptionalSourceObject.Get();
-							const UActorComponent* const ObjectAsComponent = Cast<UActorComponent>(SourceObject);
-							const AActor* const ComponentOwner = ObjectAsComponent ? ObjectAsComponent->GetOwner() : nullptr;
 							const FVector2D AdjustedAreaBoundsSize(AdjustedAreaBounds.GetSize());
-		
+
 							UE_LOG(LogNavigationDirtyArea, Warning,
-								TEXT("(navmesh: %-30s) Added an oversized dirty area | Tiles marked: %2u | Source object = %s | Potential comp owner = %s | Bounds size = %s | Threshold: %.0f"),
-								*GetNameSafe(GetOwner()), PendingTilesMarked, *GetFullNameSafe(SourceObject), *GetFullNameSafe(ComponentOwner),
+								TEXT("(navmesh: %-30s) Added an oversized dirty area | Tiles marked: %2u | Source Element = %s | Bounds size = %s | Threshold: %.0f"),
+								*GetNameSafe(GetOwner()), PendingTilesMarked, *DirtyArea.GetSourceDescription(),
 								*AdjustedAreaBoundsSize.ToString(), NavSys->GetDirtyAreaWarningSizeThreshold());
 						}
 					}
@@ -6535,10 +6508,20 @@ void FRecastNavMeshGenerator::MarkDirtyTiles(const TArray<FNavigationDirtyArea>&
 	}
 }
 
+// Deprecated
 bool FRecastNavMeshGenerator::ShouldDirtyTilesRequestedByObject(const UNavigationSystemV1& NavSys,
 	const FNavigationOctree& NavOctreeInstance, const UObject& SourceObject, const FNavDataConfig& NavDataConfig) const
 {
-	const FOctreeElementId2* const OctreeElementId = NavSys.GetObjectsNavOctreeId(SourceObject);
+	return ShouldDirtyTilesRequestedByElement(NavSys, NavOctreeInstance, FNavigationElementHandle(&SourceObject), NavDataConfig); 
+}
+
+bool FRecastNavMeshGenerator::ShouldDirtyTilesRequestedByElement(
+	const UNavigationSystemV1& NavSys,
+	const FNavigationOctree& NavOctreeInstance,
+	const FNavigationElementHandle SourceElement,
+	const FNavDataConfig& NavDataConfig) const
+{
+	const FOctreeElementId2* const OctreeElementId = NavSys.GetNavOctreeIdForElement(SourceElement);
 
 	return (OctreeElementId == nullptr) || ShouldGenerateGeometryForOctreeElement(NavOctreeInstance.GetElementById(*OctreeElementId), NavDataConfig);
 }
@@ -7144,7 +7127,7 @@ TArray<FNavTileRef> FRecastNavMeshGenerator::ProcessTileTasksAndGetUpdatedTiles(
 
 		if (RebuildAllStartTime != 0)
 		{
-			UE_LOG(LogNavigationDataBuild, Display, TEXT("   %s build time: %.2fs"), ANSI_TO_TCHAR(__FUNCTION__), (FPlatformTime::Seconds() - RebuildAllStartTime));
+			UE_LOG(LogNavigationDataBuild, Display, TEXT("   %hs build time: %.2fs"), __FUNCTION__, (FPlatformTime::Seconds() - RebuildAllStartTime));
 			RebuildAllStartTime = 0;
 		}
 		
@@ -7203,20 +7186,22 @@ void FRecastNavMeshGenerator::ExportComponentGeometry(UActorComponent* InOutComp
 {
 	if (INavRelevantInterface* NavRelevantInterface = Cast<INavRelevantInterface>(InOutComponent))
 	{
-		FRecastGeometryExport::ExportNavRelevantObjectGeometry(*NavRelevantInterface, OutData);
+		const TSharedRef<const FNavigationElement> TmpElement = FNavigationElement::CreateFromNavRelevantInterface(*NavRelevantInterface);
+		FRecastGeometryExport::ExportElementGeometry(TmpElement.Get(), OutData);
 	}
 }
 
 // Deprecated
 void FRecastNavMeshGenerator::ExportNavRelevantObjectGeometry(INavRelevantInterface& InOutNavRelevantInterface, FNavigationRelevantData& OutData)
 {
-	FRecastGeometryExport::ExportNavRelevantObjectGeometry(InOutNavRelevantInterface, OutData);
+	const TSharedRef<const FNavigationElement> TmpElement = FNavigationElement::CreateFromNavRelevantInterface(InOutNavRelevantInterface);
+	FRecastGeometryExport::ExportElementGeometry(TmpElement.Get(), OutData);
 }
 
-void FRecastGeometryExport::ExportNavRelevantObjectGeometry(INavRelevantInterface& InOutNavRelevantInterface, FNavigationRelevantData& OutData)
+void FRecastGeometryExport::ExportElementGeometry(const FNavigationElement& InElement, FNavigationRelevantData& OutData)
 {
 	FRecastGeometryExport GeomExport(OutData);
-	RecastGeometryExport::ExportObject(InOutNavRelevantInterface, GeomExport);
+	RecastGeometryExport::ExportObject(InElement, GeomExport);
 
 #if !UE_BUILD_SHIPPING
 	RecastGeometryExport::ValidateGeometryExport(GeomExport);
@@ -7450,7 +7435,7 @@ bool FRecastNavMeshGenerator::IsTileChanged(const FNavTileRef InTileRef) const
 
 uint32 FRecastNavMeshGenerator::LogMemUsed() const 
 {
-	UE_LOG(LogNavigation, Display, TEXT("    FRecastNavMeshGenerator: self %d"), sizeof(FRecastNavMeshGenerator));
+	UE_LOG(LogNavigation, Display, TEXT("    FRecastNavMeshGenerator: self %llu"), sizeof(FRecastNavMeshGenerator));
 	
 	uint32 GeneratorsMem = 0;
 	for (const TRunningTileElement<FRecastTileGeneratorWrapper>& Element : RunningDirtyTiles)
@@ -7550,9 +7535,9 @@ void FRecastNavMeshGenerator::GrabDebugSnapshot(struct FVisualLogEntry* Snapshot
 							continue;
 						}
 
-						const uint8 AreaId = IntCastChecked<uint8>(NavData->GetAreaID(AreaMod.GetAreaClass()));
+						const int32 AreaId = NavData->GetAreaID(AreaMod.GetAreaClass());
 						const UClass* AreaClass = NavData->GetAreaClass(AreaId);
-						const UNavArea* DefArea = AreaClass ? ((UClass*)AreaClass)->GetDefaultObject<UNavArea>() : NULL;
+						const UNavArea* DefArea = AreaClass ? ((UClass*)AreaClass)->GetDefaultObject<UNavArea>() : nullptr;
 						const FColor PolygonColor = AreaClass != FNavigationSystem::GetDefaultWalkableArea() ? (DefArea ? DefArea->DrawColor : NavData->GetConfig().Color) : FColorList::Cyan;
 
 						if (ShapeType == ENavigationShapeType::Box)
@@ -7704,7 +7689,8 @@ void FRecastNavMeshGenerator::ExportNavigationData(const FString& FileName) cons
 						if (ShapeType == ENavigationShapeType::Convex || ShapeType == ENavigationShapeType::InstancedConvex)
 						{
 							FAreaExportData ExportInfo;
-							ExportInfo.AreaId = IntCastChecked<uint8>(NavData->GetAreaID(AreaMod.GetAreaClass()));
+							const int32 AreaId = NavData->GetAreaID(AreaMod.GetAreaClass());
+							ExportInfo.AreaId = AreaId != INDEX_NONE ? IntCastChecked<uint8>(AreaId) : INDEX_NONE;
 
 							auto AddAreaExportDataFunc = [&](const FConvexNavAreaData& InConvexNavAreaData)
 							{
diff --git a/Engine/Source/Runtime/NavigationSystem/Private/NavigationDataHandler.cpp b/Engine/Source/Runtime/NavigationSystem/Private/NavigationDataHandler.cpp
--- a/Engine/Source/Runtime/NavigationSystem/Private/NavigationDataHandler.cpp
+++ b/Engine/Source/Runtime/NavigationSystem/Private/NavigationDataHandler.cpp
@@ -1,10 +1,14 @@
 // Copyright Epic Games, Inc. All Rights Reserved.
 
 #include "NavigationDataHandler.h"
+#include "AI/Navigation/NavigationDirtyElement.h"
+#include "AI/Navigation/NavRelevantInterface.h"
+#include "DetourCrowd/DetourCrowd.h"
 #include "Engine/Level.h"
 #include "GameFramework/Actor.h"
 #include "NavAreas/NavArea.h"
 #include "NavMesh/RecastGeometryExport.h"
+#include "VisualLogger/VisualLogger.h"
 
 DEFINE_LOG_CATEGORY_STATIC(LogNavOctree, Warning, All);
 
@@ -16,15 +20,6 @@ namespace UE::NavigationHelper::Private
 			((UpdateFlags & FNavigationOctreeController::OctreeUpdate_Modifiers) != 0) ? ENavigationDirtyFlag::DynamicModifier :
 			DefaultValue;
 	}
-
-	FString GetElementName(const UObject& ElementOwner)
-	{
-		if (const UActorComponent* ActorComp = Cast<UActorComponent>(&ElementOwner))
-		{
-			return FString::Format(TEXT("Comp {0} of Actor {1}"), {*GetNameSafe(&ElementOwner), *GetNameSafe(ActorComp->GetOwner())});
-		}
-		return *GetNameSafe(&ElementOwner);
-	}
 }
 	
 FNavigationDataHandler::FNavigationDataHandler(FNavigationOctreeController& InOctreeController, FNavigationDirtyAreasController& InDirtyAreasController)
@@ -43,7 +38,13 @@ void FNavigationDataHandler::ConstructNavOctree(const FVector& Origin, const dou
 #endif // !UE_BUILD_SHIPPING
 }
 
-void FNavigationDataHandler::RemoveNavOctreeElementId(const FOctreeElementId2& ElementId, int32 UpdateFlags)
+// Deprecated
+void FNavigationDataHandler::RemoveNavOctreeElementId(const FOctreeElementId2& ElementId, const int32 UpdateFlags)
+{
+	RemoveFromNavOctree(ElementId, UpdateFlags);
+}
+
+void FNavigationDataHandler::RemoveFromNavOctree(const FOctreeElementId2& ElementId, const int32 UpdateFlags)
 {
 	if (ensure(OctreeController.IsValidElement(ElementId)))
 	{
@@ -52,63 +53,76 @@ void FNavigationDataHandler::RemoveNavOctreeElementId(const FOctreeElementId2& E
 		if (!ElementData.Data->bShouldSkipDirtyAreaOnAddOrRemove)
 		{
 			const ENavigationDirtyFlag DirtyFlag = UE::NavigationHelper::Private::GetDirtyFlag(UpdateFlags, ElementData.Data->GetDirtyFlag());
-			DirtyAreasController.AddArea(ElementData.Bounds.GetBox(), DirtyFlag, [&ElementData] { return ElementData.Data->SourceObject.Get(); }, nullptr, "Remove from navoctree");
+			DirtyAreasController.AddArea(
+				ElementData.Bounds.GetBox(),
+				DirtyFlag,
+				[&ElementData]
+				{
+					return ElementData.Data->SourceElement;
+				},
+				/*DirtyElement*/nullptr,
+				"Remove from navoctree");
 		}
 
-		OctreeController.NavOctree->RemoveNode(ElementId);
+		OctreeController.RemoveNode(ElementId, ElementData.Data.Get().SourceElement.Get().GetHandle());
 	}
 }
 
+// Deprecated
 FSetElementId FNavigationDataHandler::RegisterNavOctreeElement(UObject& ElementOwner, INavRelevantInterface& ElementInterface, int32 UpdateFlags)
+{
+	check(false);
+	return RegisterElementWithNavOctree(FNavigationElement::CreateFromNavRelevantInterface(ElementInterface), UpdateFlags);
+}
+
+FSetElementId FNavigationDataHandler::RegisterElementWithNavOctree(const TSharedRef<const FNavigationElement>& ElementRef, const int32 UpdateFlags)
 {
 	SCOPE_CYCLE_COUNTER(STAT_Navigation_RegisterNavOctreeElement);
 
 	FSetElementId SetId;
+	const FNavigationElement& NavigationElement = ElementRef.Get();
 
-	if (OctreeController.NavOctree.IsValid() == false)
+	if (OctreeController.IsValid() == false)
 	{
+		UE_LOG(LogNavOctree, VeryVerbose, TEXT("IGNORE(%hs) %s: octree not created yet"),
+			__FUNCTION__, *NavigationElement.GetPathName());
 		return SetId;
 	}
 
 	if (OctreeController.IsNavigationOctreeLocked())
 	{
-		UE_LOG(LogNavOctree, Log, TEXT("IGNORE(RegisterNavOctreeElement) %s"), *GetPathNameSafe(&ElementOwner));
+		UE_LOG(LogNavOctree, Log, TEXT("IGNORE(%hs) %s: navigation octree locked"),
+			__FUNCTION__, *NavigationElement.GetPathName());
 		return SetId;
 	}
 
-	const bool bIsRelevant = ElementInterface.IsNavigationRelevant();
-	UE_LOG(LogNavOctree, Log, TEXT("REG %s %s"), *UE::NavigationHelper::Private::GetElementName(ElementOwner), bIsRelevant ? TEXT("[relevant]") : TEXT("[skip]"));
+	UE_LOG(LogNavOctree, Log, TEXT("REG %s"), *NavigationElement.GetPathName());
+
+	bool bCanAdd = false;
+	if (const TWeakObjectPtr<const UObject>& NavigationParent = NavigationElement.GetNavigationParent(); !NavigationParent.IsExplicitlyNull())
+	{
+		OctreeController.AddChild(FNavigationElementHandle(NavigationParent), ElementRef);
+		bCanAdd = true;
+	}
+	else
+	{
+		bCanAdd = (OctreeController.HasElementNavOctreeId(NavigationElement.GetHandle()) == false);
+	}
 
-	if (bIsRelevant)
+	if (bCanAdd)
 	{
-		bool bCanAdd = false;
+		FNavigationDirtyElement UpdateInfo(ElementRef, UE::NavigationHelper::Private::GetDirtyFlag(UpdateFlags, ENavigationDirtyFlag::None), DirtyAreasController.bUseWorldPartitionedDynamicMode);
 
-		UObject* ParentNode = ElementInterface.GetNavigationParent();
-		if (ParentNode)
+		SetId = OctreeController.PendingUpdates.FindId(NavigationElement.GetHandle());
+		if (SetId.IsValidId())
 		{
-			OctreeController.OctreeChildNodesMap.AddUnique(ParentNode, FWeakObjectPtr(&ElementOwner));
-			bCanAdd = true;
+			// make sure this request stays, in case it has been invalidated already and keep any dirty areas
+			UpdateInfo.ExplicitAreasToDirty = OctreeController.PendingUpdates[SetId].ExplicitAreasToDirty;
+			OctreeController.PendingUpdates[SetId] = UpdateInfo;
 		}
 		else
 		{
-			bCanAdd = (OctreeController.HasObjectsNavOctreeId(ElementOwner) == false);
-		}
-
-		if (bCanAdd)
-		{
-			FNavigationDirtyElement UpdateInfo(&ElementOwner, &ElementInterface, UE::NavigationHelper::Private::GetDirtyFlag(UpdateFlags, ENavigationDirtyFlag::None), DirtyAreasController.bUseWorldPartitionedDynamicMode);
-
-			SetId = OctreeController.PendingOctreeUpdates.FindId(UpdateInfo);
-			if (SetId.IsValidId())
-			{
-				// make sure this request stays, in case it has been invalidated already and keep any dirty areas
-				UpdateInfo.ExplicitAreasToDirty = OctreeController.PendingOctreeUpdates[SetId].ExplicitAreasToDirty;
-				OctreeController.PendingOctreeUpdates[SetId] = UpdateInfo;
-			}
-			else
-			{
-				SetId = OctreeController.PendingOctreeUpdates.Add(UpdateInfo);
-			}
+			SetId = OctreeController.PendingUpdates.Add(UpdateInfo);
 		}
 	}
 
@@ -117,7 +131,7 @@ FSetElementId FNavigationDataHandler::RegisterNavOctreeElement(UObject& ElementO
 
 void FNavigationDataHandler::AddElementToNavOctree(const FNavigationDirtyElement& DirtyElement)
 {
-	check(OctreeController.NavOctree.IsValid());
+	check(OctreeController.IsValid());
 	LLM_SCOPE_BYTAG(NavigationOctree);
 
 	// handle invalidated requests first
@@ -125,125 +139,172 @@ void FNavigationDataHandler::AddElementToNavOctree(const FNavigationDirtyElement
 	{
 		if (DirtyElement.bHasPrevData)
 		{
-			DirtyAreasController.AddArea(DirtyElement.PrevBounds, DirtyElement.PrevFlags, [&DirtyElement] { return DirtyElement.Owner.Get(); }, &DirtyElement, "Addition to navoctree (invalid request)");
+			DirtyAreasController.AddArea(DirtyElement.PrevBounds,
+				DirtyElement.PrevFlags,
+				[&DirtyElement]
+				{
+					return DirtyElement.NavigationElement;
+				},
+				&DirtyElement,
+				"Addition to navoctree (invalid request)");
 		}
 
 		return;
 	}
 
-	UObject* ElementOwner = DirtyElement.Owner.Get();
-	if (!IsValid(ElementOwner) || DirtyElement.NavInterface == nullptr)
+	FNavigationOctreeElement OctreeElement(DirtyElement.NavigationElement);
+	const FNavigationElement& NavigationElement = DirtyElement.NavigationElement.Get();
+
+	const TWeakObjectPtr<const UObject> ElementWeakUObject = NavigationElement.GetWeakUObject();
+	if (!ElementWeakUObject.IsExplicitlyNull())
 	{
-		return;
+		UE_VLOG_ALWAYS_UELOG(ElementWeakUObject.Get(), LogNavOctree, Verbose, TEXT("Create FNavigationOctreeElement for %s"),
+			*NavigationElement.GetPathName());
 	}
 
-	FNavigationOctreeElement GeneratedData(*ElementOwner);
-
 	// In WP dynamic mode, store if this is loaded data.
 	if (DirtyAreasController.bUseWorldPartitionedDynamicMode)
 	{
-		GeneratedData.Data->bLoadedData = DirtyElement.bIsFromVisibilityChange || FNavigationSystem::IsLevelVisibilityChanging(ElementOwner);
+		OctreeElement.Data->bLoadedData = DirtyElement.bIsFromVisibilityChange || NavigationElement.IsFromLevelVisibilityChange();
 	}
-	
-	const FBox ElementBounds = DirtyElement.NavInterface->GetNavigationBounds();
 
-	UObject* NavigationParent = DirtyElement.NavInterface->GetNavigationParent();
-	if (NavigationParent)
+	const FBox ElementBounds = NavigationElement.GetBounds();
+
+	if (const TWeakObjectPtr<const UObject>& NavigationParent = NavigationElement.GetNavigationParent(); !NavigationParent.IsExplicitlyNull())
 	{
+		const FNavigationElementHandle ParentKey(NavigationParent);
+
 		// check if parent node is waiting in queue
-		const FSetElementId ParentRequestId = OctreeController.PendingOctreeUpdates.FindId(FNavigationDirtyElement(NavigationParent));
-		const FOctreeElementId2* ParentId = OctreeController.GetObjectsNavOctreeId(*NavigationParent);
+		const FSetElementId ParentRequestId = OctreeController.PendingUpdates.FindId(ParentKey);
+		const FOctreeElementId2* ParentId = OctreeController.GetNavOctreeIdForElement(ParentKey);
 		if (ParentRequestId.IsValidId() && ParentId == nullptr)
 		{
-			FNavigationDirtyElement& ParentNode = OctreeController.PendingOctreeUpdates[ParentRequestId];
-			AddElementToNavOctree(ParentNode);
+			FNavigationDirtyElement& ParentDirtyElement = OctreeController.PendingUpdates[ParentRequestId];
+			AddElementToNavOctree(ParentDirtyElement);
 
 			// mark as invalid so it won't be processed twice
-			ParentNode.bInvalidRequest = true;
+			ParentDirtyElement.bInvalidRequest = true;
 		}
 
-		const FOctreeElementId2* ElementId = ParentId ? ParentId : OctreeController.GetObjectsNavOctreeId(*NavigationParent);
+		const FOctreeElementId2* ElementId = ParentId ? ParentId : OctreeController.GetNavOctreeIdForElement(ParentKey);
 		if (ElementId && ensure(OctreeController.IsValidElement(*ElementId)))
 		{
-			UE_LOG(LogNavOctree, Log, TEXT("ADD %s to %s"), *UE::NavigationHelper::Private::GetElementName(*ElementOwner), *GetNameSafe(NavigationParent));
-			OctreeController.NavOctree->AppendToNode(*ElementId, DirtyElement.NavInterface, ElementBounds, GeneratedData);
+			UE_LOG(LogNavOctree, Log, TEXT("ADD %s to %s"), *NavigationElement.GetPathName(), *GetNameSafe(NavigationParent.Get()));
+			OctreeController.NavOctree->AppendToNode(*ElementId, DirtyElement.NavigationElement, ElementBounds, OctreeElement);
 		}
 		else
 		{
-			UE_LOG(LogNavOctree, Warning, TEXT("Can't add node [%s] - parent [%s] not found in octree!"), *UE::NavigationHelper::Private::GetElementName(*ElementOwner), *UE::NavigationHelper::Private::GetElementName(*NavigationParent));
+			UE_LOG(LogNavOctree, Warning, TEXT("Can't add node [%s] - parent [%s] not found in octree!"),
+				*NavigationElement.GetPathName(),
+				*GetNameSafe(NavigationParent.Get()));
 		}
 	}
 	else
 	{
-		UE_LOG(LogNavOctree, Log, TEXT("ADD %s"), *UE::NavigationHelper::Private::GetElementName(*ElementOwner));
-		OctreeController.NavOctree->AddNode(ElementOwner, DirtyElement.NavInterface, ElementBounds, GeneratedData);
+		OctreeController.NavOctree->AddNode(ElementBounds, OctreeElement);
+		UE_SUPPRESS(LogNavOctree, Verbose,
+			{
+				const FOctreeElementId2* ElementId = OctreeController.GetNavOctreeIdForElement(DirtyElement.NavigationElement->GetHandle());
+				UE_VLOG_ALWAYS_UELOG(NavigationElement.GetWeakUObject().Get(), LogNavOctree, Log, TEXT("ADD %s - %s"),
+					*NavigationElement.GetPathName(),
+					ElementId ? *LexToString(*ElementId) : TEXT("No element"));
+			});
 	}
 
 	// mark area occupied by given element as dirty except if explicitly set to skip this default behavior
-	const ENavigationDirtyFlag DirtyFlag = DirtyElement.FlagsOverride != ENavigationDirtyFlag::None ? DirtyElement.FlagsOverride : GeneratedData.Data->GetDirtyFlag();
-	if (GeneratedData.Data->bShouldSkipDirtyAreaOnAddOrRemove)
+	const ENavigationDirtyFlag DirtyFlag = DirtyElement.FlagsOverride != ENavigationDirtyFlag::None ? DirtyElement.FlagsOverride : OctreeElement.Data->GetDirtyFlag();
+	if (OctreeElement.Data->bShouldSkipDirtyAreaOnAddOrRemove)
 	{
 		if (DirtyElement.ExplicitAreasToDirty.Num() > 0)
 		{
-			DirtyAreasController.AddAreas(DirtyElement.ExplicitAreasToDirty, DirtyFlag, [&ElementOwner] { return ElementOwner; }, &DirtyElement, "Addition to navoctree");
+			DirtyAreasController.AddAreas(
+				DirtyElement.ExplicitAreasToDirty,
+				DirtyFlag, 
+				[ElementOwner = DirtyElement.NavigationElement]
+				{
+					return ElementOwner;
+				},
+				&DirtyElement,
+				"Addition to navoctree");
 		}
 	}
-	else if (!GeneratedData.IsEmpty())
+	else if (!OctreeElement.IsEmpty())
 	{
-		DirtyAreasController.AddArea(GeneratedData.Bounds.GetBox(), DirtyFlag, [&ElementOwner] { return ElementOwner; }, &DirtyElement, "Addition to navoctree");
+		DirtyAreasController.AddArea(
+			OctreeElement.Bounds.GetBox(),
+			DirtyFlag,
+			[ElementOwner = DirtyElement.NavigationElement]
+			{
+				return ElementOwner;
+			},
+			&DirtyElement,
+			"Addition to navoctree");
 	}
 }
 
+// Deprecated
 bool FNavigationDataHandler::UnregisterNavOctreeElement(UObject& ElementOwner, INavRelevantInterface& ElementInterface, int32 UpdateFlags)
+{
+	return UnregisterElementWithNavOctree(FNavigationElement::CreateFromNavRelevantInterface(ElementInterface), UpdateFlags);
+}
+
+bool FNavigationDataHandler::UnregisterElementWithNavOctree(const TSharedRef<const FNavigationElement>& ElementRef, const int32 UpdateFlags)
 {
 	SCOPE_CYCLE_COUNTER(STAT_Navigation_UnregisterNavOctreeElement);
 
-	if (OctreeController.NavOctree.IsValid() == false)
+	const FNavigationElement& NavRelevantElement = ElementRef.Get();
+	if (OctreeController.IsValid() == false)
 	{
+		UE_LOG(LogNavOctree, VeryVerbose, TEXT("IGNORE(%hs) %s: octree not created yet"),
+			__FUNCTION__, *NavRelevantElement.GetPathName());
 		return false;
 	}
 
 	if (OctreeController.IsNavigationOctreeLocked())
 	{
-#if !WITH_EDITOR
-		UE_LOG(LogNavOctree, Log, TEXT("IGNORE(UnregisterNavOctreeElement) %s"), *GetPathNameSafe(&ElementOwner));
-#endif // WITH_EDITOR
+		UE_LOG(LogNavOctree, Log, TEXT("IGNORE(%hs) %s: octree locked"),
+			__FUNCTION__, *ElementRef.Get().GetPathName());
 		return false;
 	}
 
+	const FNavigationElementHandle NavRelevantElementHandle = NavRelevantElement.GetHandle();
+
+	const FOctreeElementId2* OctreeElementId = OctreeController.GetNavOctreeIdForElement(NavRelevantElementHandle);
+	UE_VLOG_ALWAYS_UELOG(NavRelevantElement.GetWeakUObject().Get(), LogNavOctree, Log, TEXT("UNREG %s %s"),
+		*NavRelevantElement.GetPathName(),
+		OctreeElementId ? *FString::Printf(TEXT("[exists %s]"), *LexToString(*OctreeElementId)) : TEXT("[doesn\'t exist]"));
+
 	bool bUnregistered = false;
-	const FOctreeElementId2* ElementId = OctreeController.GetObjectsNavOctreeId(ElementOwner);
-	UE_LOG(LogNavOctree, Log, TEXT("UNREG %s %s"), *UE::NavigationHelper::Private::GetElementName(ElementOwner), ElementId ? TEXT("[exists]") : TEXT("[doesn\'t exist]"));
 
-	if (ElementId != nullptr)
+	if (OctreeElementId != nullptr)
 	{
-		RemoveNavOctreeElementId(*ElementId, UpdateFlags);
-		OctreeController.RemoveObjectsNavOctreeId(ElementOwner);
+		RemoveFromNavOctree(*OctreeElementId, UpdateFlags);
 		bUnregistered = true;
 	}
-	else
+	else if (const bool bCanRemoveChildNode = (UpdateFlags & FNavigationOctreeController::OctreeUpdate_ParentChain) == 0)
 	{
-		const bool bCanRemoveChildNode = (UpdateFlags & FNavigationOctreeController::OctreeUpdate_ParentChain) == 0;
-		UObject* ParentNode = ElementInterface.GetNavigationParent();
-		if (ParentNode && bCanRemoveChildNode)
+		if (const TWeakObjectPtr<const UObject>& NavigationParent = NavRelevantElement.GetNavigationParent(); !NavigationParent.IsExplicitlyNull())
 		{
-			// if node has navigation parent (= doesn't exists in octree on its own)
+			// if node has navigation parent (= doesn't exist in octree on its own)
 			// and it's not part of parent chain update
 			// remove it from map and force update on parent to rebuild octree element
+			const FNavigationElementHandle ParentKey(NavigationParent);
+			OctreeController.RemoveChild(ParentKey, ElementRef);
 
-			OctreeController.OctreeChildNodesMap.RemoveSingle(ParentNode, FWeakObjectPtr(&ElementOwner));
-			UpdateNavOctreeParentChain(*ParentNode);
+			if (const FNavigationRelevantData* NavigationData = OctreeController.GetDataForElement(ParentKey))
+			{
+				UpdateNavOctreeParentChain(NavigationData->SourceElement);
+			}
 		}
 	}
 
 	// mark pending update as invalid, it will be dirtied according to currently active settings
-	const bool bCanInvalidateQueue = (UpdateFlags & FNavigationOctreeController::OctreeUpdate_Refresh) == 0;
-	if (bCanInvalidateQueue)
+	if (const bool bCanInvalidateQueue = (UpdateFlags & FNavigationOctreeController::OctreeUpdate_Refresh) == 0)
 	{
-		const FSetElementId RequestId = OctreeController.PendingOctreeUpdates.FindId(FNavigationDirtyElement(&ElementOwner));
+		const FSetElementId RequestId = OctreeController.PendingUpdates.FindId(NavRelevantElementHandle);
 		if (RequestId.IsValidId())
 		{
-			FNavigationDirtyElement& DirtyElement = OctreeController.PendingOctreeUpdates[RequestId];
+			FNavigationDirtyElement& DirtyElement = OctreeController.PendingUpdates[RequestId];
 
 			// Only consider as unregistered when pending update was not already invalidated since return value must indicate
 			// that ElementOwner was fully added or about to be added (valid pending update).
@@ -256,38 +317,53 @@ bool FNavigationDataHandler::UnregisterNavOctreeElement(UObject& ElementOwner, I
 	return bUnregistered;
 }
 
+// Deprecated
 void FNavigationDataHandler::UpdateNavOctreeElement(UObject& ElementOwner, INavRelevantInterface& ElementInterface, int32 UpdateFlags)
+{
+	UpdateNavOctreeElement(FNavigationElementHandle(&ElementOwner), FNavigationElement::CreateFromNavRelevantInterface(ElementInterface), UpdateFlags);
+}
+
+void FNavigationDataHandler::UpdateNavOctreeElement(FNavigationElementHandle ElementHandle, const TSharedRef<const FNavigationElement>& UpdatedElement, int32 UpdateFlags)
 {
 	INC_DWORD_STAT(STAT_Navigation_UpdateNavOctree);
 
+	if (OctreeController.IsValid() == false)
+	{
+		UE_LOG(LogNavOctree, VeryVerbose, TEXT("IGNORE(%hs) %s: octree not created yet"),
+			__FUNCTION__, *UpdatedElement.Get().GetPathName());
+		return;
+	}
+
 	if (OctreeController.IsNavigationOctreeLocked())
 	{
-		UE_LOG(LogNavOctree, Log, TEXT("IGNORE(UpdateNavOctreeElement) %s"), *ElementOwner.GetPathName());
+		UE_LOG(LogNavOctree, Log, TEXT("IGNORE(%hs) %s: octree locked"),
+			__FUNCTION__, *UpdatedElement.Get().GetPathName());
 		return;
 	}
 
 	// grab existing octree data
 	FBox CurrentBounds;
 	ENavigationDirtyFlag CurrentFlags;
-	const bool bAlreadyExists = OctreeController.GetNavOctreeElementData(ElementOwner, CurrentFlags, CurrentBounds);
+	const bool bAlreadyExists = OctreeController.GetNavOctreeElementData(ElementHandle, CurrentFlags, CurrentBounds);
 
 	// don't invalidate pending requests
 	UpdateFlags |= FNavigationOctreeController::OctreeUpdate_Refresh;
 
-	// always try to unregister, even if element owner doesn't exists in octree (parent nodes)
-	UnregisterNavOctreeElement(ElementOwner, ElementInterface, UpdateFlags);
+	// Always try to unregister, even if element owner doesn't exist in octree (parent nodes).
+	// This is also why we need to provide the new element and not only the handle, so we can access the parent (expected to be always the same for an update).
+	UnregisterElementWithNavOctree(UpdatedElement, UpdateFlags);
 
-	const FSetElementId RequestId = RegisterNavOctreeElement(ElementOwner, ElementInterface, UpdateFlags);
+	const FSetElementId RequestId = RegisterElementWithNavOctree(UpdatedElement, UpdateFlags);
 
 	// add original data to pending registration request
 	// so it could be dirtied properly when system receive unregister request while actor is still queued
 	if (RequestId.IsValidId())
 	{
-		FNavigationDirtyElement& UpdateInfo = OctreeController.PendingOctreeUpdates[RequestId];
+		FNavigationDirtyElement& UpdateInfo = OctreeController.PendingUpdates[RequestId];
 		UpdateInfo.PrevFlags = CurrentFlags;
 		if (UpdateInfo.PrevBounds.IsValid)
 		{
-			// Is we have something stored already we want to 
+			// If we have something stored already we want to 
 			// sum it up, since we care about the whole bounding
 			// box of changes that potentially took place
 			UpdateInfo.PrevBounds += CurrentBounds;
@@ -299,81 +375,75 @@ void FNavigationDataHandler::UpdateNavOctreeElement(UObject& ElementOwner, INavR
 		UpdateInfo.bHasPrevData = bAlreadyExists;
 	}
 
-	UpdateNavOctreeParentChain(ElementOwner, /*bSkipElementOwnerUpdate=*/ true);
+	UpdateNavOctreeParentChain(UpdatedElement, /*bSkipElementOwnerUpdate=*/ true);
 }
 
+// Deprecated
 void FNavigationDataHandler::UpdateNavOctreeParentChain(UObject& ElementOwner, bool bSkipElementOwnerUpdate)
 {
-	const int32 UpdateFlags = FNavigationOctreeController::OctreeUpdate_ParentChain | FNavigationOctreeController::OctreeUpdate_Refresh;
+	if (const INavRelevantInterface* NavRelevantInterface = Cast<INavRelevantInterface>(&ElementOwner))
+	{
+		UpdateNavOctreeParentChain(FNavigationElement::CreateFromNavRelevantInterface(*NavRelevantInterface), bSkipElementOwnerUpdate);
+	}
+}
+
+void FNavigationDataHandler::UpdateNavOctreeParentChain(const TSharedRef<const FNavigationElement>& Element, const bool bSkipElementOwnerUpdate)
+{
+	constexpr int32 UpdateFlags = FNavigationOctreeController::OctreeUpdate_ParentChain | FNavigationOctreeController::OctreeUpdate_Refresh;
 
-	TArray<FWeakObjectPtr> ChildNodes;
-	OctreeController.OctreeChildNodesMap.MultiFind(&ElementOwner, ChildNodes);
+	TArray<const TSharedRef<const FNavigationElement>> ChildNodes;
+	OctreeController.GetChildren(Element->GetHandle(), ChildNodes);
 
 	auto ElementOwnerUpdateFunc = [&]()->bool
 	{
 		bool bShouldRegisterChildren = true;
 		if (bSkipElementOwnerUpdate == false)
 		{
-			INavRelevantInterface* ElementInterface = Cast<INavRelevantInterface>(&ElementOwner);
-			if (ElementInterface != nullptr)
+			// We don't want to register NavOctreeElement if owner was not already registered or queued
+			// so we use Unregister/Register combo instead of UpdateNavOctreeElement
+			if (UnregisterElementWithNavOctree(Element, UpdateFlags))
 			{
-				// We don't want to register NavOctreeElement if owner was not already registered or queued
-				// so we use Unregister/Register combo instead of UpdateNavOctreeElement
-				if (UnregisterNavOctreeElement(ElementOwner, *ElementInterface, UpdateFlags))
-				{
-					FSetElementId NewId = RegisterNavOctreeElement(ElementOwner, *ElementInterface, UpdateFlags);
-					bShouldRegisterChildren = NewId.IsValidId();
-				}
-				else
-				{
-					bShouldRegisterChildren = false;
-				}
+				const FSetElementId NewId = RegisterElementWithNavOctree(Element, UpdateFlags);
+				bShouldRegisterChildren = NewId.IsValidId();
+			}
+			else
+			{
+				bShouldRegisterChildren = false;
 			}
 		}
-		return bShouldRegisterChildren;		
+		return bShouldRegisterChildren;
 	};
 
 	if (ChildNodes.Num() == 0)
 	{
-		// Last child was removed, only need to rebuild owner's NavOctreeElement 
+		// Last child was removed, only need to rebuild owner's NavOctreeElement
 		ElementOwnerUpdateFunc();
 		return;
 	}
 
-	TArray<INavRelevantInterface*> ChildNavInterfaces;
-	ChildNavInterfaces.AddZeroed(ChildNodes.Num());
-
-	for (int32 Idx = 0; Idx < ChildNodes.Num(); Idx++)
+	for (const TSharedRef<const FNavigationElement>& ChildNode : ChildNodes)
 	{
-		if (ChildNodes[Idx].IsValid())
-		{
-			UObject* ChildNodeOb = ChildNodes[Idx].Get();
-			ChildNavInterfaces[Idx] = Cast<INavRelevantInterface>(ChildNodeOb);
-			if (ChildNodeOb && ChildNavInterfaces[Idx])
-			{
-				UnregisterNavOctreeElement(*ChildNodeOb, *ChildNavInterfaces[Idx], UpdateFlags);
-			}
-		}
+		UnregisterElementWithNavOctree(ChildNode, UpdateFlags);
 	}
 
-	const bool bShouldRegisterChildren = ElementOwnerUpdateFunc();
-	
-	if (bShouldRegisterChildren)
+	if (const bool bShouldRegisterChildren = ElementOwnerUpdateFunc())
 	{
-		for (int32 Idx = 0; Idx < ChildNodes.Num(); Idx++)
+		for (const TSharedRef<const FNavigationElement>& ChildNode : ChildNodes)
 		{
-			UObject* ChildElement = ChildNodes[Idx].Get();
-			if (ChildElement && ChildNavInterfaces[Idx])
-			{
-				RegisterNavOctreeElement(*ChildElement, *ChildNavInterfaces[Idx], UpdateFlags);
-			}
+			RegisterElementWithNavOctree(ChildNode, UpdateFlags);
 		}
 	}
 }
 
-bool FNavigationDataHandler::UpdateNavOctreeElementBounds(UObject& Object, const FBox& NewBounds, const TConstArrayView<FBox> DirtyAreas)
+// Deprecated
+bool FNavigationDataHandler::UpdateNavOctreeElementBounds(UObject& ElementOwner, const FBox& NewBounds, const TConstArrayView<FBox> DirtyAreas)
+{
+	return UpdateNavOctreeElementBounds(FNavigationElementHandle(&ElementOwner), NewBounds, DirtyAreas);
+}
+
+bool FNavigationDataHandler::UpdateNavOctreeElementBounds(const FNavigationElementHandle ElementHandle, const FBox& NewBounds, const TConstArrayView<FBox> DirtyAreas)
 {
-	const FOctreeElementId2* ElementId = OctreeController.GetObjectsNavOctreeId(Object);
+	const FOctreeElementId2* ElementId = OctreeController.GetNavOctreeIdForElement(ElementHandle);
 	if (ElementId != nullptr && ensure(OctreeController.IsValidElement(*ElementId)))
 	{
 		OctreeController.NavOctree->UpdateNode(*ElementId, NewBounds);
@@ -382,44 +452,50 @@ bool FNavigationDataHandler::UpdateNavOctreeElementBounds(UObject& Object, const
 		if (DirtyAreas.Num() > 0)
 		{
 			// Refresh ElementId since object may be stored in a different node after updating bounds
-			ElementId = OctreeController.GetObjectsNavOctreeId(Object);
+			ElementId = OctreeController.GetNavOctreeIdForElement(ElementHandle);
 			if (ElementId != nullptr && ensure(OctreeController.IsValidElement(*ElementId)))
 			{
 				const FNavigationOctreeElement& ElementData = OctreeController.NavOctree->GetElementById(*ElementId);
-				DirtyAreasController.AddAreas(DirtyAreas, ElementData.Data->GetDirtyFlag(), [&Object] { return &Object; }, nullptr, "Bounds change");
+				DirtyAreasController.AddAreas(
+					DirtyAreas,
+					ElementData.Data->GetDirtyFlag(),
+					[SourceElement = ElementData.Data->SourceElement]
+					{
+						return SourceElement;
+					},
+					/*DirtyElement*/ nullptr,
+					"Bounds change");
 			}
 		}
 
 		return true;
 	}
 
-	// If dirty areas are provided we need to append them to a pending update since the object is not added yet.
-	// Not necessary for the bounds since they are not stored in the update but fetched when the update is processed.
-	if (DirtyAreas.Num() > 0)
+	// Update bounds and to append dirty areas to a pending update since the element is not added yet.
+	if (const FSetElementId PendingElementId = OctreeController.PendingUpdates.FindId(ElementHandle); PendingElementId.IsValidId())
 	{
-		const FSetElementId PendingElementId = OctreeController.PendingOctreeUpdates.FindId(FNavigationDirtyElement(&Object));
-		if (PendingElementId.IsValidId())
+		if (FNavigationDirtyElement& DirtyElement = OctreeController.PendingUpdates[PendingElementId]; !DirtyElement.bInvalidRequest)
 		{
-			FNavigationDirtyElement& DirtyElement = OctreeController.PendingOctreeUpdates[PendingElementId];
-			if (!DirtyElement.bInvalidRequest)
-			{
-				DirtyElement.ExplicitAreasToDirty.Append(DirtyAreas);
-				return true;
-			}
+			const TSharedRef<FNavigationElement> Updated = MakeShared<FNavigationElement>(DirtyElement.NavigationElement.Get());
+			Updated->SetBounds(NewBounds);
+			DirtyElement.NavigationElement = Updated;
+			DirtyElement.ExplicitAreasToDirty.Append(DirtyAreas);
+			return true;
 		}
 	}
 
 	return false;
 }
 
+// Deprecated
 bool FNavigationDataHandler::UpdateNavOctreeElementBounds(UObject& Object, const FBox& NewBounds, const FBox& DirtyArea)
 {
-	return UpdateNavOctreeElementBounds(Object, NewBounds,  TConstArrayView<FBox>{DirtyArea});
+	return UpdateNavOctreeElementBounds(FNavigationElementHandle(&Object), NewBounds, TConstArrayView<FBox>{DirtyArea});
 }
-	
+
 void FNavigationDataHandler::FindElementsInNavOctree(const FBox& QueryBox, const FNavigationOctreeFilter& Filter, TArray<FNavigationOctreeElement>& Elements)
 {
-	if (OctreeController.NavOctree.IsValid() == false)
+	if (OctreeController.IsValid() == false)
 	{
 		UE_LOG(LogNavOctree, Warning, TEXT("FNavigationDataHandler::FindElementsInNavOctree gets called while NavOctree is null"));
 		return;
@@ -434,9 +510,19 @@ void FNavigationDataHandler::FindElementsInNavOctree(const FBox& QueryBox, const
 	});
 }
 
+// Deprecated
 bool FNavigationDataHandler::ReplaceAreaInOctreeData(const UObject& Object, TSubclassOf<UNavArea> OldArea, TSubclassOf<UNavArea> NewArea, bool bReplaceChildClasses)
 {
-	FNavigationRelevantData* Data = OctreeController.GetMutableDataForObject(Object);
+	return ReplaceAreaInOctreeData(FNavigationElementHandle(&Object), OldArea, NewArea, bReplaceChildClasses);
+}
+
+bool FNavigationDataHandler::ReplaceAreaInOctreeData(
+	const FNavigationElementHandle Element,
+	const TSubclassOf<UNavArea> OldArea,
+	const TSubclassOf<UNavArea> NewArea,
+	const bool bReplaceChildClasses) const
+{
+	FNavigationRelevantData* Data = OctreeController.GetMutableDataForElement(Element);
 
 	if (Data == nullptr || Data->HasModifiers() == false)
 	{
@@ -470,10 +556,8 @@ bool FNavigationDataHandler::ReplaceAreaInOctreeData(const UObject& Object, TSub
 			}
 		}
 	}
-	for (FCustomLinkNavModifier& CustomLink : Data->Modifiers.GetCustomLinks())
-	{
-		ensureMsgf(false, TEXT("Not implemented yet"));
-	}
+
+	ensureMsgf(Data->Modifiers.GetCustomLinks().IsEmpty(), TEXT("Not implemented yet"));
 
 	return true;
 }
@@ -481,16 +565,18 @@ bool FNavigationDataHandler::ReplaceAreaInOctreeData(const UObject& Object, TSub
 void FNavigationDataHandler::AddLevelCollisionToOctree(ULevel& Level)
 {
 #if WITH_RECAST
-	if (OctreeController.NavOctree.IsValid() &&
+	if (OctreeController.IsValid() &&
 		OctreeController.NavOctree->GetNavGeometryStoringMode() == FNavigationOctree::StoreNavGeometry)
 	{
+		const FNavigationElementHandle ElementKey(&Level);
 		const TArray<FVector>* LevelGeom = Level.GetStaticNavigableGeometry();
-		const FOctreeElementId2* ElementId = OctreeController.GetObjectsNavOctreeId(Level);
+		const FOctreeElementId2* ElementId = OctreeController.GetNavOctreeIdForElement(ElementKey);
 
 		if (!ElementId && LevelGeom && LevelGeom->Num() > 0)
 		{
-			FNavigationOctreeElement BSPElem(Level);
-			
+			TSharedRef<const FNavigationElement> NavigationElement = MakeShared<const FNavigationElement>(Level, INDEX_NONE);
+			FNavigationOctreeElement BSPElem(NavigationElement);
+
 			// In WP dynamic mode, store if this is loaded data.
 			if (DirtyAreasController.bUseWorldPartitionedDynamicMode)
 			{
@@ -502,10 +588,18 @@ void FNavigationDataHandler::AddLevelCollisionToOctree(ULevel& Level)
 			const FBox& Bounds = BSPElem.Data->Bounds;
 			if (!Bounds.GetExtent().IsNearlyZero())
 			{
-				OctreeController.NavOctree->AddNode(&Level, nullptr, Bounds, BSPElem);
-				DirtyAreasController.AddArea(Bounds, ENavigationDirtyFlag::All, [&Level] { return &Level; }, nullptr, "Add level");
-
-				UE_LOG(LogNavOctree, Log, TEXT("ADD %s"), *UE::NavigationHelper::Private::GetElementName(Level));
+				OctreeController.NavOctree->AddNode(Bounds, BSPElem);
+				DirtyAreasController.AddArea(
+					Bounds, 
+					ENavigationDirtyFlag::All,
+					[SourceElement = NavigationElement]
+					{
+						return SourceElement;
+					},
+					/*DirtyElement*/ nullptr,
+					"Add level");
+
+				UE_LOG(LogNavOctree, Log, TEXT("ADD %s"), *NavigationElement.Get().GetPathName());
 			}
 		}
 	}
@@ -514,48 +608,13 @@ void FNavigationDataHandler::AddLevelCollisionToOctree(ULevel& Level)
 
 void FNavigationDataHandler::RemoveLevelCollisionFromOctree(ULevel& Level)
 {
-	if (OctreeController.NavOctree.IsValid())
+	if (OctreeController.IsValid())
 	{
-		const FOctreeElementId2* ElementId = OctreeController.GetObjectsNavOctreeId(Level);
-		UE_LOG(LogNavOctree, Log, TEXT("UNREG %s %s"), *UE::NavigationHelper::Private::GetElementName(Level), ElementId ? TEXT("[exists]") : TEXT(""));
-
-		if (ElementId != nullptr)
+		const FNavigationElementHandle NavigationElementHandle(&Level);
+		if (const FOctreeElementId2* OctreeElementId = OctreeController.GetNavOctreeIdForElement(NavigationElementHandle))
 		{
-			if (ensure(OctreeController.IsValidElement(*ElementId)))
-			{
-				// mark area occupied by given actor as dirty
-				const FNavigationOctreeElement& ElementData = OctreeController.NavOctree->GetElementById(*ElementId);
-				DirtyAreasController.AddArea(ElementData.Bounds.GetBox(), ENavigationDirtyFlag::All, [&Level] { return &Level; }, nullptr, "Remove level");
-			}
-
-			OctreeController.NavOctree->RemoveNode(*ElementId);
-			OctreeController.RemoveObjectsNavOctreeId(Level);
-		}
-	}
-}
-
-void FNavigationDataHandler::UpdateActorAndComponentsInNavOctree(AActor& Actor)
-{
-	INavRelevantInterface* NavElement = Cast<INavRelevantInterface>(&Actor);
-	if (NavElement)
-	{
-		UpdateNavOctreeElement(Actor, *NavElement, FNavigationOctreeController::OctreeUpdate_Default);
-	}
-
-	for (UActorComponent* Component : Actor.GetComponents())
-	{
-		INavRelevantInterface* CompNavElement = Cast<INavRelevantInterface>(Component);
-		if (CompNavElement)
-		{
-			// Component != null is implied by successful INavRelevantInterface cast 
-			if (Actor.IsComponentRelevantForNavigation(Component))
-			{
-				UpdateNavOctreeElement(*Component, *CompNavElement, FNavigationOctreeController::OctreeUpdate_Default);
-			}
-			else
-			{
-				UnregisterNavOctreeElement(*Component, *CompNavElement, FNavigationOctreeController::OctreeUpdate_Default);
-			}
+			UE_LOG(LogNavOctree, Log, TEXT("UNREG %s %s"), *Level.GetPathName(), OctreeElementId ? TEXT("[exists]") : TEXT(""));
+			RemoveFromNavOctree(*OctreeElementId, FNavigationOctreeController::OctreeUpdate_Geometry);
 		}
 	}
 }
@@ -566,18 +625,18 @@ void FNavigationDataHandler::ProcessPendingOctreeUpdates()
 
 	if (OctreeController.NavOctree)
 	{
-		// AddElementToNavOctree (through some of its resulting function calls) modifies PendingOctreeUpdates so invalidates the iterators,
-		// (via WaitUntilAsyncPropertyReleased() / UpdateComponentInNavOctree() / RegisterNavOctreeElement()). This means we can't iterate
+		// AddElementToNavOctree (through some of its resulting function calls) modifies PendingUpdates so invalidates the iterators,
+		// (via WaitUntilAsyncPropertyReleased() / UpdateComponentInNavOctree() / RegisterElementWithNavOctree()). This means we can't iterate
 		// through this set in the normal way. Previously the code iterated through this which also left us open to other potential bugs
 		// in that we may have tried to modify elements we had already processed.
-		while (TSet<FNavigationDirtyElement>::TIterator It = OctreeController.PendingOctreeUpdates.CreateIterator())
+		while (TSet<FNavigationDirtyElement, FNavigationDirtyElementKeyFunctions>::TIterator It = OctreeController.PendingUpdates.CreateIterator())
 		{
-			FNavigationDirtyElement Element = *It;
+			FNavigationDirtyElement DirtyElement = *It;
 			It.RemoveCurrent();
-			AddElementToNavOctree(Element);
+			AddElementToNavOctree(DirtyElement);
 		}
 	}
-	OctreeController.PendingOctreeUpdates.Empty(32);
+	OctreeController.PendingUpdates.Empty(32);
 }
 
 void FNavigationDataHandler::DemandLazyDataGathering(FNavigationRelevantData& ElementData)
@@ -588,21 +647,13 @@ void FNavigationDataHandler::DemandLazyDataGathering(FNavigationRelevantData& El
     // Check if any child asked for some lazy gathering
 	if (ElementData.IsPendingChildLazyModifiersGathering())
 	{
-		TArray<FWeakObjectPtr> ChildNodes;
-		OctreeController.OctreeChildNodesMap.MultiFind(ElementData.GetOwner(), ChildNodes);
+		TArray<const TSharedRef<const FNavigationElement>> ChildNodes;
+		OctreeController.GetChildren(ElementData.SourceElement->GetHandle(), ChildNodes);
 
-		for (FWeakObjectPtr& ChildNode : ChildNodes)
+		for (const TSharedRef<const FNavigationElement>& ChildNode : ChildNodes)
 		{
-			if (ChildNode.IsValid())
-			{
-				UObject* ChildNodeOb = ChildNode.Get();
-				INavRelevantInterface* ChildNavInterface = ChildNodeOb ? Cast<INavRelevantInterface>(ChildNodeOb) : nullptr;
-				if (ChildNavInterface)
-				{
-					OctreeController.NavOctree->DemandChildLazyDataGathering(ElementData, *ChildNavInterface);
-				}
-			}
+			OctreeController.NavOctree->DemandChildLazyDataGathering(ElementData, ChildNode.Get());
 		}
 		ElementData.bPendingChildLazyModifiersGathering = false;
 	}
-}
+}
\ No newline at end of file
diff --git a/Engine/Source/Runtime/NavigationSystem/Private/NavigationDirtyAreasController.cpp b/Engine/Source/Runtime/NavigationSystem/Private/NavigationDirtyAreasController.cpp
--- a/Engine/Source/Runtime/NavigationSystem/Private/NavigationDirtyAreasController.cpp
+++ b/Engine/Source/Runtime/NavigationSystem/Private/NavigationDirtyAreasController.cpp
@@ -78,7 +78,8 @@ void FNavigationDirtyAreasController::Tick(const float DeltaSeconds, const TArra
 			for (const FNavigationDirtyArea& DirtyArea : DirtyAreas)
 			{
 				const FBox& AreaBound = DirtyArea.Bounds;
-				if (!ensureMsgf(AreaBound.IsValid, TEXT("%hs Attempting to use DirtyArea.Bounds which are not valid. SourceObject: %s"), __FUNCTION__, *GetFullNameSafe(DirtyArea.OptionalSourceObject.Get())))
+				if (!ensureMsgf(AreaBound.IsValid, TEXT("%hs Attempting to use DirtyArea.Bounds which are not valid. SourceObject: %s"),
+					__FUNCTION__, *DirtyArea.GetSourceDescription()))
 				{
 					continue;
 				}
@@ -91,7 +92,7 @@ void FNavigationDirtyAreasController::Tick(const float DeltaSeconds, const TArra
 						const FBox OverlapBox = AreaBound.Overlap(SeedBounds);
 						if (OverlapBox.IsValid)
 						{
-							SubAreaArray.Emplace(OverlapBox, DirtyArea.Flags, DirtyArea.OptionalSourceObject.Get());
+							SubAreaArray.Emplace(OverlapBox, DirtyArea.Flags, DirtyArea.OptionalSourceElement);
 						}
 					}
 				}
@@ -119,22 +120,22 @@ void FNavigationDirtyAreasController::Tick(const float DeltaSeconds, const TArra
 void FNavigationDirtyAreasController::AddArea(const FBox& NewArea, const int32 Flags, const TFunction<UObject*()>& ObjectProviderFunc /*= nullptr*/,
 	const FNavigationDirtyElement* DirtyElement /*= nullptr*/, const FName& DebugReason /*= NAME_None*/)
 {
-	AddAreas({NewArea}, static_cast<ENavigationDirtyFlag>(Flags), ObjectProviderFunc, DirtyElement, DebugReason);
+	AddAreas({NewArea}, static_cast<ENavigationDirtyFlag>(Flags), /*ElementProviderFunc*/nullptr, DirtyElement, DebugReason);
 }
 
 // Deprecated
 void FNavigationDirtyAreasController::AddAreas(const TConstArrayView<FBox> NewAreas, const int32 Flags, const TFunction<UObject*()>& ObjectProviderFunc, const FNavigationDirtyElement* DirtyElement, const FName& DebugReason)
 {
-	AddAreas(NewAreas, static_cast<ENavigationDirtyFlag>(Flags), ObjectProviderFunc, DirtyElement, DebugReason);
+	AddAreas(NewAreas, static_cast<ENavigationDirtyFlag>(Flags), /*ElementProviderFunc*/nullptr, DirtyElement, DebugReason);
 }
 
-void FNavigationDirtyAreasController::AddArea(const FBox& NewArea, const ENavigationDirtyFlag Flags, const TFunction<UObject*()>& ObjectProviderFunc /*= nullptr*/,
+void FNavigationDirtyAreasController::AddArea(const FBox& NewArea, const ENavigationDirtyFlag Flags, const TFunction<const TSharedPtr<const FNavigationElement>()>& ElementProviderFunc /*= nullptr*/,
 	const FNavigationDirtyElement* DirtyElement /*= nullptr*/, const FName& DebugReason /*= NAME_None*/)
 {
-	AddAreas({NewArea}, Flags, ObjectProviderFunc, DirtyElement, DebugReason);
+	AddAreas({NewArea}, Flags, ElementProviderFunc, DirtyElement, DebugReason);
 }
 
-void FNavigationDirtyAreasController::AddAreas(const TConstArrayView<FBox> NewAreas, const ENavigationDirtyFlag Flags, const TFunction<UObject*()>& ObjectProviderFunc, const FNavigationDirtyElement* DirtyElement, const FName& DebugReason)
+void FNavigationDirtyAreasController::AddAreas(const TConstArrayView<FBox> NewAreas, const ENavigationDirtyFlag Flags, const TFunction<const TSharedPtr<const FNavigationElement>()>& ElementProviderFunc, const FNavigationDirtyElement* DirtyElement, const FName& DebugReason)
 {
 #if !UE_BUILD_SHIPPING
 	// always keep track of reported areas even when filtered out by invalid area as long as flags are valid
@@ -143,30 +144,30 @@ void FNavigationDirtyAreasController::AddAreas(const TConstArrayView<FBox> NewAr
 	checkf(NewAreas.Num() > 0, TEXT("All callers of this method are expected to provide at least one area."));
 #endif // !UE_BUILD_SHIPPING
 
-	UObject* SourceObject = ObjectProviderFunc ? ObjectProviderFunc() : nullptr;
-
+	const TSharedPtr<const FNavigationElement> SourceElement = ElementProviderFunc ? ElementProviderFunc() : nullptr;
 	if (bUseWorldPartitionedDynamicMode)
 	{
 		// Both conditions must be true to ignore dirtiness.
 		//  If it's only a visibility change and it's not in the base navmesh, it's the case created by loading a data layer (dirtiness must be applied)
 		//  If there is no visibility change, the change is not from loading/unloading a cell (dirtiness must be applied)
 		
-		// ObjectProviderFunc() is not always providing a valid object.
-		if (const bool bIsFromVisibilityChange = (DirtyElement && DirtyElement->bIsFromVisibilityChange) || (SourceObject && FNavigationSystem::IsLevelVisibilityChanging(SourceObject)))
+		// ElementProviderFunc() is not always providing a valid element.
+		if (const bool bIsFromVisibilityChange = (DirtyElement && DirtyElement->bIsFromVisibilityChange) || (SourceElement && SourceElement->IsFromLevelVisibilityChange()))
 		{
-			// If the area is from the addition or removal of objects caused by level loading/unloading and it's already in the base navmesh ignore the dirtiness.
-			if (const bool bIsIncludedInBaseNavmesh = (DirtyElement && DirtyElement->bIsInBaseNavmesh) || (SourceObject && FNavigationSystem::IsInBaseNavmesh(SourceObject)))
+			// If the area is from the addition or removal of elements caused by level loading/unloading and it's already in the base navmesh ignore the dirtiness.
+			if (const bool bIsIncludedInBaseNavmesh = (DirtyElement && DirtyElement->bIsInBaseNavmesh) || (SourceElement && SourceElement->IsInBaseNavigationData()))
 			{
-				UE_LOG(LogNavigationDirtyArea, VeryVerbose, TEXT("Ignoring dirtyness (visibility changed and in base navmesh). (object: %s from: %s)"),
-					*GetFullNameSafe(SourceObject), *DebugReason.ToString());
+				UE_LOG(LogNavigationDirtyArea, VeryVerbose, TEXT("Ignoring dirtyness (visibility changed and in base navmesh). (element: %s from: %s)"),
+					*GetFullNameSafe(SourceElement.Get()), *DebugReason.ToString());
 				return;
 			}
 		}
 	}
 
-	if (ShouldSkipObjectPredicate.IsBound() && SourceObject)
+	if (ShouldSkipObjectPredicate.IsBound())
 	{
-		if (ShouldSkipObjectPredicate.Execute(*SourceObject))
+		const UObject* SourceObject = SourceElement->GetWeakUObject().Get();
+		if (SourceObject && ShouldSkipObjectPredicate.Execute(*SourceObject))
 		{
 			return;
 		}
@@ -190,7 +191,9 @@ void FNavigationDirtyAreasController::AddAreas(const TConstArrayView<FBox> NewAr
 		}
 
 #if !UE_BUILD_SHIPPING
-		auto DumpExtraInfo = [SourceObject, DebugReason, BoundsSize, NewArea]() {
+		auto DumpExtraInfo = [SourceElement, DebugReason, BoundsSize, NewArea]()
+			{
+				const UObject* SourceObject = SourceElement->GetWeakUObject().Get();
 				FString ObjectInfo;
 				if (const UObject* ObjectOwner = (SourceObject != nullptr ? SourceObject->GetOuter() : nullptr))
 				{
@@ -221,7 +224,7 @@ void FNavigationDirtyAreasController::AddAreas(const TConstArrayView<FBox> NewAr
 					*ObjectInfo,
 					*BoundsSize.ToString(),
 					*ActorInfo);
-		};
+			};
 
 		if (ShouldReportOversizedDirtyArea() && BoundsSize.GetMax() > DirtyAreaWarningSizeThreshold)
 		{
@@ -235,13 +238,13 @@ void FNavigationDirtyAreasController::AddAreas(const TConstArrayView<FBox> NewAr
 
 		if (Flags != ENavigationDirtyFlag::None && bCanAccumulateDirtyAreas)
 		{
-			DirtyAreas.Add(FNavigationDirtyArea(NewArea, Flags, SourceObject));
+			DirtyAreas.Add(FNavigationDirtyArea(NewArea, Flags, SourceElement));
 		}
 	}
 	
 	UE_CLOG(NumInvalidBounds > 0 || NumEmptyBounds > 0, LogNavigationDirtyArea, Warning,
-		TEXT("Skipped some dirty area creation due to: %d invalid bounds, %d empty bounds (object: %s, from: %s)"),
-		NumInvalidBounds, NumEmptyBounds, *GetFullNameSafe(SourceObject), *DebugReason.ToString());
+		TEXT("Skipped some dirty area creation due to: %d invalid bounds, %d empty bounds (element: %s, from: %s)"),
+		NumInvalidBounds, NumEmptyBounds, *GetFullNameSafe(SourceElement.Get()), *DebugReason.ToString());
 }
 
 void FNavigationDirtyAreasController::OnNavigationBuildLocked()
diff --git a/Engine/Source/Runtime/NavigationSystem/Private/NavigationObjectRepository.cpp b/Engine/Source/Runtime/NavigationSystem/Private/NavigationObjectRepository.cpp
--- a/Engine/Source/Runtime/NavigationSystem/Private/NavigationObjectRepository.cpp
+++ b/Engine/Source/Runtime/NavigationSystem/Private/NavigationObjectRepository.cpp
@@ -1,47 +1,212 @@
 ﻿// Copyright Epic Games, Inc. All Rights Reserved.
 
 #include "NavigationObjectRepository.h"
-#include "AI/Navigation/NavRelevantInterface.h"
 #include "NavLinkCustomInterface.h"
+#include "AI/NavigationSystemBase.h"
+#include "AI/Navigation/NavigationElement.h"
+#include "AI/Navigation/NavRelevantInterface.h"
+#include "UObject/ObjectKey.h"
 
-void UNavigationObjectRepository::RegisterNavRelevantObject(INavRelevantInterface& NavRelevantObject)
+TSharedPtr<const FNavigationElement> UNavigationObjectRepository::AddNavigationElement(FNavigationElement&& Element, const ENotifyOnSuccess NotifyOnSuccess /*= ENotifyOnSuccess::Yes*/)
 {
 #if DO_ENSURE // We don't want to execute the Find at all for targets where ensures are disabled
-	if (!ensureMsgf(NavRelevantObjects.Find(&NavRelevantObject) == INDEX_NONE, TEXT("Same interface pointer can't be registered twice.")))
 	{
-		return;
+		UE_MT_SCOPED_READ_ACCESS(NavElementAccessDetector);
+
+		if (!ensureMsgf(NavRelevantElements.Find(Element.GetHandle()) == nullptr, TEXT("Same element can't be registered twice.")))
+		{
+			return nullptr;
+		}
 	}
 #endif
 
-	NavRelevantObjects.Emplace(&NavRelevantObject);
+	const TSharedRef SharedElement(MakeShared<FNavigationElement>(MoveTemp(Element)));
+	{
+		UE_MT_SCOPED_WRITE_ACCESS(NavElementAccessDetector);
+		NavRelevantElements.Emplace(Element.GetHandle(), SharedElement);
+	}
+
+	if (NotifyOnSuccess == ENotifyOnSuccess::Yes)
+	{
+		(void)OnNavigationElementAddedDelegate.ExecuteIfBound(SharedElement);
+	}
 
-	OnNavRelevantObjectRegistered.ExecuteIfBound(NavRelevantObject);
+	return SharedElement.ToSharedPtr();
 }
 
-void UNavigationObjectRepository::UnregisterNavRelevantObject(INavRelevantInterface& NavRelevantObject)
+void UNavigationObjectRepository::RemoveNavigationElement(const FNavigationElementHandle Handle)
 {
-	ensureMsgf(NavRelevantObjects.Remove(&NavRelevantObject) > 0, TEXT("Interface can't be removed since it was not registered or already unregistered)"));
+	UE_MT_SCOPED_WRITE_ACCESS(NavElementAccessDetector);
 
-	OnNavRelevantObjectUnregistered.ExecuteIfBound(NavRelevantObject);
+	TSharedPtr<const FNavigationElement> Element;
+	if (ensureMsgf(NavRelevantElements.RemoveAndCopyValue(Handle, Element),
+		TEXT("Navigation element can't be removed since it was not registered or already unregistered)")))
+	{
+		(void)OnNavigationElementRemovedDelegate.ExecuteIfBound(Element.ToSharedRef());
+	}
 }
 
-void UNavigationObjectRepository::RegisterCustomNavLinkObject(INavLinkCustomInterface& CustomNavLinkObject)
+void UNavigationObjectRepository::ForEachNavigationElement(TFunctionRef<void(const TSharedRef<const FNavigationElement>&)> PerElementCallback) const
 {
-#if DO_ENSURE // We don't want to execute the Find at all for targets where ensures are disabled
-	if (!ensureMsgf(CustomLinkObjects.Find(&CustomNavLinkObject) == INDEX_NONE, TEXT("Same interface pointer can't be registered twice.")))
+	UE_MT_SCOPED_READ_ACCESS(NavElementAccessDetector);
+
+	for (auto It = NavRelevantElements.CreateConstIterator(); It; ++It)
 	{
-		return;
+		if (const TSharedPtr<const FNavigationElement>& Element = It.Value())
+		{
+			PerElementCallback(Element.ToSharedRef());
+		}
 	}
+}
+
+TSharedPtr<const FNavigationElement> UNavigationObjectRepository::RegisterNavRelevantObject(const INavRelevantInterface& NavRelevantObject)
+{
+	return RegisterNavRelevantObjectInternal(NavRelevantObject, *Cast<UObject>(&NavRelevantObject), ENotifyOnSuccess::Yes);
+}
+
+TSharedPtr<const FNavigationElement> UNavigationObjectRepository::RegisterNavRelevantObjectInternal(
+	const INavRelevantInterface& NavRelevantInterface,
+	const UObject& NavRelevantObject,
+	const ENotifyOnSuccess NotifyOnSuccess)
+{
+	// In AActor/UActorComponent code paths it is possible that a component registration is performed more than once
+	// (i.e., Actor registering its component, then individual component registers too)
+	// In such case we update with the latest.
+	if (const TSharedPtr<const FNavigationElement> ElementPtr = GetNavigationElementForUObject(&NavRelevantObject))
+	{
+		const TSharedRef<const FNavigationElement> NewElement = FNavigationElement::CreateFromNavRelevantInterface(NavRelevantInterface);
+		NavRelevantElements[ElementPtr->GetHandle()] = NewElement;
+
+		if (NotifyOnSuccess == ENotifyOnSuccess::Yes)
+		{
+			(void)OnNavigationElementAddedDelegate.ExecuteIfBound(NewElement);
+		}
+
+		UE_LOG(LogNavigation, Verbose, TEXT("%hs [already registered - updating] (%s:%s) Bounds: [%s]->[%s]"), __FUNCTION__,
+			*GetNameSafe(NavRelevantObject.GetOuter()), *GetNameSafe(&NavRelevantObject),
+			*ElementPtr->GetBounds().ToString(), *NewElement->GetBounds().ToString());
+
+		return NewElement.ToSharedPtr();
+	}
+
+	if (NavRelevantInterface.IsNavigationRelevant())
+	{
+		if (const TSharedPtr<const FNavigationElement> SharedElement = AddNavigationElement(FNavigationElement(NavRelevantInterface), NotifyOnSuccess))
+		{
+			UE_MT_SCOPED_WRITE_ACCESS(NavElementAccessDetector);
+			ObjectsToHandleMap.Emplace(FObjectKey(&NavRelevantObject), SharedElement->GetHandle());
+
+			UE_LOG(LogNavigation, Verbose, TEXT("%hs [registered] (%s:%s) Bounds: [%s]"), __FUNCTION__,
+				*GetNameSafe(NavRelevantObject.GetOuter()), *GetNameSafe(&NavRelevantObject),
+				*NavRelevantInterface.GetNavigationBounds().ToString());
+
+			return SharedElement;
+		}
+
+		return nullptr;
+	}
+
+	UE_LOG(LogNavigation, VeryVerbose, TEXT("%hs [skipped: not relevant] (%s:%s)"), __FUNCTION__,
+		*GetNameSafe(NavRelevantObject.GetOuter()), *GetNameSafe(&NavRelevantObject));
+	return nullptr;
+}
+
+void UNavigationObjectRepository::UnregisterNavRelevantObject(const INavRelevantInterface& NavRelevantObject)
+{
+	UnregisterNavRelevantObject(Cast<UObject>(&NavRelevantObject));
+}
+
+void UNavigationObjectRepository::UnregisterNavRelevantObject(const UObject* NavRelevantObject)
+{
+	UE_LOG(LogNavigation, Verbose, TEXT("%hs (%s:%s)"), __FUNCTION__,
+		NavRelevantObject ? *GetNameSafe(NavRelevantObject->GetOuter()) : TEXT("null outer"),
+		*GetNameSafe(NavRelevantObject));
+
+	FNavigationElementHandle Handle;
+	{
+		UE_MT_SCOPED_WRITE_ACCESS(NavElementAccessDetector);
+		ObjectsToHandleMap.RemoveAndCopyValue(FObjectKey(NavRelevantObject), Handle);
+	}
+
+	if (Handle)
+	{
+		RemoveNavigationElement(Handle);
+	}
+}
+
+TSharedPtr<const FNavigationElement> UNavigationObjectRepository::GetNavigationElementForHandle(const FNavigationElementHandle Handle) const
+{
+	UE_MT_SCOPED_READ_ACCESS(NavElementAccessDetector);
+
+	if (const TSharedPtr<const FNavigationElement>* Element = NavRelevantElements.Find(Handle))
+	{
+		return *Element;
+	}
+
+	return nullptr;
+}
+
+FNavigationElementHandle UNavigationObjectRepository::GetNavigationElementHandleForUObject(const UObject* NavRelevantObject)
+{
+	UE_MT_SCOPED_READ_ACCESS(NavElementAccessDetector);
+	if (const FNavigationElementHandle* Handle = ObjectsToHandleMap.Find(FObjectKey(Cast<UObject>(NavRelevantObject))))
+	{
+		return *Handle;
+	}
+
+	return FNavigationElementHandle::Invalid;
+}
+
+TSharedPtr<const FNavigationElement> UNavigationObjectRepository::GetNavigationElementForUObject(const UObject* NavRelevantObject)
+{
+	UE_MT_SCOPED_READ_ACCESS(NavElementAccessDetector);
+
+	if (const FNavigationElementHandle* Handle = ObjectsToHandleMap.Find(FObjectKey(NavRelevantObject)))
+	{
+		if (const TSharedPtr<const FNavigationElement>* Element = NavRelevantElements.Find(*Handle))
+		{
+			return *Element;
+		}
+	}
+
+	return nullptr;
+}
+
+TSharedPtr<const FNavigationElement> UNavigationObjectRepository::UpdateNavigationElementForUObject(
+	const INavRelevantInterface& NavRelevantInterface,
+	const UObject& NavRelevantObject)
+{
+	// This method is called by the navigation system to make sure an up-to-date navigation element exists for a
+	// given navigation relevant UObject.
+	// In this case we only need to create, or update, the navigation element without sending
+	// notification (i.e. ENotifyOnSuccess::No) since the caller (NavigationSystem) is already in the process of updating.
+	return RegisterNavRelevantObjectInternal(NavRelevantInterface, NavRelevantObject, ENotifyOnSuccess::No);
+}
+
+void UNavigationObjectRepository::RegisterCustomNavLinkObject(INavLinkCustomInterface& CustomNavLinkObject)
+{
+	{
+		UE_MT_SCOPED_WRITE_ACCESS(NavElementAccessDetector);
+
+#if DO_ENSURE // We don't want to execute the Find at all for targets where ensures are disabled
+		if (!ensureMsgf(CustomLinkObjects.Find(&CustomNavLinkObject) == INDEX_NONE, TEXT("Same interface pointer can't be registered twice.")))
+		{
+			return;
+		}
 #endif
 
-	CustomLinkObjects.Emplace(&CustomNavLinkObject);
+		CustomLinkObjects.Emplace(&CustomNavLinkObject);
+	}
 
 	OnCustomNavLinkObjectRegistered.ExecuteIfBound(CustomNavLinkObject);
 }
 
 void UNavigationObjectRepository::UnregisterCustomNavLinkObject(INavLinkCustomInterface& CustomNavLinkObject)
 {
-	ensureMsgf(CustomLinkObjects.Remove(&CustomNavLinkObject) > 0, TEXT("Interface can't be removed since it was not registered or already unregistered)"));
+	{
+		UE_MT_SCOPED_WRITE_ACCESS(NavElementAccessDetector);
+		ensureMsgf(CustomLinkObjects.Remove(&CustomNavLinkObject) > 0, TEXT("Interface can't be removed since it was not registered or already unregistered)"));
+	}
 
 	OnCustomNavLinkObjectUnregistered.ExecuteIfBound(CustomNavLinkObject);
-}
+}
\ No newline at end of file
diff --git a/Engine/Source/Runtime/NavigationSystem/Private/NavigationObjectRepository.h b/Engine/Source/Runtime/NavigationSystem/Private/NavigationObjectRepository.h
--- a/Engine/Source/Runtime/NavigationSystem/Private/NavigationObjectRepository.h
+++ b/Engine/Source/Runtime/NavigationSystem/Private/NavigationObjectRepository.h
@@ -2,18 +2,23 @@
 
 #pragma once
 
+#include "AI/Navigation/NavigationElement.h" // required since we can't fwd declare types used as TMap keys
 #include "Subsystems/WorldSubsystem.h"
+#include "UObject/ObjectKey.h"
 #include "UObject/WeakInterfacePtr.h"
 #include "NavigationObjectRepository.generated.h"
 
 class INavLinkCustomInterface;
 class INavRelevantInterface;
 
+UE_DEPRECATED(5.5, "Use the overloaded version with FNavigationElement instead.")
 DECLARE_DELEGATE_OneParam(FOnNavRelevantObjectRegistrationEvent, INavRelevantInterface&);
+
+DECLARE_DELEGATE_OneParam(FOnNavigationElementRegistrationEvent, const TSharedRef<const FNavigationElement>&);
 DECLARE_DELEGATE_OneParam(FOnCustomNavLinkObjectRegistrationEvent, INavLinkCustomInterface&);
 
 /**
- * World subsystem dedicated to store different types of navigation related objects that the
+ * World subsystem dedicated to store different types of navigation related elements that the
  * NavigationSystem needs to access.
  */
 UCLASS()
@@ -23,28 +28,87 @@ class UNavigationObjectRepository : public UWorldSubsystem
 
 public:
 
+	// Specifies whether OnNavigationElementAddedDelegate must be broadcast on successful registration.
+	enum class ENotifyOnSuccess : uint8
+	{
+		No,
+		Yes
+	};
+
+	/**
+	 * Adds the provided navigation element to the list of registered elements.
+	 * @param Element Reference to the navigation element to register.
+	 * @param NotifyOnSuccess Indicates if OnNavigationElementRegisteredDelegate must be broadcast on successful registration.
+	 * @return Handle to the registered element.
+	 * @note Method will assert if the same element is registered twice.
+	 */
+	TSharedPtr<const FNavigationElement> AddNavigationElement(FNavigationElement&& Element, ENotifyOnSuccess NotifyOnSuccess = ENotifyOnSuccess::Yes);
+
+	/**
+	 * Removes the navigation element associated with the provided handle from the list of registered elements.
+	 * @param Handle Handle to the element to unregister.
+	 * @note Method will assert if the element can't be removed (i.e., not registered or already unregistered).
+	 */
+	void RemoveNavigationElement(FNavigationElementHandle Handle);
+
+	/**
+	 * Iterates through all registered navigation elements and call provided function with the element as parameter.
+	 */
+	void ForEachNavigationElement(TFunctionRef<void(const TSharedRef<const FNavigationElement>&)>) const;
+
 	/**
 	 * Adds the provided object implementing INavRelevantInterface to the list of registered navigation relevant objects.
 	 * @param NavRelevantObject INavRelevantInterface of the object to register.
 	 * @note Method will also assert if same interface pointers is registered twice.
 	 */
-	void RegisterNavRelevantObject(INavRelevantInterface& NavRelevantObject);
+	TSharedPtr<const FNavigationElement> RegisterNavRelevantObject(const INavRelevantInterface& NavRelevantObject);
 
 	/**
 	 * Removes the provided interface from the list of registered navigation relevant objects.
 	 * @param NavRelevantObject INavRelevantInterface of the object to unregister.
 	 * @note Method will also assert if interface can't be removed (i.e. not registered or already unregistered).
 	 */
-	void UnregisterNavRelevantObject(INavRelevantInterface& NavRelevantObject);
+	void UnregisterNavRelevantObject(const INavRelevantInterface& NavRelevantObject);
 
 	/**
-	 * Returns the list of registered navigation relevant objects.
-	 * @return Const view on the list of all registered navigation relevant objects.
+	 * Removes the provided object from the list of registered navigation relevant objects.
+	 * @param NavRelevantObject INavRelevantInterface of the object to unregister.
+	 * @note Method will also assert if interface can't be removed (i.e. not registered or already unregistered).
 	 */
-	TConstArrayView<TWeakInterfacePtr<INavRelevantInterface>> GetNavRelevantObjects() const
-	{
-		return NavRelevantObjects;
-	}
+	void UnregisterNavRelevantObject(const UObject* NavRelevantObject);
+
+	/**
+	 * Returns the handle associated with the specific UObject if it is registered in the repository.
+	 * @param NavRelevantObject The UObject for which the handle must be retrieved.
+	 * @return The handle associated with the specified UObject if already registered, an invalid handle otherwise.
+	 */
+	FNavigationElementHandle GetNavigationElementHandleForUObject(const UObject* NavRelevantObject);
+
+	/**
+	 * Returns a shared pointer to the FNavigationElement created for a registered UObject.
+	 * @param NavRelevantObject The UObject for which the FNavigationElement must be retrieved.
+	 * @return A valid shared pointer if the specified UObject is registered and an associated FNavigationElement was created,
+	 * null pointer otherwise.
+	 */
+	TSharedPtr<const FNavigationElement> GetNavigationElementForUObject(const UObject* NavRelevantObject);
+
+	/**
+	 * Returns a shared pointer to the FNavigationElement created for a registered UObject.
+	 * This method will not handle other navigation structure updates (e.g., NavigationOctree)
+	 * and should only be used by the NavigationSystem.
+	 * @param NavRelevantInterface The interface to use to update an existing element
+	 * @param NavRelevantObject The UObject for which the FNavigationElement must be retrieved.
+	 * @return A valid shared pointer if the specified UObject is registered and an associated FNavigationElement was created,
+	 * null pointer otherwise.
+	 */
+	TSharedPtr<const FNavigationElement> UpdateNavigationElementForUObject(const INavRelevantInterface& NavRelevantInterface, const UObject& NavRelevantObject);
+
+	/**
+	 * Returns a shared pointer to the FNavigationElement associated with the provided handle.
+	 * @param Handle The handle for which the FNavigationElement must be retrieved.
+	 * @return A valid shared pointer if a FNavigationElement associated with the provided handle was registered, null pointer otherwise.
+	 */
+	TSharedPtr<const FNavigationElement> GetNavigationElementForHandle(FNavigationElementHandle Handle) const;
 
 	/**
 	 * Adds the provided interface to the list of registered custom navigation links.
@@ -69,23 +133,64 @@ public:
 		return CustomLinkObjects;
 	}
 
-	/** Delegate executed when a navigation relevant object is registered to the repository. */
-	FOnNavRelevantObjectRegistrationEvent OnNavRelevantObjectRegistered;
+	/** Returns the number of navigation elements registered in the repository */
+	int32 GetNumRegisteredElements() const
+	{
+		return NavRelevantElements.Num();
+	}
 
-	/** Delegate executed when a navigation relevant object is unregistered from the repository. */
-	FOnNavRelevantObjectRegistrationEvent OnNavRelevantObjectUnregistered;
+	/** Returns the number of UObject registered in the repository for which a FNavigationElement has been created and registered. */
+	int32 GetNumRegisteredUObjects() const
+	{
+		return ObjectsToHandleMap.Num();
+	}
+
+	/** Returns the number of UObjects implementing INavLinkCustomInterface registered in the repository */
+	int32 GetNumRegisteredCustomLinks() const
+	{
+		return CustomLinkObjects.Num();
+	}
+
+	/** Delegate executed when a navigation element is added in the repository. */
+	FOnNavigationElementRegistrationEvent OnNavigationElementAddedDelegate;
+
+	/** Delegate executed when a navigation element is removed from the repository. */
+	FOnNavigationElementRegistrationEvent OnNavigationElementRemovedDelegate;
 
-	/** Delegate executed when a custom navigation link is registered to the repository. */
+	/** Delegate executed when a custom navigation link is registered with the repository. */
 	FOnCustomNavLinkObjectRegistrationEvent OnCustomNavLinkObjectRegistered;
 
-	/** Delegate executed when a custom navigation link is unregistered from the repository. */
+	/** Delegate executed when a custom navigation link is unregistered with the repository. */
 	FOnCustomNavLinkObjectRegistrationEvent OnCustomNavLinkObjectUnregistered;
 
 private:
 
-	/** List of registered navigation relevant objects. */
-	TArray<TWeakInterfacePtr<INavRelevantInterface>> NavRelevantObjects;
+	TSharedPtr<const FNavigationElement> RegisterNavRelevantObjectInternal(const INavRelevantInterface& NavRelevantInterface, const UObject& NavRelevantObject, ENotifyOnSuccess NotifyOnSuccess);
+
+	/** For legacy object registration path (i.e., Actor/ActorComponent) */
+	TMap<FObjectKey, FNavigationElementHandle> ObjectsToHandleMap;
+
+	/** List of registered navigation elements. */
+	TMap<FNavigationElementHandle, TSharedPtr<const FNavigationElement>> NavRelevantElements;
 
 	/** List of registered custom navigation link objects. */
 	TArray<TWeakInterfacePtr<INavLinkCustomInterface>> CustomLinkObjects;
+
+	/** Multi thread access detector used to validate accesses to the maps of registered UObjects and FNavigationElement */
+	UE_MT_DECLARE_RW_ACCESS_DETECTOR(NavElementAccessDetector);
+
+	//----------------------------------------------------------------------//
+	// DEPRECATED
+	//----------------------------------------------------------------------//
+public:
+
+#if WITH_EDITORONLY_DATA
+PRAGMA_DISABLE_DEPRECATION_WARNINGS
+	UE_DEPRECATED(5.5, "Use OnNavigationElementAddedDelegate instead.")
+	FOnNavRelevantObjectRegistrationEvent OnNavRelevantObjectRegistered;
+
+	UE_DEPRECATED(5.5, "Use OnNavigationElementRemovedDelegate instead.")
+	FOnNavRelevantObjectRegistrationEvent OnNavRelevantObjectUnregistered;
+PRAGMA_ENABLE_DEPRECATION_WARNINGS
+#endif // WITH_EDITORONLY_DATA
 };
diff --git a/Engine/Source/Runtime/NavigationSystem/Private/NavigationOctree.cpp b/Engine/Source/Runtime/NavigationSystem/Private/NavigationOctree.cpp
--- a/Engine/Source/Runtime/NavigationSystem/Private/NavigationOctree.cpp
+++ b/Engine/Source/Runtime/NavigationSystem/Private/NavigationOctree.cpp
@@ -2,12 +2,29 @@
 
 #include "NavigationOctree.h"
 #include "AI/Navigation/NavRelevantInterface.h"
+#include "AI/Navigation/NavigationElement.h"
 #include "Interfaces/Interface_AsyncCompilation.h"
 #include "NavigationSystem.h"
 #include "UObject/Package.h"
 
 LLM_DEFINE_TAG(NavigationOctree);
 
+namespace UE::NavigationOctree::Private
+{
+#if !UE_BUILD_SHIPPING
+bool bValidateConsistencyWhenAddingNode = false;
+
+FAutoConsoleVariableRef ConsoleVariables[] =
+{
+	FAutoConsoleVariableRef(
+		TEXT("ai.debug.nav.validateConsistencyWhenAddingOctreeNode"),
+		bValidateConsistencyWhenAddingNode,
+		TEXT("Used to validate that registered FNavigationElement matches the values "
+			"returned by INavRelevantInterface when processing pending updates to add elements to the octree."))
+};
+#endif // !UE_BUILD_SHIPPING
+
+} // UE::NavigationOctree::Private
 
 //----------------------------------------------------------------------//
 // FNavigationOctree
@@ -17,20 +34,22 @@ FNavigationOctree::FNavigationOctree(const FVector& Origin, FVector::FReal Radiu
 	, DefaultGeometryGatheringMode(ENavDataGatheringMode::Instant)
 	, bGatherGeometry(false)
 	, NodesMemory(0)
-#if !UE_BUILD_SHIPPING	
+#if !UE_BUILD_SHIPPING
 	, GatheringNavModifiersTimeLimitWarning(-1.0f)
-#endif // !UE_BUILD_SHIPPING	
+#endif // !UE_BUILD_SHIPPING
 {
 	INC_DWORD_STAT_BY( STAT_NavigationMemory, sizeof(*this) );
 }
 
+PRAGMA_DISABLE_DEPRECATION_WARNINGS
 FNavigationOctree::~FNavigationOctree()
 {
 	DEC_DWORD_STAT_BY( STAT_NavigationMemory, sizeof(*this) );
 	DEC_MEMORY_STAT_BY(STAT_Navigation_CollisionTreeMemory, NodesMemory);
 	
-	ObjectToOctreeId.Empty();
+	ElementToOctreeId.Empty();
 }
+PRAGMA_ENABLE_DEPRECATION_WARNINGS
 
 void FNavigationOctree::SetDataGatheringMode(ENavDataGatheringModeConfig Mode)
 {
@@ -47,13 +66,6 @@ void FNavigationOctree::DemandLazyDataGathering(FNavigationRelevantData& Element
 {
 	LLM_SCOPE_BYTAG(NavigationOctree);
 
-	UObject* const NavRelevantObject = ElementData.GetOwner();
-	INavRelevantInterface* const NavRelevantInterface = Cast<INavRelevantInterface>(NavRelevantObject);
-	if (NavRelevantInterface == nullptr)
-	{
-		return;
-	}
-
 	bool bShrink = false;
 	const int32 OrgElementMemory = IntCastChecked<int32>(ElementData.GetGeometryAllocatedSize());
 
@@ -61,7 +73,7 @@ void FNavigationOctree::DemandLazyDataGathering(FNavigationRelevantData& Element
 	{
 		QUICK_SCOPE_CYCLE_COUNTER(STAT_RecastNavMeshGenerator_LazyGeometryExport);
 
-		NavRelevantGeometryExportDelegate.ExecuteIfBound(*NavRelevantInterface, ElementData);
+		GeometryExportDelegate.ExecuteIfBound(ElementData.SourceElement.Get(), ElementData);
 		bShrink = true;
 
 		// mark this element as no longer needing geometry gathering
@@ -77,23 +89,21 @@ void FNavigationOctree::DemandLazyDataGathering(FNavigationRelevantData& Element
 		const double StartTime = bCanOutputDurationWarning ? FPlatformTime::Seconds() : 0.0f;
 #endif //!UE_BUILD_SHIPPING
 
-		NavRelevantInterface->GetNavigationData(ElementData);
+		ElementData.SourceElement->NavigationDataExportDelegate.ExecuteIfBound(ElementData.SourceElement.Get(), ElementData);
 		ElementData.bPendingLazyModifiersGathering = false;
 		bShrink = true;
 
 #if !UE_BUILD_SHIPPING
-		// If GatheringNavModifiersWarningLimitTime is positive, it will print a Warning if the time taken to call GetNavigationData is more than GatheringNavModifiersWarningLimitTime			
+		// If GatheringNavModifiersWarningLimitTime is positive, it will print a Warning if the time taken to call GetNavigationData is more than GatheringNavModifiersWarningLimitTime
 		if (bCanOutputDurationWarning)
 		{
 			const double DeltaTime = FPlatformTime::Seconds() - StartTime;
 			if (DeltaTime > GatheringNavModifiersTimeLimitWarning)
 			{
-				UE_LOG(LogNavigation, Warning, TEXT("The time (%f sec) for gathering navigation data on an INavRelevantInterface navigation element exceeded the time limit (%f sec) | NavElement = %s | NavElement's Owner = %s | Level = %s"),
+				UE_LOG(LogNavigation, Warning, TEXT("The time (%f sec) for gathering navigation data on a navigation element exceeded the time limit (%f sec) | Element = %s"),
 					DeltaTime,
 					GatheringNavModifiersTimeLimitWarning,
-					*GetNameSafe(NavRelevantObject),
-					*GetNameSafe(NavRelevantObject->GetOuter()),
-					*GetNameSafe(NavRelevantObject->GetOutermost()));
+					*ElementData.SourceElement->GetName());
 			}
 		}
 #endif //!UE_BUILD_SHIPPING
@@ -112,126 +122,167 @@ void FNavigationOctree::DemandLazyDataGathering(FNavigationRelevantData& Element
 	INC_MEMORY_STAT_BY(STAT_Navigation_CollisionTreeMemory, ElementMemoryChange);
 }
 
+// Deprecated
 void FNavigationOctree::DemandChildLazyDataGathering(FNavigationRelevantData& ElementData, INavRelevantInterface& ChildNavInterface)
+{
+	const TSharedRef<const FNavigationElement> TmpElement = FNavigationElement::CreateFromNavRelevantInterface(ChildNavInterface);
+	DemandChildLazyDataGathering(ElementData, TmpElement.Get());
+}
+
+void FNavigationOctree::DemandChildLazyDataGathering(FNavigationRelevantData& ElementData, const FNavigationElement& ChildElement) const
 {
 	LLM_SCOPE_BYTAG(NavigationOctree);
 
-	if (IsLazyGathering(ChildNavInterface))
+	if (IsLazyGathering(ChildElement))
 	{
-		ChildNavInterface.GetNavigationData(ElementData);
-		ElementData.ValidateAndShrink();
+		if (ChildElement.NavigationDataExportDelegate.ExecuteIfBound(ChildElement, ElementData))
+		{
+			ElementData.ValidateAndShrink();
+		}
 	}
 }
 
-#if !UE_BUILD_SHIPPING	
+#if !UE_BUILD_SHIPPING
 void FNavigationOctree::SetGatheringNavModifiersTimeLimitWarning(const float Threshold)
 {
 	GatheringNavModifiersTimeLimitWarning = Threshold;
 }
-#endif // !UE_BUILD_SHIPPING	
+#endif // !UE_BUILD_SHIPPING
 
+// Deprecated
 bool FNavigationOctree::IsLazyGathering(const INavRelevantInterface& ChildNavInterface) const
 {
-	const ENavDataGatheringMode GatheringMode = ChildNavInterface.GetGeometryGatheringMode();
+	const TSharedRef<const FNavigationElement> TmpElement = FNavigationElement::CreateFromNavRelevantInterface(ChildNavInterface);
+	return IsLazyGathering(TmpElement.Get());
+}
+
+bool FNavigationOctree::IsLazyGathering(const FNavigationElement& NavigationElement) const
+{
+	const ENavDataGatheringMode GatheringMode = NavigationElement.GetGeometryGatheringMode();
 	const bool bDoInstantGathering = ((GatheringMode == ENavDataGatheringMode::Default && DefaultGeometryGatheringMode == ENavDataGatheringMode::Instant)
 		|| GatheringMode == ENavDataGatheringMode::Instant);
 
 	return !bDoInstantGathering;
 }
 
+// Deprecated
 void FNavigationOctree::AddNode(UObject* ElementOb, INavRelevantInterface* NavElement, const FBox& Bounds, FNavigationOctreeElement& Element)
 {
-	LLM_SCOPE_BYTAG(NavigationOctree);
-
-	if (UNLIKELY(!Bounds.IsValid || Bounds.GetSize().IsNearlyZero())) 
+	if (NavElement)
 	{
-		UE_LOG(LogNavigation, Warning, TEXT("%hs: %s bounds, ignoring %s."), __FUNCTION__, !Bounds.IsValid ? TEXT("Invalid") : TEXT("Empty"), *GetFullNameSafe(ElementOb));
-		return;
+		AddNode(Bounds, Element);
 	}
-	
-	Element.Bounds = Bounds;
+}
 
-	if (NavElement)
-	{
-		checkf(ElementOb, TEXT("We assume NavElement is ElementOb already cast"));
+void FNavigationOctree::AddNode(const FBox& Bounds, FNavigationOctreeElement& OctreeElement)
+{
+	LLM_SCOPE_BYTAG(NavigationOctree);
 
-		Element.Data->bShouldSkipDirtyAreaOnAddOrRemove = NavElement && NavElement->ShouldSkipDirtyAreaOnAddOrRemove();
+	const TSharedRef<const FNavigationElement>& ElementRef = OctreeElement.Data->SourceElement;
+	const FNavigationElement& SourceElement = ElementRef.Get();
 
-		const bool bDoInstantGathering = !IsLazyGathering(*NavElement);
+	UE_LOG(LogNavigation, VeryVerbose, TEXT("%hs: '%s' bounds: [%s]"), __FUNCTION__, *SourceElement.GetName(), *Bounds.ToString());
 
-		if (bGatherGeometry)
+#if !UE_BUILD_SHIPPING
+	if (UE::NavigationOctree::Private::bValidateConsistencyWhenAddingNode)
+	{
+		if (const INavRelevantInterface* NavRelevantInterface = Cast<INavRelevantInterface>(SourceElement.GetWeakUObject().Get()))
 		{
-			bool bIsCompiling = false;
-#if WITH_EDITOR
-			const IInterface_AsyncCompilation* AsyncCompiledObject = Cast<IInterface_AsyncCompilation>(ElementOb);
-			bIsCompiling = AsyncCompiledObject && AsyncCompiledObject->IsCompiling();
-
-			UE_CLOG(bIsCompiling, LogNavigation, Warning, TEXT("%hs: Objects %s should not be considered relevant to navigation until associated asset compilation is completed."),
-				__FUNCTION__, *GetFullNameSafe(ElementOb));
-#endif
-
-			// Skip custom navigation export during async compilation, the node will be invalidated once
-			// compilation finishes.
-			if (!bIsCompiling)
+			// The following validations are used to detect scenarios where properties of a NavigationElement created from
+			// a navigation relevant UObject (implementing INavRelevantInterface and returning 'true' to IsNavigationRelevant())
+			// differs from the values that are provided by that same UObject when the pending registration of the element is processed.
+			// These indicate that the values were not up-to-date when the element was added to a pending FNavigationDirtyElement,
+			// or during that frame since an update would have refreshed the pending DirtyElement.
+			if (const FBox NewBounds = NavRelevantInterface->GetNavigationBounds(); !Bounds.Equals(NewBounds))
+			{
+				UE_LOG(LogNavigation, Warning, TEXT("%hs: '%s' bounds changed between element's creation and its addition to the octree: [%s] --> [%s]"),
+					__FUNCTION__, *SourceElement.GetName(), *Bounds.ToString(), *NewBounds.ToString());
+			}
+
+			if (const UObject* NewParent = NavRelevantInterface->GetNavigationParent(); NewParent != SourceElement.GetNavigationParent())
 			{
-				if (bDoInstantGathering)
-				{
-					NavRelevantGeometryExportDelegate.ExecuteIfBound(*NavElement, *Element.Data);
-				}
-				else
-				{
-					Element.Data->bPendingLazyGeometryGathering = true;
-					Element.Data->bSupportsGatheringGeometrySlices = NavElement && NavElement->SupportsGatheringGeometrySlices();
-				}
+				UE_LOG(LogNavigation, Warning, TEXT("%hs: '%s' parent changed between element's creation and its addition to the octree: [%s] --> [%s]"),
+					__FUNCTION__, *SourceElement.GetName(), *GetNameSafe(SourceElement.GetNavigationParent().Get()), *GetNameSafe(NewParent));
 			}
 		}
+	}
+#endif // !UE_BUILD_SHIPPING
 
-		SCOPE_CYCLE_COUNTER(STAT_Navigation_GatheringNavigationModifiersSync);
+	if (UNLIKELY(!Bounds.IsValid || Bounds.GetSize().IsNearlyZero()))
+	{
+		UE_LOG(LogNavigation, Warning, TEXT("%hs: %s bounds, ignoring %s."), __FUNCTION__, !Bounds.IsValid ? TEXT("Invalid") : TEXT("Empty"), *SourceElement.GetFullName());
+		return;
+	}
+
+	OctreeElement.Bounds = Bounds;
+	OctreeElement.Data->bShouldSkipDirtyAreaOnAddOrRemove = !SourceElement.GetDirtyAreaOnRegistration();
+
+	const bool bDoInstantGathering = !IsLazyGathering(SourceElement);
+
+	if (bGatherGeometry)
+	{
 		if (bDoInstantGathering)
 		{
+			GeometryExportDelegate.ExecuteIfBound(SourceElement, *OctreeElement.Data);
+		}
+		else
+		{
+			OctreeElement.Data->bPendingLazyGeometryGathering = true;
+			OctreeElement.Data->bSupportsGatheringGeometrySlices = SourceElement.GeometrySliceExportDelegate.IsBound();
+		}
+	}
+
+	SCOPE_CYCLE_COUNTER(STAT_Navigation_GatheringNavigationModifiersSync);
+	if (bDoInstantGathering)
+	{
 #if !UE_BUILD_SHIPPING
-			const bool bCanOutputDurationWarning = GatheringNavModifiersTimeLimitWarning >= 0.0f;
-			const double StartTime = bCanOutputDurationWarning ? FPlatformTime::Seconds() : 0.0f;
+		const bool bCanOutputDurationWarning = GatheringNavModifiersTimeLimitWarning >= 0.0f;
+		const double StartTime = bCanOutputDurationWarning ? FPlatformTime::Seconds() : 0.0f;
 #endif //!UE_BUILD_SHIPPING
 
-			NavElement->GetNavigationData(*Element.Data);
+		SourceElement.NavigationDataExportDelegate.ExecuteIfBound(SourceElement, *OctreeElement.Data);
 
 #if !UE_BUILD_SHIPPING
-			// If GatheringNavModifiersWarningLimitTime is positive, it will print a Warning if the time taken to call GetNavigationData is more than GatheringNavModifiersWarningLimitTime			
-			if (bCanOutputDurationWarning)
+		// If GatheringNavModifiersWarningLimitTime is positive, it will print a Warning if the time taken to call GetNavigationData is more than GatheringNavModifiersWarningLimitTime			
+		if (bCanOutputDurationWarning)
+		{
+			if (const double DeltaTime = FPlatformTime::Seconds() - StartTime; DeltaTime > GatheringNavModifiersTimeLimitWarning)
 			{
-				const double DeltaTime = FPlatformTime::Seconds() - StartTime;
-				if (DeltaTime > GatheringNavModifiersTimeLimitWarning)
-				{
-					UE_LOG(LogNavigation, Warning, TEXT("The time (%f sec) for gathering navigation data on an INavRelevantInterface navigation element exceeded the time limit (%f sec) | NavElement = %s | Element's Owner = %s | Level = %s"),
-						DeltaTime,
-						GatheringNavModifiersTimeLimitWarning,
-						*GetNameSafe(ElementOb),
-						*GetNameSafe(ElementOb->GetOuter()),
-						*GetNameSafe(ElementOb->GetOutermost()));
-				}
+				UE_LOG(LogNavigation, Warning, TEXT("The time (%f sec) for gathering navigation data on a navigation element exceeded the time limit (%f sec) | Element = %s"),
+					DeltaTime,
+					GatheringNavModifiersTimeLimitWarning,
+					*SourceElement.GetName());
 			}
-#endif //!UE_BUILD_SHIPPING
-		}
-		else
-		{
-			Element.Data->bPendingLazyModifiersGathering = true;
 		}
+#endif //!UE_BUILD_SHIPPING
+	}
+	else
+	{
+		OctreeElement.Data->bPendingLazyModifiersGathering = true;
 	}
 
 	// validate exported data
 	// shrink arrays before counting memory
 	// it will be reallocated when adding to octree and RemoveNode will have different value returned by GetAllocatedSize()
-	Element.ValidateAndShrink();
+	OctreeElement.ValidateAndShrink();
 
-	const int32 ElementMemory = Element.GetAllocatedSize();
+	const int32 ElementMemory = OctreeElement.GetAllocatedSize();
 	NodesMemory += ElementMemory;
 	INC_MEMORY_STAT_BY(STAT_Navigation_CollisionTreeMemory, ElementMemory);
 
-	AddElement(Element);
+	AddElement(OctreeElement);
 }
 
+// Deprecated
 void FNavigationOctree::AppendToNode(const FOctreeElementId2& Id, INavRelevantInterface* NavElement, const FBox& Bounds, FNavigationOctreeElement& Element)
+{
+	if (NavElement)
+	{
+		AppendToNode(Id, FNavigationElement::CreateFromNavRelevantInterface(*NavElement), Bounds, Element);
+	}
+}
+
+void FNavigationOctree::AppendToNode(const FOctreeElementId2& Id, const TSharedRef<const FNavigationElement>& ElementRef, const FBox& Bounds, FNavigationOctreeElement& Element)
 {
 	LLM_SCOPE_BYTAG(NavigationOctree);
 
@@ -240,19 +291,14 @@ void FNavigationOctree::AppendToNode(const FOctreeElementId2& Id, INavRelevantIn
 	Element = OrgData;
 	Element.Bounds = Bounds + OrgData.Bounds.GetBox();
 
-	if (NavElement)
+	SCOPE_CYCLE_COUNTER(STAT_Navigation_GatheringNavigationModifiersSync);
+	if (IsLazyGathering(ElementRef.Get()))
 	{
-		SCOPE_CYCLE_COUNTER(STAT_Navigation_GatheringNavigationModifiersSync);
-		const bool bDoInstantGathering = !IsLazyGathering(*NavElement);
-
-		if (bDoInstantGathering)
-		{
-			NavElement->GetNavigationData(*Element.Data);
-		}
-		else
-		{
-			Element.Data->bPendingChildLazyModifiersGathering = true;
-		}
+		Element.Data->bPendingChildLazyModifiersGathering = true;
+	}
+	else
+	{
+		ElementRef->NavigationDataExportDelegate.ExecuteIfBound(ElementRef.Get(), *Element.Data);
 	}
 
 	// validate exported data
@@ -291,19 +337,17 @@ void FNavigationOctree::RemoveNode(const FOctreeElementId2& Id)
 
 const FNavigationRelevantData* FNavigationOctree::GetDataForID(const FOctreeElementId2& Id) const
 {
-	if (Id.IsValidId() == false)
-	{
-		return nullptr;
-	}
-
-	const FNavigationOctreeElement& OctreeElement = GetElementById(Id);
+	return Id.IsValidId() ? &*GetElementById(Id).Data : nullptr;
+}
 
-	return &*OctreeElement.Data;
+FNavigationRelevantData* FNavigationOctree::GetMutableDataForID(const FOctreeElementId2& Id)
+{
+	return Id.IsValidId() ? &*GetElementById(Id).Data : nullptr;
 }
 
-void FNavigationOctree::SetElementIdImpl(const uint32 OwnerUniqueId, FOctreeElementId2 Id)
+void FNavigationOctree::SetElementIdImpl(const FNavigationElementHandle ElementHandle, const FOctreeElementId2 Id)
 {
-	ObjectToOctreeId.Add(OwnerUniqueId, Id);
+	ElementToOctreeId.Add(ElementHandle, Id);
 }
 
 //----------------------------------------------------------------------//
@@ -312,7 +356,7 @@ void FNavigationOctree::SetElementIdImpl(const uint32 OwnerUniqueId, FOctreeElem
 #if NAVSYS_DEBUG
 FORCENOINLINE
 #endif // NAVSYS_DEBUG
-void FNavigationOctreeSemantics::SetElementId(FNavigationOctreeSemantics::FOctree& OctreeOwner, const FNavigationOctreeElement& Element, FOctreeElementId2 Id)
+void FNavigationOctreeSemantics::SetElementId(FOctree& OctreeOwner, const FNavigationOctreeElement& Element, const FOctreeElementId2 Id)
 {
-	((FNavigationOctree&)OctreeOwner).SetElementIdImpl(Element.OwnerUniqueId, Id);
-}
+	static_cast<FNavigationOctree&>(OctreeOwner).SetElementIdImpl(Element.Data.Get().SourceElement.Get().GetHandle(), Id);
+}
\ No newline at end of file
diff --git a/Engine/Source/Runtime/NavigationSystem/Private/NavigationOctreeController.cpp b/Engine/Source/Runtime/NavigationSystem/Private/NavigationOctreeController.cpp
--- a/Engine/Source/Runtime/NavigationSystem/Private/NavigationOctreeController.cpp
+++ b/Engine/Source/Runtime/NavigationSystem/Private/NavigationOctreeController.cpp
@@ -1,7 +1,6 @@
 // Copyright Epic Games, Inc. All Rights Reserved.
 
 #include "NavigationOctreeController.h"
-#include "AI/Navigation/NavRelevantInterface.h"
 #include "NavigationSystem.h"
 
 
@@ -13,9 +12,14 @@ void FNavigationOctreeController::Reset()
 	if (NavOctree.IsValid())
 	{
 		NavOctree->Destroy();
-		NavOctree = NULL;
+		NavOctree = nullptr;
 	}
-	PendingOctreeUpdates.Empty(32);
+	PendingUpdates.Empty(32);
+}
+
+bool FNavigationOctreeController::HasPendingUpdateForElement(const FNavigationElementHandle Element) const
+{ 
+	return PendingUpdates.Contains(Element);
 }
 
 void FNavigationOctreeController::SetNavigableGeometryStoringMode(FNavigationOctree::ENavGeometryStoringMode NavGeometryMode)
@@ -24,9 +28,9 @@ void FNavigationOctreeController::SetNavigableGeometryStoringMode(FNavigationOct
 	NavOctree->SetNavigableGeometryStoringMode(NavGeometryMode);
 }
 
-bool FNavigationOctreeController::GetNavOctreeElementData(const UObject& NodeOwner, ENavigationDirtyFlag& OutDirtyFlags, FBox& OutDirtyBounds)
+bool FNavigationOctreeController::GetNavOctreeElementData(const FNavigationElementHandle Element, ENavigationDirtyFlag& OutDirtyFlags, FBox& OutDirtyBounds)
 {
-	const FOctreeElementId2* ElementId = GetObjectsNavOctreeId(NodeOwner);
+	const FOctreeElementId2* ElementId = GetNavOctreeIdForElement(Element);
 	if (ElementId != nullptr && IsValidElement(*ElementId))
 	{
 		// mark area occupied by given actor as dirty
@@ -43,16 +47,20 @@ bool FNavigationOctreeController::GetNavOctreeElementData(const UObject& NodeOwn
 bool FNavigationOctreeController::GetNavOctreeElementData(const UObject& NodeOwner, int32& DirtyFlags, FBox& DirtyBounds)
 {
 	ENavigationDirtyFlag TmpDirtyFlags = ENavigationDirtyFlag::None;
-	const bool bSuccess = GetNavOctreeElementData(NodeOwner, TmpDirtyFlags, DirtyBounds);
+	const bool bSuccess = GetNavOctreeElementData(FNavigationElementHandle(&NodeOwner), TmpDirtyFlags, DirtyBounds);
 	DirtyFlags = static_cast<int32>(TmpDirtyFlags);
 	return bSuccess;
 }
 
+// Deprecated
 const FNavigationRelevantData* FNavigationOctreeController::GetDataForObject(const UObject& Object) const
 {
-	const FOctreeElementId2* ElementId = GetObjectsNavOctreeId(Object);
+	return GetDataForElement(FNavigationElementHandle(&Object));
+}
 
-	if (ElementId != nullptr && IsValidElement(*ElementId))
+const FNavigationRelevantData* FNavigationOctreeController::GetDataForElement(const FNavigationElementHandle Element) const
+{
+	if (const FOctreeElementId2* ElementId = GetNavOctreeIdForElement(Element); IsValidElement(ElementId))
 	{
 		return NavOctree->GetDataForID(*ElementId);
 	}
@@ -60,7 +68,54 @@ const FNavigationRelevantData* FNavigationOctreeController::GetDataForObject(con
 	return nullptr;
 }
 
+// Deprecated
 FNavigationRelevantData* FNavigationOctreeController::GetMutableDataForObject(const UObject& Object)
 {
-	return const_cast<FNavigationRelevantData*>(GetDataForObject(Object));
-}
\ No newline at end of file
+	return GetMutableDataForElement(FNavigationElementHandle(&Object));
+}
+
+FNavigationRelevantData* FNavigationOctreeController::GetMutableDataForElement(const FNavigationElementHandle Element)
+{
+	if (const FOctreeElementId2* ElementId = GetNavOctreeIdForElement(Element); IsValidElement(ElementId))
+	{
+		return NavOctree->GetMutableDataForID(*ElementId);
+	}
+
+	return nullptr;
+}
+
+//----------------------------------------------------------------------//
+// Deprecated methods
+//----------------------------------------------------------------------//
+
+PRAGMA_DISABLE_DEPRECATION_WARNINGS
+
+// Deprecated
+bool FNavigationOctreeController::HasPendingObjectNavOctreeId(UObject& Object) const
+{
+	return HasPendingUpdateForElement(FNavigationElementHandle(&Object));
+}
+
+// Deprecated
+bool FNavigationOctreeController::HasObjectsNavOctreeId(const UObject& Object) const
+{
+	return HasElementNavOctreeId(FNavigationElementHandle(&Object));
+}
+
+// Deprecated
+const FOctreeElementId2* FNavigationOctreeController::GetObjectsNavOctreeId(const UObject& Object) const
+{
+	return GetNavOctreeIdForElement(FNavigationElementHandle(&Object));
+}
+
+// Deprecated
+void FNavigationOctreeController::RemoveObjectsNavOctreeId(const UObject& Object)
+{
+	const FNavigationElementHandle ElementHandle(&Object);
+	if (const FOctreeElementId2* ElementId = GetNavOctreeIdForElement(ElementHandle); IsValidElement(ElementId))
+	{
+		RemoveNode(*ElementId, ElementHandle);
+	}	
+}
+
+PRAGMA_ENABLE_DEPRECATION_WARNINGS
\ No newline at end of file
diff --git a/Engine/Source/Runtime/NavigationSystem/Private/NavigationSystem.cpp b/Engine/Source/Runtime/NavigationSystem/Private/NavigationSystem.cpp
--- a/Engine/Source/Runtime/NavigationSystem/Private/NavigationSystem.cpp
+++ b/Engine/Source/Runtime/NavigationSystem/Private/NavigationSystem.cpp
@@ -9,6 +9,7 @@
 #include "AI/Navigation/NavigationDirtyElement.h"
 #include "AI/Navigation/NavigationInvokerInterface.h"
 #include "AI/Navigation/NavigationInvokerPriority.h"
+#include "AI/Navigation/NavigationElement.h"
 #include "AI/Navigation/NavRelevantInterface.h"
 #include "AI/NavigationModifier.h"
 #include "Components/PrimitiveComponent.h"
@@ -48,12 +49,12 @@
 #endif // WITH_RECAST
 
 #if WITH_EDITOR
-#include "WorldPartition/WorldPartition.h"
-#include "Misc/MessageDialog.h"
-#include "LevelEditor.h"
-#include "EditorModes.h"
 #include "EditorModeManager.h"
-#endif
+#include "EditorModes.h"
+#include "LevelEditor.h"
+#include "Misc/MessageDialog.h"
+#include "WorldPartition/WorldPartition.h"
+#endif // WITH_EDITOR
 
 
 #include UE_INLINE_GENERATED_CPP_BY_NAME(NavigationSystem)
@@ -68,7 +69,7 @@ DECLARE_CYCLE_STAT(TEXT("Nav Tick: mark dirty"), STAT_Navigation_TickMarkDirty,
 DECLARE_CYCLE_STAT(TEXT("Nav Tick: async build"), STAT_Navigation_TickAsyncBuild, STATGROUP_Navigation);
 DECLARE_CYCLE_STAT(TEXT("Nav Tick: dispatch async pathfinding results"), STAT_Navigation_DispatchAsyncPathfindingResults, STATGROUP_Navigation);
 DECLARE_CYCLE_STAT(TEXT("Nav Tick: async pathfinding"), STAT_Navigation_TickAsyncPathfinding, STATGROUP_Navigation);
-DECLARE_CYCLE_STAT(TEXT("Debug NavOctree Time"), STAT_DebugNavOctree, STATGROUP_Navigation);
+DECLARE_CYCLE_STAT(TEXT("NavOctree bookkeeping"), STAT_NavOctreeBookkeeping, STATGROUP_Navigation);
 
 //----------------------------------------------------------------------//
 // Stats
@@ -211,6 +212,41 @@ static FAutoConsoleCommandWithWorldArgsAndOutputDevice CmdNavDirtyAreaAroundPlay
 		}
 	));
 
+
+static FAutoConsoleCommandWithWorldArgsAndOutputDevice CmdDumpOctreeElements(
+	TEXT("ai.debug.nav.DumpOctreeElements"),
+	TEXT("Iterates through all nodes of the navigation octree and log details about each element to the output device."),
+	FConsoleCommandWithWorldArgsAndOutputDeviceDelegate::CreateLambda([](const TArray<FString>& Args, const UWorld* World, FOutputDevice& OutputDevice)
+		{
+			if (const UNavigationSystemV1* NavSys = FNavigationSystem::GetCurrent<UNavigationSystemV1>(World))
+			{
+				if (const FNavigationOctree* Octree = NavSys->GetNavOctree())
+				{
+					Octree->FindNodesWithPredicate(
+						[](FNavigationOctree::FNodeIndex /*ParentNodeIndex*/, FNavigationOctree::FNodeIndex /*NodeIndex*/, const FBoxCenterAndExtent&) { return true; },
+						[&](FNavigationOctree::FNodeIndex /*ParentNodeIndex*/, const FNavigationOctree::FNodeIndex NodeIndex, const FBoxCenterAndExtent&)
+						{
+							for (const FNavigationOctreeElement& OctreeElement : Octree->GetElementsForNode(NodeIndex))
+							{
+								OutputDevice.Logf(ELogVerbosity::Type::Log, TEXT("%s bounds: [%s] parent:'%s'"),
+									*OctreeElement.GetSourceElement()->GetPathName(),
+									*OctreeElement.Bounds.ToString(),
+									*GetNameSafe(OctreeElement.GetSourceElement().Get().GetNavigationParent().Get()));
+							};
+						});
+				}
+				else
+				{
+					OutputDevice.Log(ELogVerbosity::Error, TEXT("Octree not used in the current configuration"));
+				}
+			}
+			else
+			{
+				OutputDevice.Log(ELogVerbosity::Error, TEXT("Command failed since it was unable to find the navigation system"));
+			}
+		}
+	));
+
 	const FNavDataConfig& GetFallbackNavDataConfig()
 	{
 		static FNavDataConfig FallbackNavDataConfig(FNavigationSystem::FallbackAgentRadius, FNavigationSystem::FallbackAgentHeight);
@@ -253,7 +289,7 @@ static FAutoConsoleCommandWithWorldArgsAndOutputDevice CmdNavDirtyAreaAroundPlay
 			ActorComp->SetCanEverAffectNavigation(false);
 		}
 	}
-}
+} // FNavigationSystem
 
 namespace NavigationDebugDrawing
 {
@@ -670,6 +706,19 @@ UNavigationSystemV1::UNavigationSystemV1(const FObjectInitializer& ObjectInitial
 	{
 		FDelegatesInitializer()
 		{
+			UNavigationSystemBase::GetAddNavigationElementDelegate().BindStatic(&UNavigationSystemV1::AddNavigationElement);
+			UNavigationSystemBase::GetRemoveNavigationElementDelegate().BindStatic(&UNavigationSystemV1::RemoveNavigationElement);
+			UNavigationSystemBase::GetUpdateNavigationElementDelegate().BindStatic(&UNavigationSystemV1::OnNavigationElementUpdated);
+
+			UNavigationSystemBase::GetUpdateNavigationElementBoundsDelegate().BindLambda(
+				[](UWorld* World, FNavigationElementHandle Handle, const FBox& NewBounds, const TConstArrayView<FBox> DirtyAreas)
+				{
+					if (UNavigationSystemV1* NavSys = FNavigationSystem::GetCurrent<UNavigationSystemV1>(World))
+					{
+						NavSys->UpdateNavOctreeElementBounds(Handle, NewBounds, DirtyAreas);
+					}
+				});
+
 			UNavigationSystemBase::RegisterNavRelevantObjectDelegate().BindLambda([](UObject& Object) { UNavigationSystemV1::OnNavRelevantObjectRegistered(Object); });
 			UNavigationSystemBase::UpdateNavRelevantObjectDelegate().BindStatic(&UNavigationSystemV1::UpdateNavRelevantObjectInNavOctree);
 			UNavigationSystemBase::UnregisterNavRelevantObjectDelegate().BindLambda([](UObject& Object) { UNavigationSystemV1::OnNavRelevantObjectUnregistered(Object); });
@@ -677,7 +726,7 @@ UNavigationSystemV1::UNavigationSystemV1(const FObjectInitializer& ObjectInitial
 				{
 					if (UNavigationSystemV1* NavSys = FNavigationSystem::GetCurrent<UNavigationSystemV1>(Object.GetWorld()))
 					{
-						NavSys->UpdateNavOctreeElementBounds(Object, NewBounds, DirtyAreas);
+						NavSys->UpdateNavOctreeElementBounds(FNavigationElementHandle(&Object), NewBounds, DirtyAreas);
 					}
 				});
 
@@ -726,18 +775,21 @@ UNavigationSystemV1::UNavigationSystemV1(const FObjectInitializer& ObjectInitial
 			UNavigationSystemBase::RegisterComponentDelegate().BindLambda([](UActorComponent& Comp) { UNavigationSystemV1::RegisterComponent(&Comp); });
 			UNavigationSystemBase::UnregisterComponentDelegate().BindLambda([](UActorComponent& Comp) { UNavigationSystemV1::UnregisterComponent(&Comp); });
 			UNavigationSystemBase::RemoveActorDataDelegate().BindLambda([](AActor& Actor) { UNavigationSystemV1::ClearNavOctreeAll(&Actor); });
-			UNavigationSystemBase::HasComponentDataDelegate().BindLambda([](UActorComponent& Comp) {
+			UNavigationSystemBase::HasComponentDataDelegate().BindLambda([](UActorComponent& Comp)
+			{
 				UNavigationSystemV1* NavSys = FNavigationSystem::GetCurrent<UNavigationSystemV1>(Comp.GetWorld());
-				return (NavSys && (NavSys->GetObjectsNavOctreeId(Comp) || NavSys->HasPendingObjectNavOctreeId(&Comp)));
+				const FNavigationElementHandle Element(&Comp);
+				return (NavSys && (NavSys->GetNavOctreeIdForElement(Element) || NavSys->HasPendingUpdateForElement(Element)));
 			});
 			UNavigationSystemBase::GetDefaultSupportedAgentDelegate().BindStatic(&UNavigationSystemV1::GetDefaultSupportedAgent);
 			UNavigationSystemBase::GetBiggestSupportedAgentDelegate().BindStatic(&UNavigationSystemV1::GetBiggestSupportedAgent);
 			UNavigationSystemBase::UpdateActorAndComponentDataDelegate().BindStatic(&UNavigationSystemV1::UpdateActorAndComponentsInNavOctree);
-			UNavigationSystemBase::OnComponentBoundsChangedDelegate().BindLambda([](UActorComponent& Comp, const FBox& NewBounds, const FBox& DirtyArea) {
+			UNavigationSystemBase::OnComponentBoundsChangedDelegate().BindLambda([](UActorComponent& Comp, const FBox& NewBounds, const FBox& DirtyArea)
+			{
 				UNavigationSystemV1* NavSys = FNavigationSystem::GetCurrent<UNavigationSystemV1>(Comp.GetWorld());
 				if (NavSys)
 				{
-					NavSys->UpdateNavOctreeElementBounds(Comp, NewBounds, {DirtyArea});
+					NavSys->UpdateNavOctreeElementBounds(FNavigationElementHandle(&Comp), NewBounds, {DirtyArea});
 				}
 			});
 			//UNavigationSystemBase::GetNavDataForPropsDelegate();
@@ -809,6 +861,9 @@ UNavigationSystemV1::UNavigationSystemV1(const FObjectInitializer& ObjectInitial
 		}
 	};
 	static FDelegatesInitializer DelegatesInitializer;
+
+	// NOP line to silence code analysis warning: "Local variable 'DelegatesInitializer' is never used"
+	(void)DelegatesInitializer;
 	
 	// Set to the ai module's crowd manager, this module may not exist at spawn time but then it will just fail to load
 	CrowdManagerClass = FSoftObjectPath(TEXT("/Script/AIModule.CrowdManager"));
@@ -833,6 +888,12 @@ UNavigationSystemV1::UNavigationSystemV1(const FObjectInitializer& ObjectInitial
 #if !UE_BUILD_SHIPPING
 		FCoreDelegates::OnGetOnScreenMessages.AddUObject(this, &UNavigationSystemV1::GetOnScreenMessages);
 #endif // !UE_BUILD_SHIPPING
+
+		if (const UWorld* World = UNavigationSystemV1::GetWorld())
+		{
+			Repository = World->GetSubsystem<UNavigationObjectRepository>();
+		}
+		checkf(IsRunningCookCommandlet() || Repository, TEXT("UNavigationObjectRepository is mandatory for navigation system operations."));
 	}
 	else if (GetClass() == UNavigationSystemV1::StaticClass())
 	{
@@ -840,8 +901,8 @@ UNavigationSystemV1::UNavigationSystemV1(const FObjectInitializer& ObjectInitial
 		SetDefaultObstacleArea(UNavArea_Obstacle::StaticClass());
 		
 #if WITH_RECAST
-		const FTransform RecastToUnrealTransfrom(Recast2UnrealMatrix());
-		SetCoordTransform(ENavigationCoordSystem::Navigation, ENavigationCoordSystem::Unreal, RecastToUnrealTransfrom);
+		const FTransform RecastToUnrealTransform(Recast2UnrealMatrix());
+		SetCoordTransform(ENavigationCoordSystem::Navigation, ENavigationCoordSystem::Unreal, RecastToUnrealTransform);
 #endif // WITH_RECAST
 	}
 }
@@ -855,6 +916,56 @@ UNavigationSystemV1::~UNavigationSystemV1()
 	FCoreDelegates::OnGetOnScreenMessages.RemoveAll(this);
 #endif // !UE_BUILD_SHIPPING
 }
+
+void UNavigationSystemV1::GatherDebugLabels(TArray<FString>& InOutDebugLabels) const
+{
+	if (Repository)
+	{
+		InOutDebugLabels.Add(FString::Printf(TEXT("Repository Navigation Elements count: %i"), Repository->GetNumRegisteredElements()));
+		InOutDebugLabels.Add(FString::Printf(TEXT("Repository UObjects count: %i"), Repository->GetNumRegisteredUObjects()));
+		InOutDebugLabels.Add(FString::Printf(TEXT("Repository Custom NavLinks count: %i"), Repository->GetNumRegisteredCustomLinks()));
+	}
+
+	InOutDebugLabels.Add(FString::Printf(TEXT("NavData count: %i"), NavDataSet.Num()));
+	InOutDebugLabels.Add(FString::Printf(TEXT("MainNavData: %s"), MainNavData ? *MainNavData->GetName() : TEXT("none")));
+	InOutDebugLabels.Add(FString::Printf(TEXT("Custom NavLinks count: %i"), GetNumCustomLinks()));
+
+	if (GetNavOctree())
+	{
+		int32 NumNodes = 0;
+		int32 NumElements = 0;
+
+		GetNavOctree()->FindNodesWithPredicate(
+			[](FNavigationOctree::FNodeIndex /*ParentNodeIndex*/, FNavigationOctree::FNodeIndex /*NodeIndex*/, const FBoxCenterAndExtent&) { return true; },
+			[&, Octree = GetNavOctree()](FNavigationOctree::FNodeIndex /*ParentNodeIndex*/, const FNavigationOctree::FNodeIndex NodeIndex, const FBoxCenterAndExtent&)
+			{
+				NumNodes++;
+				NumElements += Octree->GetElementsForNode(NodeIndex).Num();
+			});
+
+		InOutDebugLabels.Add(FString::Printf(TEXT("Octree node count: %i"), NumNodes));
+		InOutDebugLabels.Add(FString::Printf(TEXT("Octree element count: %i"), NumElements));
+	}
+
+#if WITH_NAVMESH_CLUSTER_LINKS
+	InOutDebugLabels.Add(FString::Printf(TEXT("Using cluster links")));
+#endif // WITH_NAVMESH_CLUSTER_LINKS
+
+	if (IsActiveTilesGenerationEnabled()) // Checks bGenerateNavigationOnlyAroundNavigationInvokers
+	{
+		InOutDebugLabels.Add(FString::Printf(TEXT("Invoker Locations: %i"), GetInvokerLocations().Num()));
+	}
+
+	const int32 Running = GetNumRunningBuildTasks();
+	const int32 Remaining = GetNumRemainingBuildTasks();
+	if (Running || Remaining)
+	{
+		InOutDebugLabels.Add(FString::Printf(TEXT("Tile jobs running/remaining: %6d / %6d"), Running, Remaining));
+	}
+
+	InOutDebugLabels.Add(TEXT("")); // empty line
+}
+
 PRAGMA_ENABLE_DEPRECATION_WARNINGS
 
 void UNavigationSystemV1::ConfigureAsStatic(bool bEnableStatic)
@@ -1039,37 +1150,19 @@ bool UNavigationSystemV1::ConditionalPopulateNavOctree()
 			if (bStoreNavGeometry)
 			{
 #if WITH_RECAST
-				DefaultOctreeController.NavOctree->NavRelevantGeometryExportDelegate = FNavigationOctree::FNavRelevantGeometryExportDelegate::CreateStatic(&FRecastGeometryExport::ExportNavRelevantObjectGeometry);
+				DefaultOctreeController.NavOctree->GeometryExportDelegate =
+					FNavigationOctree::FGeometryExportDelegate::CreateStatic(&FRecastGeometryExport::ExportElementGeometry);
 #endif // WITH_RECAST
 			}
 
 			if (!DefaultOctreeController.IsNavigationOctreeLocked())
 			{
-				UWorld* World = GetWorld();
-				check(World);
-
-				// now process all actors on all levels
-				for (int32 LevelIndex = 0; LevelIndex < World->GetNumLevels(); ++LevelIndex)
-				{
-					ULevel* Level = World->GetLevel(LevelIndex);
-					if (ensure(Level) && Level->bIsVisible)
-					{
-						AddLevelToOctree(*Level);
-					}
-				}
-
-				// Register nav relevant objects currently registered in the repository world subsystem.
-				// This covers objects that are not AActor/UActorComponent based.
-				if (const UNavigationObjectRepository* Repository = UWorld::GetSubsystem<UNavigationObjectRepository>(GetWorld()))
+				// Register all elements registered in the repository world subsystem.
+				checkf(Repository, TEXT("%hs is expected to be called after the repository gets cached."), __FUNCTION__);
+				Repository->ForEachNavigationElement([this](const TSharedRef<const FNavigationElement>& Element)
 				{
-					for (TWeakInterfacePtr<INavRelevantInterface> It : Repository->GetNavRelevantObjects())
-					{
-						if (INavRelevantInterface* Interface = It.Get())
-						{
-							RegisterNavOctreeElement(Cast<UObject>(Interface), Interface, FNavigationOctreeController::OctreeUpdate_Default);
-						}
-					}
-				}
+					RegisterNavigationElementWithNavOctree(Element, FNavigationOctreeController::OctreeUpdate_Default);
+				});
 			}
 		}
 	}
@@ -1453,7 +1546,7 @@ void UNavigationSystemV1::Tick(float DeltaSeconds)
 		ProcessRegistrationCandidates();
 	}
 
-	if (DefaultOctreeController.PendingOctreeUpdates.Num() > 0)
+	if (DefaultOctreeController.PendingUpdates.Num() > 0)
 	{
 		SCOPE_CYCLE_COUNTER(STAT_Navigation_AddingActorsToNavOctree);
 		CSV_SCOPED_TIMING_STAT(NavigationBuildDetailed, Navigation_ProcessPendingOctreeUpdates);
@@ -2592,21 +2685,19 @@ void UNavigationSystemV1::ProcessRegistrationCandidates()
 
 void UNavigationSystemV1::ProcessCustomLinkPendingRegistration()
 {
-	if (const UNavigationObjectRepository* Repository = UWorld::GetSubsystem<UNavigationObjectRepository>(GetWorld()))
+	checkf(Repository, TEXT("%hs is expected to be called after the repository gets cached."), __FUNCTION__);
+	for (TWeakInterfacePtr<INavLinkCustomInterface> It : Repository->GetCustomLinks())
 	{
-		for (TWeakInterfacePtr<INavLinkCustomInterface> It : Repository->GetCustomLinks())
+		if (INavLinkCustomInterface* Interface = It.Get())
 		{
-			if (INavLinkCustomInterface* Interface = It.Get())
-			{
-				RegisterCustomLink(*Interface);
-			}
+			RegisterCustomLink(*Interface);
 		}
 	}
 }
 
 UNavigationSystemV1::ERegistrationResult UNavigationSystemV1::RegisterNavData(ANavigationData* NavData)
 {
-	UE_LOG(LogNavigation, Verbose, TEXT("%s %s"), ANSI_TO_TCHAR(__FUNCTION__), *GetFullNameSafe(NavData));
+	UE_LOG(LogNavigation, Verbose, TEXT("%hs %s"), __FUNCTION__, *GetFullNameSafe(NavData));
 	
 	if (NavData == nullptr)
 	{
@@ -2733,7 +2824,7 @@ UNavigationSystemV1::ERegistrationResult UNavigationSystemV1::RegisterNavData(AN
 
 void UNavigationSystemV1::UnregisterNavData(ANavigationData* NavData)
 {
-	UE_LOG(LogNavigation, Verbose, TEXT("%s %s"), ANSI_TO_TCHAR(__FUNCTION__), *GetFullNameSafe(NavData));
+	UE_LOG(LogNavigation, Verbose, TEXT("%hs %s"), __FUNCTION__, *GetFullNameSafe(NavData));
 	
 	NavDataSet.RemoveSingle(NavData);
 
@@ -2850,6 +2941,8 @@ void UNavigationSystemV1::UpdateCustomLink(const INavLinkCustomInterface* Custom
 
 void UNavigationSystemV1::RequestCustomLinkRegistering(INavLinkCustomInterface& CustomLink, UObject* Owner)
 {
+	SCOPE_CYCLE_COUNTER(STAT_NavOctreeBookkeeping);
+
 	if (Owner != nullptr)
 	{
 		if (UNavigationObjectRepository* Repository = UWorld::GetSubsystem<UNavigationObjectRepository>(Owner->GetWorld()))
@@ -2862,6 +2955,8 @@ void UNavigationSystemV1::RequestCustomLinkRegistering(INavLinkCustomInterface&
 
 void UNavigationSystemV1::RequestCustomLinkUnregistering(INavLinkCustomInterface& CustomLink, UObject* Owner)
 {
+	SCOPE_CYCLE_COUNTER(STAT_NavOctreeBookkeeping);
+
 	if (Owner != nullptr)
 	{
 		if (UNavigationObjectRepository* Repository = UWorld::GetSubsystem<UNavigationObjectRepository>(Owner->GetWorld()))
@@ -3138,7 +3233,7 @@ UNavigationSystemV1* UNavigationSystemV1::CreateNavigationSystem(UWorld* WorldOw
 		AWorldSettings* WorldSettings = WorldOwner->GetWorldSettings();
 		if (WorldSettings == nullptr || WorldSettings->IsNavigationSystemEnabled())
 		{
-			NavSys = NewObject<UNavigationSystemV1>(WorldOwner, GEngine->NavigationSystemClass);		
+			NavSys = NewObject<UNavigationSystemV1>(WorldOwner, GEngine->NavigationSystemClass);
 			WorldOwner->SetNavigationSystem(NavSys);
 		}
 	}
@@ -3175,6 +3270,18 @@ ANavigationData* UNavigationSystemV1::GetNavDataWithID(const uint16 NavDataID) c
 	return nullptr;
 }
 
+FNavigationElementHandle UNavigationSystemV1::GetNavigationElementHandleForUObject(const UObject* Object) const
+{
+	checkf(Repository, TEXT("%hs is expected to be called after the repository gets cached."), __FUNCTION__);
+	return Repository->GetNavigationElementHandleForUObject(Object);
+}
+
+TSharedPtr<const FNavigationElement> UNavigationSystemV1::GetNavigationElementForUObject(const UObject* Object) const
+{
+	checkf(Repository, TEXT("%hs is expected to be called after the repository gets cached."), __FUNCTION__);
+	return Repository->GetNavigationElementForUObject(Object);
+}
+
 void UNavigationSystemV1::OnNavRelevantObjectRegistered(UObject& Object)
 {
 	if (IsNavigationSystemStatic())
@@ -3182,89 +3289,90 @@ void UNavigationSystemV1::OnNavRelevantObjectRegistered(UObject& Object)
 		return;
 	}
 
-	if (INavRelevantInterface* NavInterface = Cast<INavRelevantInterface>(&Object))
+	if (const INavRelevantInterface* NavInterface = Cast<INavRelevantInterface>(&Object))
 	{
-		if (UNavigationObjectRepository* Repository = UWorld::GetSubsystem<UNavigationObjectRepository>(Object.GetWorld()))
-		{
-			Repository->RegisterNavRelevantObject(*NavInterface);
-		}
+		RegisterNavRelevantObjectStatic(*NavInterface, Object);
 	}
 }
 
 void UNavigationSystemV1::RegisterComponentToNavOctree(UActorComponent* Comp)
 {
+	SCOPE_CYCLE_COUNTER(STAT_NavOctreeBookkeeping);
+
 	if ((Comp == nullptr) || IsNavigationSystemStatic())
 	{
 		return;
 	}
 
-	SCOPE_CYCLE_COUNTER(STAT_DebugNavOctree);
-	INavRelevantInterface* NavInterface = Cast<INavRelevantInterface>(Comp);
-	if (NavInterface)
+	if (INavRelevantInterface* NavInterface = Cast<INavRelevantInterface>(Comp))
 	{
-		AActor* OwnerActor = Comp->GetOwner();
+		const AActor* OwnerActor = Comp->GetOwner();
 		if (OwnerActor && OwnerActor->IsComponentRelevantForNavigation(Comp))
 		{
-			UNavigationSystemV1* NavSys = FNavigationSystem::GetCurrent<UNavigationSystemV1>(OwnerActor->GetWorld());
-			if (NavSys)
-			{
-				NavSys->RegisterNavOctreeElement(Comp, NavInterface, FNavigationOctreeController::OctreeUpdate_Default);
-			}
+			RegisterNavRelevantObjectStatic(*NavInterface, *Comp);
 		}
 	}
 }
 
-void UNavigationSystemV1::OnNavRelevantObjectUnregistered(UObject& Object)
+FNavigationElementHandle UNavigationSystemV1::AddNavigationElement(UWorld* World, FNavigationElement&& Element)
 {
+	SCOPE_CYCLE_COUNTER(STAT_NavOctreeBookkeeping);
+
 	if (IsNavigationSystemStatic())
 	{
-		return;
+		return FNavigationElementHandle::Invalid;
 	}
 
-	if (INavRelevantInterface* NavInterface = Cast<INavRelevantInterface>(&Object))
+	if (UNavigationObjectRepository* Repository = UWorld::GetSubsystem<UNavigationObjectRepository>(World))
 	{
-		if (UNavigationObjectRepository* Repository = UWorld::GetSubsystem<UNavigationObjectRepository>(Object.GetWorld()))
+		if (const TSharedPtr<const FNavigationElement> SharedElement = Repository->AddNavigationElement(MoveTemp(Element)))
 		{
-			Repository->UnregisterNavRelevantObject(*NavInterface);
+			return SharedElement->GetHandle();
 		}
 	}
+
+	return FNavigationElementHandle::Invalid;
 }
 
-void UNavigationSystemV1::UnregisterComponentToNavOctree(UActorComponent* Comp)
+void UNavigationSystemV1::RemoveNavigationElement(UWorld* World, const FNavigationElementHandle ElementHandle)
 {
-	if ((Comp == nullptr) || IsNavigationSystemStatic())
+	SCOPE_CYCLE_COUNTER(STAT_NavOctreeBookkeeping);
+
+	if (IsNavigationSystemStatic())
 	{
 		return;
 	}
 
-	SCOPE_CYCLE_COUNTER(STAT_DebugNavOctree);
-	INavRelevantInterface* NavInterface = Cast<INavRelevantInterface>(Comp);
-	if (NavInterface)
+	if (UNavigationObjectRepository* Repository = UWorld::GetSubsystem<UNavigationObjectRepository>(World))
 	{
-		AActor* OwnerActor = Comp->GetOwner();
-		if (OwnerActor)
-		{
-			// skip IsComponentRelevantForNavigation check, it's only for adding new stuff
+		Repository->RemoveNavigationElement(ElementHandle);
+	}
+}
 
-			UNavigationSystemV1* NavSys = FNavigationSystem::GetCurrent<UNavigationSystemV1>(OwnerActor->GetWorld());
-			if (NavSys)
-			{
-				NavSys->UnregisterNavOctreeElement(Comp, NavInterface, FNavigationOctreeController::OctreeUpdate_Default);
-			}
-		}
+void UNavigationSystemV1::OnNavRelevantObjectUnregistered(UObject& Object)
+{
+	UnregisterNavRelevantObjectStatic(Object);
+}
+
+void UNavigationSystemV1::UnregisterComponentToNavOctree(UActorComponent* Comp)
+{
+	// skip IsComponentRelevantForNavigation check, it's only for adding new stuff
+	if (Comp)
+	{
+		UnregisterNavRelevantObjectStatic(*Comp);
 	}
 }
 
 // Deprecated
-void UNavigationSystemV1::AddDirtyArea(const FBox& NewArea, int32 Flags, const FName& DebugReason /*= NAME_None*/)
+void UNavigationSystemV1::AddDirtyArea(const FBox& NewArea, int32 Flags, const FName& DebugReason)
 {
-	AddDirtyArea(NewArea, static_cast<ENavigationDirtyFlag>(Flags), DebugReason);
+	DefaultDirtyAreasController.AddArea(NewArea, static_cast<ENavigationDirtyFlag>(Flags), /*ElementProvideFunc*/ nullptr, /*DirtyElement*/nullptr, DebugReason);
 }
 
 // Deprecated
-void UNavigationSystemV1::AddDirtyArea(const FBox& NewArea, int32 Flags, const TFunction<UObject*()>& ObjectProviderFunc, const FName& DebugReason /*= NAME_None*/)
+void UNavigationSystemV1::AddDirtyArea(const FBox& NewArea, int32 Flags, const TFunction<UObject*()>&, const FName& DebugReason /*= NAME_None*/)
 {
-	AddDirtyArea(NewArea, static_cast<ENavigationDirtyFlag>(Flags), ObjectProviderFunc, DebugReason);
+	DefaultDirtyAreasController.AddArea(NewArea, static_cast<ENavigationDirtyFlag>(Flags), /*ElementProvideFunc*/ nullptr, /*DirtyElement*/nullptr, DebugReason);
 }
 
 // Deprecated
@@ -3278,9 +3386,9 @@ void UNavigationSystemV1::AddDirtyArea(const FBox& NewArea, const ENavigationDir
 	DefaultDirtyAreasController.AddArea(NewArea, Flags, nullptr, nullptr, DebugReason);
 }
 
-void UNavigationSystemV1::AddDirtyArea(const FBox& NewArea, const ENavigationDirtyFlag Flags, const TFunction<UObject*()>& ObjectProviderFunc, const FName& DebugReason /*= NAME_None*/)
+void UNavigationSystemV1::AddDirtyArea(const FBox& NewArea, const ENavigationDirtyFlag Flags, const TFunction<const TSharedPtr<const FNavigationElement>()>& ElementProviderFunc, const FName& DebugReason /*= NAME_None*/)
 {
-	DefaultDirtyAreasController.AddArea(NewArea, Flags, ObjectProviderFunc, nullptr, DebugReason);
+	DefaultDirtyAreasController.AddArea(NewArea, Flags, ElementProviderFunc, /*DirtyElement*/ nullptr, DebugReason);
 }
 
 void UNavigationSystemV1::AddDirtyAreas(const TArray<FBox>& NewAreas, const ENavigationDirtyFlag Flags, const FName& DebugReason /*= NAME_None*/)
@@ -3306,10 +3414,16 @@ bool UNavigationSystemV1::HasDirtyAreasQueued() const
 	return DefaultDirtyAreasController.IsDirty();
 }
 
+FSetElementId UNavigationSystemV1::RegisterNavigationElementWithNavOctree(const TSharedRef<const FNavigationElement>& Element, const int32 UpdateFlags)
+{
+	return FNavigationDataHandler(DefaultOctreeController, DefaultDirtyAreasController).RegisterElementWithNavOctree(Element, UpdateFlags);
+}
+
+// Deprecated
 FSetElementId UNavigationSystemV1::RegisterNavOctreeElement(UObject* ElementOwner, INavRelevantInterface* ElementInterface, int32 UpdateFlags)
 {
 	return (ElementOwner && ElementInterface) 
-		? FNavigationDataHandler(DefaultOctreeController, DefaultDirtyAreasController).RegisterNavOctreeElement(*ElementOwner, *ElementInterface, UpdateFlags)
+		? FNavigationDataHandler(DefaultOctreeController, DefaultDirtyAreasController).RegisterElementWithNavOctree(FNavigationElement::CreateFromNavRelevantInterface(*ElementInterface), UpdateFlags)
 		: FSetElementId();
 }
 
@@ -3318,31 +3432,49 @@ void UNavigationSystemV1::AddElementToNavOctree(const FNavigationDirtyElement& D
 	FNavigationDataHandler(DefaultOctreeController, DefaultDirtyAreasController).AddElementToNavOctree(DirtyElement);
 }
 
-bool UNavigationSystemV1::GetNavOctreeElementData(const UObject& NodeOwner, ENavigationDirtyFlag& OutDirtyFlags, FBox& OutDirtyBounds)
+bool UNavigationSystemV1::GetNavOctreeElementData(const FNavigationElementHandle Element, ENavigationDirtyFlag& OutDirtyFlags, FBox& OutDirtyBounds)
 {
-	return DefaultOctreeController.GetNavOctreeElementData(NodeOwner, OutDirtyFlags, OutDirtyBounds);
+	return DefaultOctreeController.GetNavOctreeElementData(Element, OutDirtyFlags, OutDirtyBounds);
 }
 
 // Deprecated
 bool UNavigationSystemV1::GetNavOctreeElementData(const UObject& NodeOwner, int32& DirtyFlags, FBox& DirtyBounds)
 {
 	ENavigationDirtyFlag TmpDirtyFlags = ENavigationDirtyFlag::None;
-	const bool bSuccess = GetNavOctreeElementData(NodeOwner, TmpDirtyFlags, DirtyBounds);
+	const bool bSuccess = GetNavOctreeElementData(FNavigationElementHandle(&NodeOwner), TmpDirtyFlags, DirtyBounds);
 	DirtyFlags = static_cast<int32>(TmpDirtyFlags);
 	return bSuccess;
 }
 
+// Deprecated
 void UNavigationSystemV1::UnregisterNavOctreeElement(UObject* ElementOwner, INavRelevantInterface* ElementInterface, int32 UpdateFlags)
 {
 	if (ElementOwner && ElementInterface)
 	{
-		FNavigationDataHandler(DefaultOctreeController, DefaultDirtyAreasController).UnregisterNavOctreeElement(*ElementOwner, *ElementInterface, UpdateFlags);
+		UnregisterNavRelevantObjectInternal(*ElementOwner);
 	}
 }
 
-void UNavigationSystemV1::RemoveNavOctreeElementId(const FOctreeElementId2& ElementId, int32 UpdateFlags)
+void UNavigationSystemV1::UnregisterNavigationElementWithOctree(const TSharedRef<const FNavigationElement>& Element, const int32 UpdateFlags)
+{
+	FNavigationDataHandler(DefaultOctreeController, DefaultDirtyAreasController).UnregisterElementWithNavOctree(Element, UpdateFlags);
+}
+
+// Deprecated
+void UNavigationSystemV1::RemoveObjectsNavOctreeId(const UObject& Object)
+{
+	// doing nothing since we don't want external calls to remove a mapping without properly update the nodes
+}
+
+// Deprecated
+void UNavigationSystemV1::RemoveNavOctreeElementId(const FOctreeElementId2& ElementId, const int32 UpdateFlags)
+{
+	RemoveFromNavOctree(ElementId, UpdateFlags);
+}
+
+void UNavigationSystemV1::RemoveFromNavOctree(const FOctreeElementId2& ElementId, const int32 UpdateFlags)
 {
-	FNavigationDataHandler(DefaultOctreeController, DefaultDirtyAreasController).RemoveNavOctreeElementId(ElementId, UpdateFlags);
+	FNavigationDataHandler(DefaultOctreeController, DefaultDirtyAreasController).RemoveFromNavOctree(ElementId, UpdateFlags);
 }
 
 void UNavigationSystemV1::DemandLazyDataGathering(FNavigationRelevantData& ElementData)
@@ -3350,58 +3482,164 @@ void UNavigationSystemV1::DemandLazyDataGathering(FNavigationRelevantData& Eleme
 	FNavigationDataHandler(DefaultOctreeController, DefaultDirtyAreasController).DemandLazyDataGathering(ElementData);
 }
 
+// Deprecated
 const FNavigationRelevantData* UNavigationSystemV1::GetDataForObject(const UObject& Object) const
 {
-	return DefaultOctreeController.GetDataForObject(Object);
+	return GetDataForElement(FNavigationElementHandle(&Object));
 }
 
+// Deprecated
 FNavigationRelevantData* UNavigationSystemV1::GetMutableDataForObject(const UObject& Object)
 {
-	return DefaultOctreeController.GetMutableDataForObject(Object);
+	return GetMutableDataForElement(FNavigationElementHandle(&Object));
 }
 
-void UNavigationSystemV1::UpdateNavRelevantObjectInNavOctree(UObject& Object)
+const FNavigationRelevantData* UNavigationSystemV1::GetDataForElement(const FNavigationElementHandle Element) const
 {
-	SCOPE_CYCLE_COUNTER(STAT_DebugNavOctree);
+	return DefaultOctreeController.GetDataForElement(Element);
+}
 
-	if (INavRelevantInterface* NavElement = Cast<INavRelevantInterface>(&Object))
+FNavigationRelevantData* UNavigationSystemV1::GetMutableDataForElement(const FNavigationElementHandle Element)
+{
+	return DefaultOctreeController.GetMutableDataForElement(Element);
+}
+
+void UNavigationSystemV1::RegisterNavRelevantObjectStatic(const INavRelevantInterface& NavRelevantObject, const UObject& Object)
+{
+	SCOPE_CYCLE_COUNTER(STAT_NavOctreeBookkeeping);
+
+	if (IsNavigationSystemStatic())
 	{
-		if (UNavigationSystemV1* NavSys = FNavigationSystem::GetCurrent<UNavigationSystemV1>(Object.GetWorld()))
-		{
-			if (NavElement->IsNavigationRelevant())
-			{
-				NavSys->UpdateNavOctreeElement(&Object, NavElement, FNavigationOctreeController::OctreeUpdate_Default);	
-			}
-			else
-			{
-				NavSys->UnregisterNavOctreeElement(&Object, NavElement, FNavigationOctreeController::OctreeUpdate_Default);
-			}
-		}
+		return;
+	}
+
+	if (UNavigationObjectRepository* Repository = UWorld::GetSubsystem<UNavigationObjectRepository>(Object.GetWorld()))
+	{
+		Repository->RegisterNavRelevantObject(NavRelevantObject);
 	}
 }
 
-void UNavigationSystemV1::UpdateActorInNavOctree(AActor& Actor)
+void UNavigationSystemV1::RegisterNavRelevantObjectInternal(const INavRelevantInterface& NavRelevantObject, const UObject& Object)
+{
+	SCOPE_CYCLE_COUNTER(STAT_NavOctreeBookkeeping);
+
+	if (IsNavigationSystemStatic())
+	{
+		return;
+	}
+
+	checkf(Repository, TEXT("%hs is expected to be called after the repository gets cached."), __FUNCTION__);
+	Repository->RegisterNavRelevantObject(NavRelevantObject);
+}
+
+void UNavigationSystemV1::UnregisterNavRelevantObjectStatic(const UObject& Object)
+{
+	SCOPE_CYCLE_COUNTER(STAT_NavOctreeBookkeeping);
+
+	if (IsNavigationSystemStatic())
+	{
+		return;
+	}
+
+	if (UNavigationObjectRepository* Repository = UWorld::GetSubsystem<UNavigationObjectRepository>(Object.GetWorld()))
+	{
+		Repository->UnregisterNavRelevantObject(&Object);
+	}
+}
+
+void UNavigationSystemV1::UnregisterNavRelevantObjectInternal(const UObject& Object)
 {
+	SCOPE_CYCLE_COUNTER(STAT_NavOctreeBookkeeping);
+
 	if (IsNavigationSystemStatic())
 	{
 		return;
 	}
-	SCOPE_CYCLE_COUNTER(STAT_DebugNavOctree);
 
-	INavRelevantInterface* NavElement = Cast<INavRelevantInterface>(&Actor);
-	if (NavElement)
+	checkf(Repository, TEXT("%hs is expected to be called after the repository gets cached."), __FUNCTION__);
+	Repository->UnregisterNavRelevantObject(&Object);
+}
+
+void UNavigationSystemV1::UpdateNavRelevantObjectInNavOctreeStatic(
+	const INavRelevantInterface& InNavRelevantObject,
+	const UObject& InObject,
+	UNavigationSystemV1* InNavigationSystem,
+	TFunctionRef<void(UNavigationSystemV1&, const TSharedRef<const FNavigationElement>&)> InCallback)
+{
+	SCOPE_CYCLE_COUNTER(STAT_NavOctreeBookkeeping);
+
+	if (IsNavigationSystemStatic())
+	{
+		return;
+	}
+
+	if (!ensureMsgf(InNavRelevantObject.IsNavigationRelevant(), TEXT("%hs: %s is not navigation relevant"), __FUNCTION__, *InObject.GetName()))
 	{
-		UNavigationSystemV1* NavSys = FNavigationSystem::GetCurrent<UNavigationSystemV1>(Actor.GetWorld());
-		if (NavSys)
+		return;
+	}
+
+	if (UNavigationSystemV1* NavSys = InNavigationSystem ? InNavigationSystem : FNavigationSystem::GetCurrent<UNavigationSystemV1>(InObject.GetWorld()))
+	{
+		checkf(NavSys->Repository, TEXT("%hs is expected to be called after the repository gets cached."), __FUNCTION__);
+		if (const TSharedPtr<const FNavigationElement> SharedElement = NavSys->Repository->UpdateNavigationElementForUObject(InNavRelevantObject, InObject))
 		{
-			NavSys->UpdateNavOctreeElement(&Actor, NavElement, FNavigationOctreeController::OctreeUpdate_Default);
+			InCallback(*NavSys, SharedElement.ToSharedRef());
 		}
 	}
+	else
+	{
+		// Navigation system not available so use the static registration to be stored in the repository
+		// so the navigation system will gather it on initialization.
+		UE_LOG(LogNavigation, VeryVerbose,
+			TEXT("%hs: %s Registering to the repository (NavigationSystem not available)"), __FUNCTION__, *InObject.GetName());
+
+		RegisterNavRelevantObjectStatic(InNavRelevantObject, InObject);
+	}
+}
+
+void UNavigationSystemV1::UpdateNavRelevantObjectInNavOctree(UObject& Object)
+{
+	SCOPE_CYCLE_COUNTER(STAT_NavOctreeBookkeeping);
+
+	if (IsNavigationSystemStatic())
+	{
+		return;
+	}
+
+	const INavRelevantInterface* NavRelevantInterface = Cast<INavRelevantInterface>(&Object);
+	if (NavRelevantInterface && NavRelevantInterface->IsNavigationRelevant())
+	{
+		UpdateNavRelevantObjectInNavOctreeStatic(*NavRelevantInterface, Object, /*NavigationSystem*/nullptr,
+			[](UNavigationSystemV1& NavSys, const TSharedRef<const FNavigationElement>& SharedElement)
+			{
+				NavSys.UpdateNavOctreeElement(SharedElement->GetHandle(), SharedElement, FNavigationOctreeController::OctreeUpdate_Default);
+			});
+	}
+}
+
+void UNavigationSystemV1::OnNavigationElementUpdated(UWorld* World, const FNavigationElementHandle ElementHandle, FNavigationElement&& Element)
+{
+	SCOPE_CYCLE_COUNTER(STAT_NavOctreeBookkeeping);
+
+	if (IsNavigationSystemStatic())
+	{
+		return;
+	}
+
+	if (UNavigationSystemV1* NavSys = FNavigationSystem::GetCurrent<UNavigationSystemV1>(World))
+	{
+		NavSys->UpdateNavOctreeElement(ElementHandle, MakeShared<FNavigationElement>(MoveTemp(Element)) , FNavigationOctreeController::OctreeUpdate_Default);
+	}
+}
+
+void UNavigationSystemV1::UpdateActorInNavOctree(AActor& Actor)
+{
+	UpdateNavRelevantObjectInNavOctree(Actor);
 }
 
 void UNavigationSystemV1::UpdateComponentInNavOctree(UActorComponent& Comp)
 {
-	SCOPE_CYCLE_COUNTER(STAT_DebugNavOctree);
+	SCOPE_CYCLE_COUNTER(STAT_NavOctreeBookkeeping);
 
 	if (ShouldUpdateNavOctreeOnComponentChange() == false)
 	{
@@ -3409,86 +3647,117 @@ void UNavigationSystemV1::UpdateComponentInNavOctree(UActorComponent& Comp)
 	}
 
 	// special case for early out: use cached nav relevancy
-	if (Comp.bNavigationRelevant == true)
+	if (Comp.bNavigationRelevant)
 	{
-		AActor* OwnerActor = Comp.GetOwner();
-		if (OwnerActor)
+		if (const AActor* OwnerActor = Comp.GetOwner())
 		{
-			INavRelevantInterface* NavElement = Cast<INavRelevantInterface>(&Comp);
-			if (NavElement)
+			const INavRelevantInterface* NavRelevantInterface = Cast<INavRelevantInterface>(&Comp);
+			if (ensureMsgf(NavRelevantInterface != nullptr, TEXT("Components reaching this point are expected to implement INavRelevantInterface.")))
 			{
-				UNavigationSystemV1* NavSys = FNavigationSystem::GetCurrent<UNavigationSystemV1>(OwnerActor->GetWorld());
-				if (NavSys)
+				if (OwnerActor->IsComponentRelevantForNavigation(&Comp) && Comp.IsNavigationRelevant())
 				{
-					if (OwnerActor->IsComponentRelevantForNavigation(&Comp))
-					{
-						NavSys->UpdateNavOctreeElement(&Comp, NavElement, FNavigationOctreeController::OctreeUpdate_Default);
-					}
-					else
+
+					UpdateNavRelevantObjectInNavOctreeStatic(*NavRelevantInterface, Comp, /*NavigationSystem*/nullptr,
+						[](UNavigationSystemV1& NavSys, const TSharedRef<const FNavigationElement>& SharedElement)
+						{
+							NavSys.UpdateNavOctreeElement(SharedElement->GetHandle(), SharedElement, FNavigationOctreeController::OctreeUpdate_Default);
+						});
+				}
+				else
+				{
+					if (UNavigationObjectRepository* Repository = UWorld::GetSubsystem<UNavigationObjectRepository>(Comp.GetWorld()))
 					{
-						NavSys->UnregisterNavOctreeElement(&Comp, NavElement, FNavigationOctreeController::OctreeUpdate_Default);
+						Repository->UnregisterNavRelevantObject(&Comp);
 					}
 				}
 			}
 		}
 	}
-	else if (Comp.CanEverAffectNavigation()) 
+	else if (Comp.CanEverAffectNavigation())
 	{
 		// could have been relevant before and not it isn't. Need to check if there's an octree element ID for it
-		INavRelevantInterface* NavElement = Cast<INavRelevantInterface>(&Comp);
-		if (NavElement)
+		if (UNavigationObjectRepository* Repository = UWorld::GetSubsystem<UNavigationObjectRepository>(Comp.GetWorld()))
 		{
-			AActor* OwnerActor = Comp.GetOwner();
-			if (OwnerActor)
-			{
-				UNavigationSystemV1* NavSys = FNavigationSystem::GetCurrent<UNavigationSystemV1>(OwnerActor->GetWorld());
-				if (NavSys)
-				{
-					NavSys->UnregisterNavOctreeElement(&Comp, NavElement, FNavigationOctreeController::OctreeUpdate_Default);
-				}
-			}
+			Repository->UnregisterNavRelevantObject(&Comp);
 		}
 	}
 }
 
-void UNavigationSystemV1::UpdateActorAndComponentsInNavOctree(AActor& Actor, bool bUpdateAttachedActors)
+void UNavigationSystemV1::UpdateActorAndComponentsInNavOctree(AActor& Actor, const bool bUpdateAttachedActors)
 {
+	SCOPE_CYCLE_COUNTER(STAT_NavOctreeBookkeeping);
+
 	if (IsNavigationSystemStatic())
 	{
 		return;
 	}
 
 	UNavigationSystemV1* NavSys = FNavigationSystem::GetCurrent<UNavigationSystemV1>(Actor.GetWorld());
-	if (NavSys)
-	{
-		FNavigationDataHandler DataHandler(NavSys->DefaultOctreeController, NavSys->DefaultDirtyAreasController);
 
-		if (ShouldUpdateNavOctreeOnComponentChange())
-		{
-			// do the full update
-			DataHandler.UpdateActorAndComponentsInNavOctree(Actor);
-		}
-		else
+	// Callback to update an actor with its components
+	auto UpdateActorAndComponentFunc = [NavSys](AActor& ActorToUpdate)
 		{
-			// update just the actor itself 
-			INavRelevantInterface* NavElement = Cast<INavRelevantInterface>(&Actor);
-			if (NavElement)
+			const INavRelevantInterface* ActorNavRelevantInterface = Cast<INavRelevantInterface>(&ActorToUpdate);
+			if (ActorNavRelevantInterface && ActorNavRelevantInterface->IsNavigationRelevant())
 			{
-				DataHandler.UpdateNavOctreeElement(Actor, *NavElement, FNavigationOctreeController::OctreeUpdate_Default);
+				UpdateNavRelevantObjectInNavOctreeStatic(*ActorNavRelevantInterface, ActorToUpdate, NavSys,
+					[](UNavigationSystemV1& InNavSys, const TSharedRef<const FNavigationElement>& SharedElement)
+					{
+						InNavSys.UpdateNavOctreeElement(SharedElement->GetHandle(), SharedElement, FNavigationOctreeController::OctreeUpdate_Default);
+					});
 			}
-		}
 
-		if (bUpdateAttachedActors)
-		{
-			TArray<AActor*> UniqueAttachedActors;
-			if (GetAllAttachedActors(Actor, UniqueAttachedActors) > 0)
+			for (UActorComponent* Component : ActorToUpdate.GetComponents())
 			{
-				for (AActor* AttachedActor : UniqueAttachedActors)
+				if (Component->CanEverAffectNavigation()
+					&& ActorToUpdate.IsComponentRelevantForNavigation(Component))
 				{
-					checkf(AttachedActor, TEXT("GetAllAttachedActors should only return unique, non-null ptrs."));
-					DataHandler.UpdateActorAndComponentsInNavOctree(*AttachedActor);
+					const INavRelevantInterface* ComponentNavRelevantInterface = Cast<INavRelevantInterface>(Component);
+					if (ComponentNavRelevantInterface && ComponentNavRelevantInterface->IsNavigationRelevant())
+					{
+						UpdateNavRelevantObjectInNavOctreeStatic(*ComponentNavRelevantInterface, *Component, NavSys,
+							[&ActorToUpdate, Component](UNavigationSystemV1& InNavSys, const TSharedRef<const FNavigationElement>& InElement)
+							{
+								InNavSys.UpdateNavOctreeElement(InElement->GetHandle(), InElement, FNavigationOctreeController::OctreeUpdate_Default);
+							});
+					}
+				}
+
+				if (NavSys
+					&& ensureMsgf(NavSys->Repository, TEXT("UNavigationObjectRepository is mandatory for navigation system operations.")))
+				{
+					NavSys->Repository->UnregisterNavRelevantObject(Component);
 				}
 			}
+		};
+
+	if (ShouldUpdateNavOctreeOnComponentChange())
+	{
+		UpdateActorAndComponentFunc(Actor);
+	}
+	else
+	{
+		const INavRelevantInterface* ActorNavRelevantInterface = Cast<INavRelevantInterface>(&Actor);
+		if (ActorNavRelevantInterface && ActorNavRelevantInterface->IsNavigationRelevant())
+		{
+			UpdateNavRelevantObjectInNavOctreeStatic(*ActorNavRelevantInterface, Actor, NavSys,
+				[](UNavigationSystemV1& InNavSys, const TSharedRef<const FNavigationElement>& SharedElement)
+				{
+					InNavSys.UpdateNavOctreeElement(SharedElement->GetHandle(), SharedElement, FNavigationOctreeController::OctreeUpdate_Default);
+				});
+		}
+	}
+
+	if (bUpdateAttachedActors)
+	{
+		TArray<AActor*> UniqueAttachedActors;
+		if (GetAllAttachedActors(Actor, UniqueAttachedActors) > 0)
+		{
+			for (AActor* AttachedActor : UniqueAttachedActors)
+			{
+				checkf(AttachedActor, TEXT("GetAllAttachedActors should only return unique, non-null ptrs."));
+				UpdateActorAndComponentFunc(*AttachedActor);
+			}
 		}
 	}
 }
@@ -3564,37 +3833,103 @@ void UNavigationSystemV1::ClearNavOctreeAll(AActor* Actor)
 	}
 }
 
+// Deprecated
 void UNavigationSystemV1::UpdateNavOctreeElement(UObject* ElementOwner, INavRelevantInterface* ElementInterface, int32 UpdateFlags)
 {
+	if (IsNavigationSystemStatic())
+	{
+		return;
+	}
+
 	if (ElementOwner && ElementInterface)
 	{
-		FNavigationDataHandler(DefaultOctreeController, DefaultDirtyAreasController).UpdateNavOctreeElement(*ElementOwner, *ElementInterface, UpdateFlags);
+		if (const FNavigationElementHandle Handle = GetNavigationElementHandleForUObject(ElementOwner))
+		{
+			FNavigationDataHandler(DefaultOctreeController, DefaultDirtyAreasController).UpdateNavOctreeElement(
+				Handle,
+				FNavigationElement::CreateFromNavRelevantInterface(*ElementInterface),
+				UpdateFlags);
+		}
 	}
 }
 
+void UNavigationSystemV1::UpdateNavOctreeElement(const FNavigationElementHandle Handle, const TSharedRef<const FNavigationElement>& Element, const int32 UpdateFlags)
+{
+	if (IsNavigationSystemStatic())
+	{
+		return;
+	}
+
+	FNavigationDataHandler(DefaultOctreeController, DefaultDirtyAreasController).UpdateNavOctreeElement(Handle, Element, UpdateFlags);
+}
+
+// Deprecated
 void UNavigationSystemV1::UpdateNavOctreeParentChain(UObject* ElementOwner, bool bSkipElementOwnerUpdate)
 {
+	PRAGMA_DISABLE_DEPRECATION_WARNINGS
 	if (ElementOwner)
 	{
 		FNavigationDataHandler(DefaultOctreeController, DefaultDirtyAreasController).UpdateNavOctreeParentChain(*ElementOwner, bSkipElementOwnerUpdate);
 	}
+	PRAGMA_ENABLE_DEPRECATION_WARNINGS
 }
 
+// Deprecated
 bool UNavigationSystemV1::UpdateNavOctreeElementBounds(UActorComponent* Comp, const FBox& NewBounds, const FBox& DirtyArea)
 {
-	return Comp
-		? UpdateNavOctreeElementBounds(*Comp, NewBounds, { DirtyArea })
-		: false;
+	checkf(Repository, TEXT("%hs is expected to be called after the repository gets cached."), __FUNCTION__);
+	const FNavigationElementHandle Handle = Repository->GetNavigationElementHandleForUObject(Comp);
+	if (Handle.IsValid())
+	{
+		return UpdateNavOctreeElementBounds(Handle, NewBounds, { DirtyArea });
+	}
+
+	return false;
 }
 
+// Deprecated
 bool UNavigationSystemV1::UpdateNavOctreeElementBounds(UObject& Object, const FBox& NewBounds, TConstArrayView<FBox> DirtyAreas)
 {
-	return FNavigationDataHandler(DefaultOctreeController, DefaultDirtyAreasController).UpdateNavOctreeElementBounds(Object, NewBounds, DirtyAreas);
+	checkf(Repository, TEXT("%hs is expected to be called after the repository gets cached."), __FUNCTION__);
+	const FNavigationElementHandle Handle = Repository->GetNavigationElementHandleForUObject(&Object);
+	if (Handle.IsValid())
+	{
+		return UpdateNavOctreeElementBounds(Handle, NewBounds, DirtyAreas);
+	}
+
+	return false;
 }
 
+bool UNavigationSystemV1::UpdateNavOctreeElementBounds(const FNavigationElementHandle Handle, const FBox& NewBounds, const TConstArrayView<FBox> DirtyAreas)
+{
+	if (IsNavigationSystemStatic())
+	{
+		return false;
+	}
+
+	return FNavigationDataHandler(DefaultOctreeController, DefaultDirtyAreasController).UpdateNavOctreeElementBounds(Handle, NewBounds, DirtyAreas);
+}
+
+// Deprecated
 bool UNavigationSystemV1::ReplaceAreaInOctreeData(const UObject& Object, TSubclassOf<UNavArea> OldArea, TSubclassOf<UNavArea> NewArea, bool bReplaceChildClasses)
 {
-	return FNavigationDataHandler(DefaultOctreeController, DefaultDirtyAreasController).ReplaceAreaInOctreeData(Object, OldArea, NewArea, bReplaceChildClasses);
+	checkf(Repository, TEXT("%hs is expected to be called after the repository gets cached."), __FUNCTION__);
+	if (const FNavigationElementHandle Handle = Repository->GetNavigationElementHandleForUObject(&Object))
+	{
+		return ReplaceAreaInOctreeData(Handle, OldArea, NewArea, bReplaceChildClasses);
+	}
+
+	return false;
+}
+
+bool UNavigationSystemV1::ReplaceAreaInOctreeData(const FNavigationElementHandle Handle, const TSubclassOf<UNavArea> OldArea, const TSubclassOf<UNavArea> NewArea, const bool bReplaceChildClasses)
+{
+	if (IsNavigationSystemStatic())
+	{
+		return false;
+	}
+
+	return FNavigationDataHandler(DefaultOctreeController, DefaultDirtyAreasController).ReplaceAreaInOctreeData(Handle, OldArea, NewArea, bReplaceChildClasses);
 }
 
 void UNavigationSystemV1::OnComponentRegistered(UActorComponent* Comp)
@@ -3619,39 +3954,27 @@ void UNavigationSystemV1::UnregisterComponent(UActorComponent* Comp)
 
 void UNavigationSystemV1::OnActorRegistered(AActor* Actor)
 {
-	if ((Actor == nullptr) || IsNavigationSystemStatic())
+	if (IsNavigationSystemStatic())
 	{
 		return;
 	}
 
-	SCOPE_CYCLE_COUNTER(STAT_DebugNavOctree);
-	INavRelevantInterface* NavInterface = Cast<INavRelevantInterface>(Actor);
-	if (NavInterface)
+	if (const INavRelevantInterface* NavInterface = Cast<INavRelevantInterface>(Actor))
 	{
-		UNavigationSystemV1* NavSys = FNavigationSystem::GetCurrent<UNavigationSystemV1>(Actor->GetWorld());
-		if (NavSys)
-		{
-			NavSys->RegisterNavOctreeElement(Actor, NavInterface, FNavigationOctreeController::OctreeUpdate_Default);
-		}
+		RegisterNavRelevantObjectStatic(*NavInterface, *Actor);
 	}
 }
 
 void UNavigationSystemV1::OnActorUnregistered(AActor* Actor)
 {
-	if ((Actor == nullptr) || IsNavigationSystemStatic())
+	if (IsNavigationSystemStatic())
 	{
 		return;
 	}
 
-	SCOPE_CYCLE_COUNTER(STAT_DebugNavOctree);
-	INavRelevantInterface* NavInterface = Cast<INavRelevantInterface>(Actor);
-	if (NavInterface)
+	if (Actor)
 	{
-		UNavigationSystemV1* NavSys = FNavigationSystem::GetCurrent<UNavigationSystemV1>(Actor->GetWorld());
-		if (NavSys)
-		{
-			NavSys->UnregisterNavOctreeElement(Actor, NavInterface, FNavigationOctreeController::OctreeUpdate_Default);
-		}
+		UnregisterNavRelevantObjectStatic(*Actor);
 	}
 }
 
@@ -3991,7 +4314,7 @@ void UNavigationSystemV1::GetInvokerSeedLocations(const UWorld& InWorld, TArray<
 
 	for (const FVector Location : Locations)
 	{
-		OutSeedLocations.Add(FVector2D(Location));	
+		OutSeedLocations.Add(FVector2D(Location));
 	}
 }
 
@@ -4444,8 +4767,6 @@ bool UNavigationSystemV1::IsNavigationBuildInProgress()
 	return bRet;
 }
 
-//deprecated
-
 void UNavigationSystemV1::OnNavigationGenerationFinished(ANavigationData& NavData)
 {
 	OnNavigationGenerationFinishedDelegate.Broadcast(&NavData);
@@ -4537,7 +4858,7 @@ void UNavigationSystemV1::OnLevelAddedToWorld(ULevel* InLevel, UWorld* InWorld)
 		// While streaming multiple levels it is possible that NavigationData and NavMeshBoundsVolume from different levels gets
 		// loaded in different order so we need to wait navigation system initialization to make sure everything is registered properly.
 		// Otherwise the register may fail and discard the navigation data since navbounds are not registered.
-		UE_LOG(LogNavigation, Log, TEXT("%s won't process navigation data registration candidates until OperationMode is set. Waiting for OnWorldInitDone."), ANSI_TO_TCHAR(__FUNCTION__));
+		UE_LOG(LogNavigation, Log, TEXT("%hs won't process navigation data registration candidates until OperationMode is set. Waiting for OnWorldInitDone."), __FUNCTION__);
 	}
 	else if (NavDataRegistrationQueue.Num() > 0)
 	{
@@ -4611,7 +4932,7 @@ void UNavigationSystemV1::RemoveLevelCollisionFromOctree(ULevel* Level)
 
 void UNavigationSystemV1::OnPostLoadMap(UWorld* LoadedWorld)
 {
-	UE_LOG(LogNavigation, Verbose, TEXT("%s (Package: %s)"), ANSI_TO_TCHAR(__FUNCTION__), *GetNameSafe(LoadedWorld->GetOuter()));
+	UE_LOG(LogNavigation, Verbose, TEXT("%hs (Package: %s)"), __FUNCTION__, *GetNameSafe(LoadedWorld->GetOuter()));
 
 	// If map has been loaded and there are some navigation bounds volumes 
 	// then create appropriate navigation structured
@@ -4624,7 +4945,7 @@ void UNavigationSystemV1::OnPostLoadMap(UWorld* LoadedWorld)
 		 (GetRuntimeGenerationType() != ERuntimeGenerationType::Static) )	// Prevent creating a static default nav instance out of the editor (GetRuntimeGenerationType() is always dynamic in editor).
 	{
 		NavData = GetDefaultNavDataInstance(FNavigationSystem::Create);
-		UE_LOG(LogNavigation, Verbose, TEXT("%s Created DefaultNavDataInstance %s"), ANSI_TO_TCHAR(__FUNCTION__), *GetNameSafe(NavData));
+		UE_LOG(LogNavigation, Verbose, TEXT("%hs Created DefaultNavDataInstance %s"), __FUNCTION__, *GetNameSafe(NavData));
 	}
 }
 
@@ -4717,7 +5038,10 @@ void UNavigationSystemV1::CleanUp(FNavigationSystem::ECleanupMode Mode)
 	const UWorld* MyWorld = (Mode == FNavigationSystem::ECleanupMode::CleanupWithWorld) ? GetWorld() : nullptr;
 	if (MyWorld)
 	{
-		UnregisterFromRepositoryDelegates();
+		if (bInitialSetupHasBeenPerformed)
+		{
+			UnregisterFromRepositoryDelegates();
+		}
 
 		// reset unique link Id for new map
 		if (MyWorld->WorldType == EWorldType::Game || MyWorld->WorldType == EWorldType::Editor)
@@ -4969,7 +5293,14 @@ bool UNavigationSystemV1::IsNavigationBeingBuiltOrLocked(UObject* WorldContextOb
 
 bool UNavigationSystemV1::K2_ReplaceAreaInOctreeData(const UObject* Object, TSubclassOf<UNavArea> OldArea, TSubclassOf<UNavArea> NewArea)
 {
-	return Object ? ReplaceAreaInOctreeData(*Object, OldArea, NewArea) : false;
+	SCOPE_CYCLE_COUNTER(STAT_NavOctreeBookkeeping);
+
+	checkf(Repository, TEXT("%hs is expected to be called after the repository gets cached."), __FUNCTION__);
+	if (const FNavigationElementHandle Handle = Repository->GetNavigationElementHandleForUObject(Object))
+	{
+		return ReplaceAreaInOctreeData(Handle, OldArea, NewArea);
+	}
+	return false;
 }
 
 //----------------------------------------------------------------------//
@@ -5005,7 +5336,7 @@ bool UNavigationSystemV1::HandleCountNavMemCommand()
 
 	if (DefaultOctreeController.NavOctree.IsValid())
 	{
-		UE_LOG(LogNavigation, Warning, TEXT("NavOctree memory: %d"), DefaultOctreeController.NavOctree->GetSizeBytes());
+		UE_LOG(LogNavigation, Warning, TEXT("NavOctree memory: %llu"), DefaultOctreeController.NavOctree->GetSizeBytes());
 	}
 
 	for (int32 NavDataIndex = 0; NavDataIndex < NavDataSet.Num(); ++NavDataIndex)
@@ -5289,41 +5620,35 @@ void UNavigationSystemV1::UnregisterInvoker_Internal(const UObject& Invoker)
 
 void UNavigationSystemV1::RegisterToRepositoryDelegates()
 {
-	if (UNavigationObjectRepository* Repository = UWorld::GetSubsystem<UNavigationObjectRepository>(GetWorld()))
-	{
-		Repository->OnCustomNavLinkObjectRegistered.BindWeakLambda(this, [this](INavLinkCustomInterface& CustomLink)
-			{
-				RegisterCustomLink(CustomLink);
-			});
+	checkf(Repository, TEXT("%hs is expected to be called after the repository gets cached."), __FUNCTION__);
+	Repository->OnCustomNavLinkObjectRegistered.BindWeakLambda(this, [this](INavLinkCustomInterface& CustomLink)
+		{
+			RegisterCustomLink(CustomLink);
+		});
 
-		Repository->OnCustomNavLinkObjectUnregistered.BindWeakLambda(this, [this](INavLinkCustomInterface& CustomLink)
-			{
-				UnregisterCustomLink(CustomLink);
-			});
+	Repository->OnCustomNavLinkObjectUnregistered.BindWeakLambda(this, [this](INavLinkCustomInterface& CustomLink)
+		{
+			UnregisterCustomLink(CustomLink);
+		});
 
-		Repository->OnNavRelevantObjectRegistered.BindWeakLambda(this, [this](INavRelevantInterface& NavRelevantObject)
-			{
-				SCOPE_CYCLE_COUNTER(STAT_DebugNavOctree);
-				RegisterNavOctreeElement(Cast<UObject>(&NavRelevantObject), &NavRelevantObject, FNavigationOctreeController::OctreeUpdate_Default);
-			});
+	Repository->OnNavigationElementAddedDelegate.BindWeakLambda(this, [this](const TSharedRef<const FNavigationElement>& Element)
+		{
+			RegisterNavigationElementWithNavOctree(Element, FNavigationOctreeController::OctreeUpdate_Default);
+		});
 
-		Repository->OnNavRelevantObjectUnregistered.BindWeakLambda(this, [this](INavRelevantInterface& NavRelevantObject)
-			{
-				SCOPE_CYCLE_COUNTER(STAT_DebugNavOctree);
-				UnregisterNavOctreeElement(Cast<UObject>(&NavRelevantObject), &NavRelevantObject, FNavigationOctreeController::OctreeUpdate_Default);
-			});
-	}
+	Repository->OnNavigationElementRemovedDelegate.BindWeakLambda(this, [this](const TSharedRef<const FNavigationElement>& Element)
+		{
+			UnregisterNavigationElementWithOctree(Element, FNavigationOctreeController::OctreeUpdate_Default);
+		});
 }
 
 void UNavigationSystemV1::UnregisterFromRepositoryDelegates() const
 {
-	if (UNavigationObjectRepository* Repository = UWorld::GetSubsystem<UNavigationObjectRepository>(GetWorld()))
-	{
-		Repository->OnCustomNavLinkObjectRegistered = nullptr;
-		Repository->OnCustomNavLinkObjectUnregistered = nullptr;
-		Repository->OnNavRelevantObjectRegistered = nullptr;
-		Repository->OnNavRelevantObjectUnregistered = nullptr;
-	}
+	checkf(Repository, TEXT("%hs is expected to be called after the repository gets cached."), __FUNCTION__);
+	Repository->OnCustomNavLinkObjectRegistered = nullptr;
+	Repository->OnCustomNavLinkObjectUnregistered = nullptr;
+	Repository->OnNavigationElementAddedDelegate = nullptr;
+	Repository->OnNavigationElementRemovedDelegate = nullptr;
 }
 
 void UNavigationSystemV1::UpdateInvokers()
@@ -5500,17 +5825,12 @@ void UNavigationSystemV1::UnregisterNavigationInvoker(AActor* Invoker)
 	}
 }
 
-//----------------------------------------------------------------------//
-// DEPRECATED
-//----------------------------------------------------------------------//
+// Deprecated
 bool UNavigationSystemV1::K2_GetRandomPointInNavigableRadius(UObject* WorldContextObject, const FVector& Origin, FVector& RandomLocation, float Radius, ANavigationData* NavData, TSubclassOf<UNavigationQueryFilter> FilterClass)
 {
 	return K2_GetRandomLocationInNavigableRadius(WorldContextObject, Origin, RandomLocation, Radius, NavData, FilterClass);
 }
 
-//----------------------------------------------------------------------//
-// NEW STUFF!
-//----------------------------------------------------------------------//
 void UNavigationSystemV1::VerifyNavigationRenderingComponents(const bool bShow)
 {
 	// make sure nav mesh has a rendering component
@@ -5898,5 +6218,29 @@ void UNavigationSystemModuleConfig::PostEditChangeProperty(FPropertyChangedEvent
 }
 #endif // WITH_EDITOR
 
-#undef LOCTEXT_NAMESPACE
+//----------------------------------------------------------------------//
+// Deprecated methods
+//----------------------------------------------------------------------//
+PRAGMA_DISABLE_DEPRECATION_WARNINGS
+
+// Deprecated
+uint32 UNavigationSystemV1::HashObject(const UObject& Object)
+{
+	return FNavigationOctree::HashObject(Object);
+}
 
+// Deprecated
+const FOctreeElementId2* UNavigationSystemV1::GetObjectsNavOctreeId(const UObject& Object) const
+{
+	return GetNavOctreeIdForElement(FNavigationElementHandle(&Object));
+}
+
+// Deprecated
+bool UNavigationSystemV1::HasPendingObjectNavOctreeId(UObject* Object) const
+{
+	return HasPendingUpdateForElement(FNavigationElementHandle(Object));
+}
+
+PRAGMA_ENABLE_DEPRECATION_WARNINGS
+
+#undef LOCTEXT_NAMESPACE
diff --git a/Engine/Source/Runtime/NavigationSystem/Public/NavMesh/RecastGeometryExport.h b/Engine/Source/Runtime/NavigationSystem/Public/NavMesh/RecastGeometryExport.h
--- a/Engine/Source/Runtime/NavigationSystem/Public/NavMesh/RecastGeometryExport.h
+++ b/Engine/Source/Runtime/NavigationSystem/Public/NavMesh/RecastGeometryExport.h
@@ -5,6 +5,7 @@
 #include "Engine/EngineTypes.h"
 
 struct FNavigationRelevantData;
+struct FNavigationElement;
 
 #if WITH_RECAST
 
@@ -43,8 +44,8 @@ struct FRecastGeometryExport : public FNavigableGeometryExport
 	/** Store Vertex and Index buffer data in associated FNavigationRelevantData. */
 	NAVIGATIONSYSTEM_API void StoreCollisionCache();
 
-	/** Collect the collision information from an INavRelevantInterface and store it into the FNavigationRelevantData's CollisionData. */
-	static NAVIGATIONSYSTEM_API void ExportNavRelevantObjectGeometry(INavRelevantInterface& InOutNavRelevantInterface, FNavigationRelevantData& OutData);
+	/** Collects the collision information from a navigation element and stores it into the FNavigationRelevantData's CollisionData. */
+	static NAVIGATIONSYSTEM_API void ExportElementGeometry(const FNavigationElement& InElement, FNavigationRelevantData& OutData);
 
 	/** Convert a list of vertices into the navigation format and store it into the FNavigationRelevantData's CollisionData.
 	 * @param InVerts Array of triangles vertices position. Each triangle will be created from 3 consecutive vertices in the array. Its size must be a multiple a 3.*/
diff --git a/Engine/Source/Runtime/NavigationSystem/Public/NavMesh/RecastNavMeshGenerator.h b/Engine/Source/Runtime/NavigationSystem/Public/NavMesh/RecastNavMeshGenerator.h
--- a/Engine/Source/Runtime/NavigationSystem/Public/NavMesh/RecastNavMeshGenerator.h
+++ b/Engine/Source/Runtime/NavigationSystem/Public/NavMesh/RecastNavMeshGenerator.h
@@ -45,6 +45,7 @@ struct FTileCacheCompressor;
 struct FTileCacheAllocator;
 struct FTileGenerationContext;
 struct dtLinkBuilderData;
+struct FNavigationElement;
 class dtNavMesh;
 class FNavRegenTimeSliceManager;
 class UNavigationSystemV1;
@@ -516,8 +517,11 @@ protected:
 
 	NAVIGATIONSYSTEM_API void AppendModifier(const FCompositeNavModifier& Modifier, const FNavDataPerInstanceTransformDelegate& InTransformsDelegate);
 	/** Appends specified geometry to tile's geometry */
+	NAVIGATIONSYSTEM_API void ValidateAndAppendGeometry(const FNavigationRelevantData& ElementData, const FCompositeNavModifier& InModifier);
+	UE_DEPRECATED(5.5, "Use the version taking a const reference on FNavigationRelevantData.")
 	NAVIGATIONSYSTEM_API void ValidateAndAppendGeometry(const TSharedRef<FNavigationRelevantData, ESPMode::ThreadSafe>& ElementData, const FCompositeNavModifier& InModifier);
-	NAVIGATIONSYSTEM_API void AppendGeometry(const FNavigationRelevantData& DataRef, const FCompositeNavModifier& InModifier, const FNavDataPerInstanceTransformDelegate& InTransformsDelegate);
+
+	NAVIGATIONSYSTEM_API void AppendGeometry(const FNavigationRelevantData& ElementData, const FCompositeNavModifier& InModifier, const FNavDataPerInstanceTransformDelegate& InTransformsDelegate);
 	
 	/** prepare voxel cache from collision data */
 	NAVIGATIONSYSTEM_API void PrepareVoxelCache(const TNavStatArray<uint8>& RawCollisionCache, const FCompositeNavModifier& InModifier, TNavStatArray<rcSpanCache>& SpanData);
@@ -845,7 +849,7 @@ public:
 	NAVIGATIONSYSTEM_API virtual void GrabDebugSnapshot(struct FVisualLogEntry* Snapshot, const FBox& BoundingBox, const FName& CategoryName, ELogVerbosity::Type Verbosity) const override;
 #endif
 
-	UE_DEPRECATED(5.4, "Use FRecastGeometryExport::ExportNavRelevantObjectGeometry")
+	UE_DEPRECATED(5.4, "Use FRecastGeometryExport::ExportElementGeometry")
 	static NAVIGATIONSYSTEM_API void ExportComponentGeometry(UActorComponent* InOutComponent, FNavigationRelevantData& OutData);
 
 	UE_DEPRECATED(5.4, "Use FRecastGeometryExport::ExportRigidBodyGeometry.")
@@ -873,7 +877,7 @@ public:
 		TNavStatArray<int32>& OutShapeBuffer,
 		const FTransform& LocalToWorld = FTransform::Identity);
 
-	UE_DEPRECATED(5.5, "Use FRecastGeometryExport::ExportNavRelevantObjectGeometry.")
+	UE_DEPRECATED(5.5, "Use FRecastGeometryExport::ExportElementGeometry.")
 	static NAVIGATIONSYSTEM_API void ExportNavRelevantObjectGeometry(INavRelevantInterface& InOutNavRelevantInterface, FNavigationRelevantData& OutData);
 	UE_DEPRECATED(5.5, "Use FRecastGeometryExport::ExportVertexSoupGeometry.")
 	static NAVIGATIONSYSTEM_API void ExportVertexSoupGeometry(const TArray<FVector>& InVerts, FNavigationRelevantData& OutData);
@@ -943,8 +947,19 @@ protected:
 	/** Marks grid tiles affected by specified areas as dirty */
 	NAVIGATIONSYSTEM_API virtual void MarkDirtyTiles(const TArray<FNavigationDirtyArea>& DirtyAreas);
 
-	/** Returns if the provided UObject that requested a navmesh dirtying should dirty this Navmesh. Useful to avoid tiles regeneration from objects that are excluded from the provided NavDataConfig */
-	NAVIGATIONSYSTEM_API virtual bool ShouldDirtyTilesRequestedByObject(const UNavigationSystemV1& NavSys, const FNavigationOctree& NavOctreeInstance, const UObject& SourceObject, const FNavDataConfig& NavDataConfig) const;
+	UE_DEPRECATED(5.5, "Use ShouldDirtyTilesRequestedByElement with FNavigationElement instead.")
+	NAVIGATIONSYSTEM_API virtual bool ShouldDirtyTilesRequestedByObject(
+		const UNavigationSystemV1& NavSys,
+		const FNavigationOctree& NavOctreeInstance,
+		const UObject& SourceObject,
+		const FNavDataConfig& NavDataConfig) const final;
+
+	/** Returns if the provided FNavigationElement that requested a navmesh dirtying should dirty this Navmesh. Useful to avoid tiles regeneration from elements that are excluded from the provided NavDataConfig */
+	NAVIGATIONSYSTEM_API virtual bool ShouldDirtyTilesRequestedByElement(
+		const UNavigationSystemV1& NavSys,
+		const FNavigationOctree& NavOctreeInstance,
+		FNavigationElementHandle SourceElement,
+		const FNavDataConfig& NavDataConfig) const;
 
 	/** Marks all tiles overlapping with InclusionBounds dirty (via MarkDirtyTiles). */
 	NAVIGATIONSYSTEM_API bool MarkNavBoundsDirty();
diff --git a/Engine/Source/Runtime/NavigationSystem/Public/NavigationDataHandler.h b/Engine/Source/Runtime/NavigationSystem/Public/NavigationDataHandler.h
--- a/Engine/Source/Runtime/NavigationSystem/Public/NavigationDataHandler.h
+++ b/Engine/Source/Runtime/NavigationSystem/Public/NavigationDataHandler.h
@@ -1,6 +1,7 @@
 // Copyright Epic Games, Inc. All Rights Reserved.
 
 #pragma once
+
 #include "NavigationOctreeController.h"
 #include "NavigationDirtyAreasController.h"
 
@@ -16,30 +17,65 @@ struct FNavigationDataHandler
 
 	NAVIGATIONSYSTEM_API void ConstructNavOctree(const FVector& Origin, const double Radius, const ENavDataGatheringModeConfig DataGatheringMode, const float GatheringNavModifiersWarningLimitTime);
 
+	UE_DEPRECATED(5.5, "Use RegisterFromNavOctree instead.")
 	NAVIGATIONSYSTEM_API void RemoveNavOctreeElementId(const FOctreeElementId2& ElementId, int32 UpdateFlags);
+	/**
+	 * Removes the octree node and the NavigationElementHandle-OctreeElementId pair associated to the specified OctreeElementId.
+	 * It will also dirty the area base of the NavigationElement values and the specified update flags.
+	 */
+	NAVIGATIONSYSTEM_API void RemoveFromNavOctree(const FOctreeElementId2& ElementId, int32 UpdateFlags);
+
+	NAVIGATIONSYSTEM_API FSetElementId RegisterElementWithNavOctree(const TSharedRef<const FNavigationElement>& ElementRef, int32 UpdateFlags);
+	UE_DEPRECATED(5.5, "Use RegisterNavigationElementWithNavOctree instead.")
 	NAVIGATIONSYSTEM_API FSetElementId RegisterNavOctreeElement(UObject& ElementOwner, INavRelevantInterface& ElementInterface, int32 UpdateFlags);
 	NAVIGATIONSYSTEM_API void AddElementToNavOctree(const FNavigationDirtyElement& DirtyElement);
 
-	/** Removes associated NavOctreeElement and invalidates associated pending updates. Also removes object from the list of children
-	* of the NavigationParent, if any.
-	* @param ElementOwner		Object for which we must remove associated NavOctreeElement
-	* @param ElementInterface	Object associated NavRelevantInterface to access NavigationParent when registered as children
-	* @param UpdateFlags		Flags indicating in which context the method is called to allow/forbid certain operations
-	*
-	* @return True if associated NavOctreeElement has been removed or pending update has been invalidated; false otherwise.
-	*/
+	/**
+	 * Removes associated NavOctreeElement and invalidates associated pending updates. Also removes element from the list of children
+	 * of the NavigationParent, if any.
+	 * @param ElementRef		Navigation element for which we must remove the associated NavOctreeElement
+	 * @param UpdateFlags		Flags indicating in which context the method is called to allow/forbid certain operations
+	 *
+	 * @return True if associated NavOctreeElement has been removed or pending update has been invalidated; false otherwise.
+	 */
+	NAVIGATIONSYSTEM_API bool UnregisterElementWithNavOctree(const TSharedRef<const FNavigationElement>& ElementRef, int32 UpdateFlags);
+	UE_DEPRECATED(5.5, "Use the overloaded version with FNavigationElement instead.")
 	NAVIGATIONSYSTEM_API bool UnregisterNavOctreeElement(UObject& ElementOwner, INavRelevantInterface& ElementInterface, int32 UpdateFlags);
+
+	/**
+	 * Unregister element associated with the provided handle and register the new element.
+	 * Also update any pending update associated to that element.
+	 */
+	NAVIGATIONSYSTEM_API void UpdateNavOctreeElement(FNavigationElementHandle ElementHandle, const TSharedRef<const FNavigationElement>& UpdatedElement, int32 UpdateFlags);
+	UE_DEPRECATED(5.5, "Use the overloaded version with FNavigationElement instead.")
 	NAVIGATIONSYSTEM_API void UpdateNavOctreeElement(UObject& ElementOwner, INavRelevantInterface& ElementInterface, int32 UpdateFlags);
+
+	UE_DEPRECATED(5.5, "This method is no longer public and should not be called directly.")
 	NAVIGATIONSYSTEM_API void UpdateNavOctreeParentChain(UObject& ElementOwner, bool bSkipElementOwnerUpdate = false);
+
+	NAVIGATIONSYSTEM_API bool UpdateNavOctreeElementBounds(FNavigationElementHandle Element, const FBox& NewBounds, const TConstArrayView<FBox> DirtyAreas);
+	UE_DEPRECATED(5.5, "Use the overloaded version with FNavigationElement instead.")
 	NAVIGATIONSYSTEM_API bool UpdateNavOctreeElementBounds(UObject& Object, const FBox& NewBounds, const TConstArrayView<FBox> DirtyAreas);
+
 	NAVIGATIONSYSTEM_API void FindElementsInNavOctree(const FBox& QueryBox, const FNavigationOctreeFilter& Filter, TArray<FNavigationOctreeElement>& Elements);
+
+	NAVIGATIONSYSTEM_API bool ReplaceAreaInOctreeData(FNavigationElementHandle Element, TSubclassOf<UNavArea> OldArea, TSubclassOf<UNavArea> NewArea, bool bReplaceChildClasses) const;
+	UE_DEPRECATED(5.5, "Use the overloaded version with FNavigationElement instead.")
 	NAVIGATIONSYSTEM_API bool ReplaceAreaInOctreeData(const UObject& Object, TSubclassOf<UNavArea> OldArea, TSubclassOf<UNavArea> NewArea, bool bReplaceChildClasses);
+
 	NAVIGATIONSYSTEM_API void AddLevelCollisionToOctree(ULevel& Level);
 	NAVIGATIONSYSTEM_API void RemoveLevelCollisionFromOctree(ULevel& Level);
+
+	UE_DEPRECATED(5.5, "This method will be removed. Use UNavigationSystemV1 version instead.")
 	NAVIGATIONSYSTEM_API void UpdateActorAndComponentsInNavOctree(AActor& Actor);
+
 	NAVIGATIONSYSTEM_API void ProcessPendingOctreeUpdates();
+
 	NAVIGATIONSYSTEM_API void DemandLazyDataGathering(FNavigationRelevantData& ElementData);
-	
+
 	UE_DEPRECATED(5.4, "Use the overloaded version taking a list of areas as parameter instead.")
 	NAVIGATIONSYSTEM_API bool UpdateNavOctreeElementBounds(UObject& Object, const FBox& NewBounds, const FBox& DirtyArea);
+
+private:
+	void UpdateNavOctreeParentChain(const TSharedRef<const FNavigationElement>& Element, bool bSkipElementOwnerUpdate = false);
 };
diff --git a/Engine/Source/Runtime/NavigationSystem/Public/NavigationDirtyAreasController.h b/Engine/Source/Runtime/NavigationSystem/Public/NavigationDirtyAreasController.h
--- a/Engine/Source/Runtime/NavigationSystem/Public/NavigationDirtyAreasController.h
+++ b/Engine/Source/Runtime/NavigationSystem/Public/NavigationDirtyAreasController.h
@@ -1,7 +1,9 @@
 // Copyright Epic Games, Inc. All Rights Reserved.
 
 #pragma once
+#if UE_ENABLE_INCLUDE_ORDER_DEPRECATED_IN_5_5
 #include "AI/Navigation/NavigationTypes.h"
+#endif // UE_ENABLE_INCLUDE_ORDER_DEPRECATED_IN_5_5
 #include "Containers/Array.h"
 #include "HAL/Platform.h"
 #include "Logging/LogMacros.h"
@@ -11,14 +13,14 @@
 #include "UObject/UnrealNames.h"
 
 class UObject;
+class ANavigationData;
+struct FNavigationDirtyElement;
+struct FNavigationElement;
 struct FNavigationDirtyArea;
 enum class ENavigationDirtyFlag : uint8;
 
 NAVIGATIONSYSTEM_API DECLARE_LOG_CATEGORY_EXTERN(LogNavigationDirtyArea, Warning, All);
 
-class ANavigationData;
-struct FNavigationDirtyElement;
-
 struct FNavigationDirtyAreasController
 {
 	/** update frequency for dirty areas on navmesh */
@@ -57,33 +59,33 @@ public:
 	NAVIGATIONSYSTEM_API void Tick(float DeltaSeconds, const TArray<ANavigationData*>& NavDataSet, bool bForceRebuilding = false);
 
 	/** Add a dirty area to the queue based on the provided bounds and flags.
-	 * Bounds must be valid and non empty otherwise the request will be ignored and a warning reported.
+	 * Bounds must be valid and non-empty otherwise the request will be ignored and a warning reported.
 	 * Accumulation must be allowed and flags valid otherwise the add is ignored.
 	 * @param NewArea Bounding box of the affected area
 	 * @param Flags Indicates the type of modification applied to the area
-	 * @param ObjectProviderFunc Optional function to retrieve source object that can be use for error reporting and navmesh exclusion
+	 * @param ElementProviderFunc Optional function to retrieve source element that can be used for error reporting and navmesh exclusion
 	 * @param DirtyElement Optional dirty element
 	 * @param DebugReason Source of the new area
 	 */
-	NAVIGATIONSYSTEM_API void AddArea(const FBox& NewArea, const ENavigationDirtyFlag Flags, const TFunction<UObject*()>& ObjectProviderFunc = nullptr,
+	NAVIGATIONSYSTEM_API void AddArea(const FBox& NewArea, const ENavigationDirtyFlag Flags, const TFunction<const TSharedPtr<const FNavigationElement>()>& ElementProviderFunc = nullptr,
 		const FNavigationDirtyElement* DirtyElement = nullptr, const FName& DebugReason = NAME_None);
-	UE_DEPRECATED(5.5, "Use the version taking ENavigationDirtyFlag instead.")
+	UE_DEPRECATED(5.5, "Use the version taking ENavigationDirtyFlag and FNavigationElement instead.")
 	NAVIGATIONSYSTEM_API void AddArea(const FBox& NewArea, const int32 Flags, const TFunction<UObject*()>& ObjectProviderFunc = nullptr,
 		const FNavigationDirtyElement* DirtyElement = nullptr, const FName& DebugReason = NAME_None);
 
-	/** Add non empty list of dirty areas to the queue based on the provided bounds and flags.
-	 * Bounds must be valid and non empty otherwise the request will be ignored and a warning reported.
+	/** Add non-empty list of dirty areas to the queue based on the provided bounds and flags.
+	 * Bounds must be valid and non-empty otherwise the request will be ignored and a warning reported.
 	 * Accumulation must be allowed and flags valid otherwise the add is ignored.
 	 * A check will be triggered if an empty array is provided.
 	 * @param NewAreas Array of bounding boxes of the affected areas
 	 * @param Flags Indicates the type of modification applied to the area
-	 * @param ObjectProviderFunc Optional function to retrieve source object that can be use for error reporting and navmesh exclusion
+	 * @param ElementProviderFunc Optional function to retrieve source element that can be used for error reporting and navmesh exclusion
 	 * @param DirtyElement Optional dirty element
 	 * @param DebugReason Source of the new area
 	 */
-	NAVIGATIONSYSTEM_API void AddAreas(const TConstArrayView<FBox> NewAreas, const ENavigationDirtyFlag Flags, const TFunction<UObject*()>& ObjectProviderFunc = nullptr,
+	NAVIGATIONSYSTEM_API void AddAreas(const TConstArrayView<FBox> NewAreas, const ENavigationDirtyFlag Flags, const TFunction<const TSharedPtr<const FNavigationElement>()>& ElementProviderFunc = nullptr,
 		const FNavigationDirtyElement* DirtyElement = nullptr, const FName& DebugReason = NAME_None);
-	UE_DEPRECATED(5.5, "Use the version taking ENavigationDirtyFlag instead.")
+	UE_DEPRECATED(5.5, "Use the version taking ENavigationDirtyFlag and FNavigationElement instead.")
 	NAVIGATIONSYSTEM_API void AddAreas(const TConstArrayView<FBox> NewAreas, const int32 Flags, const TFunction<UObject*()>& ObjectProviderFunc = nullptr,
 		const FNavigationDirtyElement* DirtyElement = nullptr, const FName& DebugReason = NAME_None);
 	
@@ -101,6 +103,6 @@ public:
 	bool HadDirtyAreasReportedWhileAccumulationLocked() const { return bCanAccumulateDirtyAreas == false && bDirtyAreasReportedWhileAccumulationLocked; }
 #endif // UE_BUILD_SHIPPING
 
-	DECLARE_DELEGATE_RetVal_OneParam(bool, FSkipObjectSignature, UObject& /*Object*/);
-	FSkipObjectSignature ShouldSkipObjectPredicate;  
+	DECLARE_DELEGATE_RetVal_OneParam(bool, FSkipObjectSignature, const UObject& /*Object*/);
+	FSkipObjectSignature ShouldSkipObjectPredicate;
 };
diff --git a/Engine/Source/Runtime/NavigationSystem/Public/NavigationOctree.h b/Engine/Source/Runtime/NavigationSystem/Public/NavigationOctree.h
--- a/Engine/Source/Runtime/NavigationSystem/Public/NavigationOctree.h
+++ b/Engine/Source/Runtime/NavigationSystem/Public/NavigationOctree.h
@@ -5,13 +5,18 @@
 #if UE_ENABLE_INCLUDE_ORDER_DEPRECATED_IN_5_4
 #include "CoreMinimal.h"
 #endif
+
+#if UE_ENABLE_INCLUDE_ORDER_DEPRECATED_IN_5_5
 #include "Stats/Stats.h"
 #include "AI/Navigation/NavigationTypes.h"
+#include "AI/NavigationModifier.h"
+#include "AI/Navigation/NavRelevantInterface.h"
+#endif // UE_ENABLE_INCLUDE_ORDER_DEPRECATED_IN_5_5
+
 #include "Math/GenericOctreePublic.h"
 #include "NavigationSystemTypes.h"
 #include "EngineStats.h"
-#include "AI/NavigationModifier.h"
-#include "AI/Navigation/NavRelevantInterface.h"
+#include "AI/Navigation/NavigationElement.h"
 #include "AI/Navigation/NavigationRelevantData.h"
 #include "Math/GenericOctree.h"
 #include "HAL/LowLevelMemTracker.h"
@@ -19,6 +24,8 @@
 class INavRelevantInterface;
 class FNavigationOctree;
 class UActorComponent;
+enum class ENavDataGatheringMode : uint8;
+struct FNavigationRelevantDataFilter;
 typedef FNavigationRelevantDataFilter FNavigationOctreeFilter;
 
 LLM_DECLARE_TAG(NavigationOctree);
@@ -27,24 +34,35 @@ struct FNavigationOctreeElement
 {
 	FBoxSphereBounds Bounds;
 	TSharedRef<FNavigationRelevantData, ESPMode::ThreadSafe> Data;
-	uint32 OwnerUniqueId = INDEX_NONE;	
 
-public:
+#if WITH_EDITORONLY_DATA
+	UE_DEPRECATED(5.5, "ID no longer used.")
+	uint32 OwnerUniqueId = INDEX_NONE;
+#endif // WITH_EDITORONLY_DATA
+
+	UE_DEPRECATED(5.5, "Use the constructor taking a FNavigationElement instead.")
 	explicit FNavigationOctreeElement(UObject& SourceObject)
-		: Data(new FNavigationRelevantData(SourceObject))
+		: Data(MakeShared<FNavigationRelevantData>(MakeShared<const FNavigationElement>(SourceObject)))
+#if WITH_EDITORONLY_DATA
 		, OwnerUniqueId(SourceObject.GetUniqueID())
-	{}
+#endif // WITH_EDITORONLY_DATA
+	{
+	}
+
+	explicit FNavigationOctreeElement(const TSharedRef<const FNavigationElement>& SourceObject)
+		: Bounds(SourceObject.Get().GetBounds())
+		, Data(MakeShared<FNavigationRelevantData>(SourceObject))
+	{
+	}
 
 	FNavigationOctreeElement(const FNavigationOctreeElement& Other)
 		: Bounds(Other.Bounds)
 		, Data(Other.Data)
-		, OwnerUniqueId(Other.OwnerUniqueId)
 	{}
 
 	FNavigationOctreeElement(FNavigationOctreeElement&& Other) noexcept
-	: Bounds(MoveTemp(Other.Bounds))
-	, Data(MoveTemp(Other.Data))
-	, OwnerUniqueId(MoveTemp(Other.OwnerUniqueId))
+		: Bounds(MoveTemp(Other.Bounds))
+		, Data(MoveTemp(Other.Data))
 	{
 	}
 
@@ -53,7 +71,7 @@ public:
 		if (this != &Other)
 		{
 			this->~FNavigationOctreeElement();
-			new(this) FNavigationOctreeElement(Forward<FNavigationOctreeElement>(Other));			
+			new(this) FNavigationOctreeElement(Forward<FNavigationOctreeElement>(Other));
 		}
 		return *this;
 	}
@@ -84,7 +102,7 @@ public:
 	 *	If it does then retrieves a copy with meta areas substituted with
 	 *	appropriate non-meta areas, depending on NavAgent
 	 */
-	FORCEINLINE FCompositeNavModifier GetModifierForAgent(const struct FNavAgentProperties* NavAgent = nullptr) const 
+	FORCEINLINE FCompositeNavModifier GetModifierForAgent(const FNavAgentProperties* NavAgent = nullptr) const
 	{ 
 		return Data->GetModifierForAgent(NavAgent);
 	}
@@ -109,7 +127,16 @@ public:
 		Data->ValidateAndShrink();
 	}
 
-	FORCEINLINE UObject* GetOwner(bool bEvenIfPendingKill = false) const { return Data->SourceObject.Get(bEvenIfPendingKill); }
+	UE_DEPRECATED(5.5, "Use GetSourceElement instead.")
+	FORCEINLINE UObject* GetOwner(bool bEvenIfPendingKill = false) const
+	{
+		return const_cast<UObject*>(Data.Get().SourceElement.Get().GetWeakUObject().Get());
+	}
+
+	const TSharedRef<const FNavigationElement>& GetSourceElement() const
+	{
+		return Data->SourceElement;
+	}
 };
 
 struct FNavigationOctreeSemantics
@@ -120,7 +147,6 @@ struct FNavigationOctreeSemantics
 	enum { MaxNodeDepth = 12 };
 
 	typedef TInlineAllocator<MaxElementsPerLeaf> ElementAllocator;
-	//typedef FDefaultAllocator ElementAllocator;
 
 	FORCEINLINE static const FBoxSphereBounds& GetBoundingBox(const FNavigationOctreeElement& NavData)
 	{
@@ -129,7 +155,7 @@ struct FNavigationOctreeSemantics
 
 	FORCEINLINE static bool AreElementsEqual(const FNavigationOctreeElement& A, const FNavigationOctreeElement& B)
 	{
-		return A.Data->SourceObject == B.Data->SourceObject;
+		return A.Data.Get().SourceElement.Get().GetHandle() == B.Data.Get().SourceElement.Get().GetHandle();
 	}
 
 	FORCEINLINE static void ApplyOffset(FNavigationOctreeElement& Element, const FVector& InOffset)
@@ -153,8 +179,14 @@ PRAGMA_DISABLE_DEPRECATION_WARNINGS
 	FNavigableGeometryComponentExportDelegate ComponentExportDelegate;
 PRAGMA_ENABLE_DEPRECATION_WARNINGS
 
+	UE_DEPRECATED(5.5, "Use FGeometryExportDelegate.")
 	DECLARE_DELEGATE_TwoParams(FNavRelevantGeometryExportDelegate, INavRelevantInterface&, FNavigationRelevantData&);
+PRAGMA_DISABLE_DEPRECATION_WARNINGS
 	FNavRelevantGeometryExportDelegate NavRelevantGeometryExportDelegate;
+PRAGMA_ENABLE_DEPRECATION_WARNINGS
+
+	DECLARE_DELEGATE_TwoParams(FGeometryExportDelegate, const FNavigationElement&, FNavigationRelevantData&);
+    FGeometryExportDelegate GeometryExportDelegate;
 
 	enum ENavGeometryStoringMode {
 		SkipNavGeometry,
@@ -162,9 +194,9 @@ PRAGMA_ENABLE_DEPRECATION_WARNINGS
 	};
 
 	/**
- * Adds an element to the octree.
- * @param Element - The element to add.
- */
+	 * Adds an element to the octree.
+	 * @param Element - The element to add.
+	 */
 	inline void AddElement(const FNavigationOctreeElement& Element)
 	{
 		LLM_SCOPE_BYTAG(NavigationOctree);
@@ -195,9 +227,13 @@ PRAGMA_ENABLE_DEPRECATION_WARNINGS
 	NAVIGATIONSYSTEM_API virtual ~FNavigationOctree();
 
 	/** Add new node and fill it with navigation export data */
+	NAVIGATIONSYSTEM_API void AddNode(const FBox& Bounds, FNavigationOctreeElement& OctreeElement);
+	UE_DEPRECATED(5.5, "Use the overloaded version with only FBox and FNavigationOctreeElement instead.")
 	NAVIGATIONSYSTEM_API void AddNode(UObject* ElementOb, INavRelevantInterface* NavElement, const FBox& Bounds, FNavigationOctreeElement& Data);
 
 	/** Append new data to existing node */
+	NAVIGATIONSYSTEM_API void AppendToNode(const FOctreeElementId2& Id, const TSharedRef<const FNavigationElement>& ElementRef, const FBox& Bounds, FNavigationOctreeElement& Data);
+	UE_DEPRECATED(5.5, "Use the overloaded version with FNavigationElement instead.")
 	NAVIGATIONSYSTEM_API void AppendToNode(const FOctreeElementId2& Id, INavRelevantInterface* NavElement, const FBox& Bounds, FNavigationOctreeElement& Data);
 
 	/** Updates element bounds remove/add operation */
@@ -210,38 +246,54 @@ PRAGMA_ENABLE_DEPRECATION_WARNINGS
 
 	NAVIGATIONSYSTEM_API const FNavigationRelevantData* GetDataForID(const FOctreeElementId2& Id) const;
 
+	NAVIGATIONSYSTEM_API FNavigationRelevantData* GetMutableDataForID(const FOctreeElementId2& Id);
+
 	ENavGeometryStoringMode GetNavGeometryStoringMode() const
 	{
 		return bGatherGeometry ? StoreNavGeometry : SkipNavGeometry;
 	}
 
 	NAVIGATIONSYSTEM_API void SetDataGatheringMode(ENavDataGatheringModeConfig Mode);
-	
+
 	// Lazy data gathering methods
+	NAVIGATIONSYSTEM_API bool IsLazyGathering(const FNavigationElement& NavigationElement) const;
+	UE_DEPRECATED(5.5, "Use the overloaded version with FNavigationElement instead.")
 	NAVIGATIONSYSTEM_API bool IsLazyGathering(const INavRelevantInterface& ChildNavInterface) const;
+
 	NAVIGATIONSYSTEM_API void DemandLazyDataGathering(FNavigationRelevantData& ElementData);
+
+	NAVIGATIONSYSTEM_API void DemandChildLazyDataGathering(FNavigationRelevantData& ElementData, const FNavigationElement& ChildElement) const;
+	UE_DEPRECATED(5.5, "Use the overloaded version with FNavigationElement instead.")
 	NAVIGATIONSYSTEM_API void DemandChildLazyDataGathering(FNavigationRelevantData& ElementData, INavRelevantInterface& ChildNavInterface);
 
+	UE_DEPRECATED(5.5, "This method is no longer used by the navigation system.")
 	FORCEINLINE static uint32 HashObject(const UObject& Object)
 	{
 		return Object.GetUniqueID();
 	}
-#if !UE_BUILD_SHIPPING	
+
+#if !UE_BUILD_SHIPPING
 	NAVIGATIONSYSTEM_API void SetGatheringNavModifiersTimeLimitWarning(const float Threshold);
-#endif // !UE_BUILD_SHIPPING	
+#endif // !UE_BUILD_SHIPPING
 protected:
 	friend struct FNavigationOctreeController;
 	friend struct FNavigationOctreeSemantics;
 
+	UE_DEPRECATED(5.5, "Use the version taking a FNavigationElementHandle instead.")
 	NAVIGATIONSYSTEM_API void SetElementIdImpl(const uint32 OwnerUniqueId, FOctreeElementId2 Id);
+	NAVIGATIONSYSTEM_API void SetElementIdImpl(FNavigationElementHandle ElementHandle, FOctreeElementId2 Id);
 
+#if WITH_EDITORONLY_DATA
+	UE_DEPRECATED(5.5, "Use ElementToOctreeId instead.")
 	TMap<uint32, FOctreeElementId2> ObjectToOctreeId;
+#endif // WITH_EDITORONLY_DATA
+	TMap<FNavigationElementHandle, FOctreeElementId2> ElementToOctreeId;
 	ENavDataGatheringMode DefaultGeometryGatheringMode;
 	uint32 bGatherGeometry : 1;
 	uint32 NodesMemory;
-#if !UE_BUILD_SHIPPING	
+#if !UE_BUILD_SHIPPING
 	float GatheringNavModifiersTimeLimitWarning;
-#endif // !UE_BUILD_SHIPPING	
+#endif // !UE_BUILD_SHIPPING
 private:
 	SIZE_T OctreeSizeBytes = 0;
 };
diff --git a/Engine/Source/Runtime/NavigationSystem/Public/NavigationOctreeController.h b/Engine/Source/Runtime/NavigationSystem/Public/NavigationOctreeController.h
--- a/Engine/Source/Runtime/NavigationSystem/Public/NavigationOctreeController.h
+++ b/Engine/Source/Runtime/NavigationSystem/Public/NavigationOctreeController.h
@@ -4,6 +4,25 @@
 #include "NavigationOctree.h"
 #include "AI/Navigation/NavigationDirtyElement.h"
 
+
+struct FNavigationDirtyElementKeyFunctions : BaseKeyFuncs<FNavigationDirtyElement, FNavigationElementHandle, /*bInAllowDuplicateKeys*/false>
+{
+	static FNavigationElementHandle GetSetKey(ElementInitType Element)
+	{
+		return Element.NavigationElement->GetHandle();
+	}
+
+	static bool Matches(KeyInitType A, KeyInitType B)
+	{
+		return A == B;
+	}
+
+	static uint32 GetKeyHash(KeyInitType Key)
+	{
+		return GetTypeHash(Key);
+	}
+};
+
 struct FNavigationOctreeController
 {
 	enum EOctreeUpdateMode
@@ -15,69 +34,110 @@ struct FNavigationOctreeController
 		OctreeUpdate_ParentChain = 8,					// update child nodes, don't remove anything
 	};
 
+	UE_DEPRECATED(5.5, "Use PendingUpdates instead.")
 	TSet<FNavigationDirtyElement> PendingOctreeUpdates;
+	TSet<FNavigationDirtyElement, FNavigationDirtyElementKeyFunctions> PendingUpdates;
 	TSharedPtr<FNavigationOctree, ESPMode::ThreadSafe> NavOctree;
-	/** Map of all objects that are tied to indexed navigation parent */
+
+	UE_DEPRECATED(5.5, "This container is no longer used. Use AddChild/RemoveChild/GetChildren methods instead.")
 	TMultiMap<UObject*, FWeakObjectPtr> OctreeChildNodesMap;
+
 	/** if set, navoctree updates are ignored, use with caution! */
-	uint8 bNavOctreeLock : 1;
+	uint8 bNavOctreeLock : 1 = false;
+
+	inline void SetNavigationOctreeLock(bool bLock);
 
-	NAVIGATIONSYSTEM_API void SetNavigationOctreeLock(bool bLock);
+	NAVIGATIONSYSTEM_API bool HasPendingUpdateForElement(FNavigationElementHandle Element) const;
+	UE_DEPRECATED(5.5, "Use HasPendingUpdateForElement instead.")
 	NAVIGATIONSYSTEM_API bool HasPendingObjectNavOctreeId(UObject& Object) const;
+
+	inline void RemoveNode(FOctreeElementId2 ElementId, FNavigationElementHandle GetHandle);
+
+	UE_DEPRECATED(5.5, "Use RemoveNode instead.")
 	NAVIGATIONSYSTEM_API void RemoveObjectsNavOctreeId(const UObject& Object);
+
 	NAVIGATIONSYSTEM_API void SetNavigableGeometryStoringMode(FNavigationOctree::ENavGeometryStoringMode NavGeometryMode);
 
 	NAVIGATIONSYSTEM_API void Reset();
 
-	NAVIGATIONSYSTEM_API const FNavigationOctree* GetOctree() const;
-	NAVIGATIONSYSTEM_API FNavigationOctree* GetMutableOctree();
+	inline const FNavigationOctree* GetOctree() const;
+	inline FNavigationOctree* GetMutableOctree();
+
+	inline const FOctreeElementId2* GetNavOctreeIdForElement(FNavigationElementHandle Element) const;
+	UE_DEPRECATED(5.5, "Use GetNavOctreeIdForElement instead.")
 	NAVIGATIONSYSTEM_API const FOctreeElementId2* GetObjectsNavOctreeId(const UObject& Object) const;
-	NAVIGATIONSYSTEM_API bool GetNavOctreeElementData(const UObject& NodeOwner, ENavigationDirtyFlag& OutDirtyFlags, FBox& OutDirtyBounds);
-	UE_DEPRECATED(5.5, "Use the version taking ENavigationDirtyFlag& instead.")
+
+	NAVIGATIONSYSTEM_API bool GetNavOctreeElementData(FNavigationElementHandle Element, ENavigationDirtyFlag& DirtyFlags, FBox& DirtyBounds);
+	UE_DEPRECATED(5.5, "Use the version taking ENavigationDirtyFlag& and FNavigationElementHandle as parameter instead.")
 	NAVIGATIONSYSTEM_API bool GetNavOctreeElementData(const UObject& NodeOwner, int32& DirtyFlags, FBox& DirtyBounds);
+
+	NAVIGATIONSYSTEM_API const FNavigationRelevantData* GetDataForElement(FNavigationElementHandle Element) const;
+	UE_DEPRECATED(5.5, "Use GetDataForElement instead.")
 	NAVIGATIONSYSTEM_API const FNavigationRelevantData* GetDataForObject(const UObject& Object) const;
+
+	NAVIGATIONSYSTEM_API FNavigationRelevantData* GetMutableDataForElement(FNavigationElementHandle Element);
+	UE_DEPRECATED(5.5, "Use GetMutableDataForElement instead.")
 	NAVIGATIONSYSTEM_API FNavigationRelevantData* GetMutableDataForObject(const UObject& Object);
+
+	inline bool HasElementNavOctreeId(const FNavigationElementHandle Element) const;
+	UE_DEPRECATED(5.5, "Use HasElementNavOctreeId instead.")
 	NAVIGATIONSYSTEM_API bool HasObjectsNavOctreeId(const UObject& Object) const;
-	NAVIGATIONSYSTEM_API bool IsNavigationOctreeLocked() const;
+
+	inline bool IsNavigationOctreeLocked() const;
+
 	/** basically says if navoctree has been created already */
 	bool IsValid() const { return NavOctree.IsValid(); }
-	NAVIGATIONSYSTEM_API bool IsValidElement(const FOctreeElementId2& ElementId) const;
-	bool IsEmpty() const { return (IsValid() == false) || NavOctree->GetSizeBytes() == 0; }
+
+	inline bool IsValidElement(const FOctreeElementId2* ElementId) const;
+	inline bool IsValidElement(const FOctreeElementId2& ElementId) const;
+
+	bool IsEmpty() const
+	{
+		return (IsValid() == false) || NavOctree->GetSizeBytes() == 0;
+	}
+
+	inline void AddChild(FNavigationElementHandle Parent, const TSharedRef<const FNavigationElement>& Child);
+	inline void RemoveChild(FNavigationElementHandle Parent, const TSharedRef<const FNavigationElement>& Child);
+	inline void GetChildren(FNavigationElementHandle Parent, TArray<const TSharedRef<const FNavigationElement>>& OutChildren) const;
 
 private:
+	UE_DEPRECATED(5.5, "This method will no be longer be used by the navigation system.")
 	static NAVIGATIONSYSTEM_API uint32 HashObject(const UObject& Object);
+
+	/** Map of all elements that are tied to indexed navigation parent */
+	TMultiMap<FNavigationElementHandle, const TSharedRef<const FNavigationElement>> OctreeParentChildNodesMap;
 };
 
 //----------------------------------------------------------------------//
 // inlines
 //----------------------------------------------------------------------//
-FORCEINLINE uint32 FNavigationOctreeController::HashObject(const UObject& Object) 
-{ 
+
+// deprecated
+FORCEINLINE uint32 FNavigationOctreeController::HashObject(const UObject& Object)
+{
+PRAGMA_DISABLE_DEPRECATION_WARNINGS
 	return FNavigationOctree::HashObject(Object);
+PRAGMA_ENABLE_DEPRECATION_WARNINGS
 }
 
-FORCEINLINE_DEBUGGABLE const FOctreeElementId2* FNavigationOctreeController::GetObjectsNavOctreeId(const UObject& Object) const
+FORCEINLINE_DEBUGGABLE const FOctreeElementId2* FNavigationOctreeController::GetNavOctreeIdForElement(const FNavigationElementHandle Element) const
 { 
 	return NavOctree.IsValid()
-		? NavOctree->ObjectToOctreeId.Find(HashObject(Object))
+		? NavOctree->ElementToOctreeId.Find(Element)
 		: nullptr;
 }
 
-FORCEINLINE_DEBUGGABLE bool FNavigationOctreeController::HasObjectsNavOctreeId(const UObject& Object) const
+FORCEINLINE_DEBUGGABLE bool FNavigationOctreeController::HasElementNavOctreeId(const FNavigationElementHandle Element) const
 {
-	return NavOctree.IsValid() && (NavOctree->ObjectToOctreeId.Find(HashObject(Object)) != nullptr);
-}
-
-FORCEINLINE bool FNavigationOctreeController::HasPendingObjectNavOctreeId(UObject& Object) const 
-{ 
-	return PendingOctreeUpdates.Contains(FNavigationDirtyElement(&Object)); 
+	return NavOctree.IsValid() && (NavOctree->ElementToOctreeId.Find(Element) != nullptr);
 }
 
-FORCEINLINE_DEBUGGABLE void FNavigationOctreeController::RemoveObjectsNavOctreeId(const UObject& Object)
+FORCEINLINE_DEBUGGABLE void FNavigationOctreeController::RemoveNode(const FOctreeElementId2 ElementId, const FNavigationElementHandle ElementHandle)
 { 
 	if (NavOctree.IsValid())
 	{
-		NavOctree->ObjectToOctreeId.Remove(HashObject(Object));
+		NavOctree->RemoveNode(ElementId);
+		NavOctree->ElementToOctreeId.Remove(ElementHandle);
 	}
 }
 
@@ -101,7 +161,27 @@ FORCEINLINE void FNavigationOctreeController::SetNavigationOctreeLock(bool bLock
 	bNavOctreeLock = bLock; 
 }
 
+FORCEINLINE bool FNavigationOctreeController::IsValidElement(const FOctreeElementId2* ElementId) const
+{
+	return ElementId && IsValidElement(*ElementId);
+}
+
 FORCEINLINE bool FNavigationOctreeController::IsValidElement(const FOctreeElementId2& ElementId) const
 {
 	return IsValid() && NavOctree->IsValidElementId(ElementId);
 }
+
+void FNavigationOctreeController::AddChild(const FNavigationElementHandle Parent, const TSharedRef<const FNavigationElement>& Child)
+{
+	OctreeParentChildNodesMap.AddUnique(Parent, Child);
+}
+
+void FNavigationOctreeController::RemoveChild(const FNavigationElementHandle Parent, const TSharedRef<const FNavigationElement>& Child)
+{
+	OctreeParentChildNodesMap.RemoveSingle(Parent, Child);
+}
+
+void FNavigationOctreeController::GetChildren(const FNavigationElementHandle Parent, TArray<const TSharedRef<const FNavigationElement>>& OutChildren) const
+{
+	OctreeParentChildNodesMap.MultiFind(Parent, OutChildren);
+}
\ No newline at end of file
diff --git a/Engine/Source/Runtime/NavigationSystem/Public/NavigationSystem.h b/Engine/Source/Runtime/NavigationSystem/Public/NavigationSystem.h
--- a/Engine/Source/Runtime/NavigationSystem/Public/NavigationSystem.h
+++ b/Engine/Source/Runtime/NavigationSystem/Public/NavigationSystem.h
@@ -27,11 +27,6 @@
 #include "Math/MovingWindowAverageFast.h"
 #include "AI/Navigation/NavigationBounds.h"
 #include "Containers/ContainerAllocationPolicies.h"
-#if UE_ENABLE_INCLUDE_ORDER_DEPRECATED_IN_5_2
-#if WITH_EDITOR
-#include "UnrealEdMisc.h"
-#endif // WITH_EDITOR
-#endif
 #include "NavigationSystem.generated.h"
 
 
@@ -46,9 +41,12 @@ class INavRelevantInterface;
 class UCrowdManagerBase;
 class UNavArea;
 class UNavigationPath;
+class UNavigationObjectRepository;
 class UNavigationSystemModuleConfig;
+struct FDebugText;
 struct FNavigationRelevantData;
 struct FNavigationOctreeElement;
+struct FNavigationElement;
 
 #if !UE_BUILD_SHIPPING
 #define ALLOW_TIME_SLICE_DEBUG 1
@@ -304,6 +302,8 @@ public:
 	NAVIGATIONSYSTEM_API UNavigationSystemV1(const FObjectInitializer& ObjectInitializer = FObjectInitializer::Get());
 	NAVIGATIONSYSTEM_API virtual ~UNavigationSystemV1();
 
+	NAVIGATIONSYSTEM_API void GatherDebugLabels(TArray<FString>& InOutDebugLabels) const;
+
 	UPROPERTY(Transient)
 	TObjectPtr<ANavigationData> MainNavData;
 
@@ -677,7 +677,7 @@ public:
 	 *	@param NavData If NavData == NULL then MainNavData is used.
 	 *	@return true if any location found, false otherwise */
 	NAVIGATIONSYSTEM_API bool GetRandomPointInNavigableRadius(const FVector& Origin, float Radius, FNavLocation& ResultLocation, ANavigationData* NavData = NULL, FSharedConstNavQueryFilter QueryFilter = NULL) const;
-	
+
 	/** Calculates a path from PathStart to PathEnd and retrieves its cost.
 	 *	@NOTE potentially expensive, so use it with caution */
 	NAVIGATIONSYSTEM_API ENavigationQueryResult::Type GetPathCost(const FVector& PathStart, const FVector& PathEnd, FVector::FReal& PathCost, const ANavigationData* NavData = NULL, FSharedConstNavQueryFilter QueryFilter = NULL) const;
@@ -877,6 +877,12 @@ protected:
 	/** @return pointer to ANavigationData instance of given ID, or NULL if it was not found. Note it looks only through registered navigation data */
 	NAVIGATIONSYSTEM_API ANavigationData* GetNavDataWithID(const uint16 NavDataID) const;
 
+	/** @return Valid shared pointer on a FNavigationElement if the provided UObject is registered in the navigation system, an invalid one otherwise. */
+	NAVIGATIONSYSTEM_API TSharedPtr<const FNavigationElement> GetNavigationElementForUObject(const UObject*) const;
+
+	/** @return Valid handle if the provided UObject is registered in the navigation system, FNavigationElementHandle::Invalid otherwise. */
+	NAVIGATIONSYSTEM_API FNavigationElementHandle GetNavigationElementHandleForUObject(const UObject*) const;
+
 	static NAVIGATIONSYSTEM_API void RegisterComponentToNavOctree(UActorComponent* Comp);
 	static NAVIGATIONSYSTEM_API void UnregisterComponentToNavOctree(UActorComponent* Comp);
 
@@ -886,6 +892,10 @@ public:
 	//----------------------------------------------------------------------//
 	// navigation octree related functions
 	//----------------------------------------------------------------------//
+	static NAVIGATIONSYSTEM_API FNavigationElementHandle AddNavigationElement(UWorld* World, FNavigationElement&& Element);
+	static NAVIGATIONSYSTEM_API void RemoveNavigationElement(UWorld* World, FNavigationElementHandle ElementHandle);
+	static NAVIGATIONSYSTEM_API void OnNavigationElementUpdated(UWorld* World, FNavigationElementHandle ElementHandle, FNavigationElement&& Element);
+
 	static NAVIGATIONSYSTEM_API void OnNavRelevantObjectRegistered(UObject& Object);
 	static NAVIGATIONSYSTEM_API void UpdateNavRelevantObjectInNavOctree(UObject& Object);
 	static NAVIGATIONSYSTEM_API void OnNavRelevantObjectUnregistered(UObject& Object);
@@ -906,9 +916,33 @@ public:
 	static NAVIGATIONSYSTEM_API void UpdateNavOctreeAfterMove(USceneComponent* Comp);
 
 protected:
+	/** Fetches the repository from the provided object and registers the object with it. */
+	static NAVIGATIONSYSTEM_API void RegisterNavRelevantObjectStatic(const INavRelevantInterface& NavRelevantObject, const UObject& Object);
+
+	/** Similar to UnregisterNavRelevantObjectStatic but uses the cached repository. */
+	NAVIGATIONSYSTEM_API void RegisterNavRelevantObjectInternal(const INavRelevantInterface& NavRelevantObject, const UObject& Object);
+
+	/** Fetches the repository from the provided object and unregisters the object with it. */
+	static NAVIGATIONSYSTEM_API void UnregisterNavRelevantObjectStatic(const UObject& Object);
+
+	/** Similar to UnregisterNavRelevantObjectStatic but uses the cached repository. */
+	NAVIGATIONSYSTEM_API void UnregisterNavRelevantObjectInternal(const UObject& Object);
+
+	/**
+	 * In actor/component code paths it is possible that updates get called without a previous call to register (wasn't relevant at that time).
+	 * In those cases we create a new, or an updated, navigation element if needed.
+	 * @note Method expects only navigation relevant objects and will ensure otherwise.
+	 */
+	static NAVIGATIONSYSTEM_API void UpdateNavRelevantObjectInNavOctreeStatic(
+		const INavRelevantInterface& InNavRelevantObject,
+		const UObject& InObject,
+		UNavigationSystemV1* InNavigationSystem,
+		TFunctionRef<void(UNavigationSystemV1&, const TSharedRef<const FNavigationElement>&)> InCallback);
+
 	/** A helper function that gathers all actors attached to RootActor and fetches 
-	 *	them back. The function does consider multi-level attachments. 
-	 *	@param AttachedActors is getting reset at the beginning of the function. 
+	 *	them back. The function does consider multi-level attachments.
+	 *	@param RootActor The reference actor to get the attached actors from.
+	 *	@param OutAttachedActors is getting reset at the beginning of the function. 
 	 *		When done it is guaranteed to contain unique, non-null ptrs. It 
 	 *		will not include RootActor.
 	 *	@return the number of unique attached actors */
@@ -920,11 +954,12 @@ protected:
 public:
 	/** removes all navoctree entries for actor and its components */
 	static NAVIGATIONSYSTEM_API void ClearNavOctreeAll(AActor* Actor);
+
 	/** updates bounds of all components implementing INavRelevantInterface */
 	static NAVIGATIONSYSTEM_API void UpdateNavOctreeBounds(AActor* Actor);
 
 	NAVIGATIONSYSTEM_API void AddDirtyArea(const FBox& NewArea, ENavigationDirtyFlag Flags, const FName& DebugReason = NAME_None);
-	NAVIGATIONSYSTEM_API void AddDirtyArea(const FBox& NewArea, ENavigationDirtyFlag Flags, const TFunction<UObject*()>& ObjectProviderFunc, const FName& DebugReason = NAME_None);
+	NAVIGATIONSYSTEM_API void AddDirtyArea(const FBox& NewArea, ENavigationDirtyFlag Flags, const TFunction<const TSharedPtr<const FNavigationElement>()>& ElementProviderFunc, const FName& DebugReason = NAME_None);
 	NAVIGATIONSYSTEM_API void AddDirtyAreas(const TArray<FBox>& NewAreas, ENavigationDirtyFlag Flags, const FName& DebugReason = NAME_None);
 	UE_DEPRECATED(5.5, "Use the version taking ENavigationDirtyFlag instead.")
 	NAVIGATIONSYSTEM_API void AddDirtyArea(const FBox& NewArea, int32 Flags, const FName& DebugReason = NAME_None);
@@ -939,36 +974,64 @@ public:
 	const FNavigationOctree* GetNavOctree() const { return DefaultOctreeController.GetOctree(); }
 	FNavigationOctree* GetMutableNavOctree() { return DefaultOctreeController.GetMutableOctree(); }
 
-	FORCEINLINE static uint32 HashObject(const UObject& Object)
+	UE_DEPRECATED(5.5, "This method is no longer used by the navigation system.")
+	NAVIGATIONSYSTEM_API static uint32 HashObject(const UObject& Object);
+
+	FORCEINLINE const FOctreeElementId2* GetNavOctreeIdForElement(FNavigationElementHandle Element) const
+	{
+		return DefaultOctreeController.GetNavOctreeIdForElement(Element);
+	}
+	UE_DEPRECATED(5.5, "Use GetNavOctreeIdForElement instead.")
+	NAVIGATIONSYSTEM_API const FOctreeElementId2* GetObjectsNavOctreeId(const UObject& Object) const;
+
+	FORCEINLINE bool HasPendingUpdateForElement(FNavigationElementHandle Element) const
 	{
-		return FNavigationOctree::HashObject(Object);
+		return DefaultOctreeController.HasPendingUpdateForElement(Element);
 	}
-	FORCEINLINE const FOctreeElementId2* GetObjectsNavOctreeId(const UObject& Object) const { return DefaultOctreeController.GetObjectsNavOctreeId(Object); }
-	FORCEINLINE bool HasPendingObjectNavOctreeId(UObject* Object) const { return Object && DefaultOctreeController.HasPendingObjectNavOctreeId(*Object); }
-	FORCEINLINE void RemoveObjectsNavOctreeId(const UObject& Object) { DefaultOctreeController.RemoveObjectsNavOctreeId(Object); }
+	UE_DEPRECATED(5.5, "Use HasPendingUpdateForElement instead.")
+	NAVIGATIONSYSTEM_API bool HasPendingObjectNavOctreeId(UObject* Object) const;
+
+	UE_DEPRECATED(5.5, "This functionality should not be exposed and will be removed.")
+	NAVIGATIONSYSTEM_API void RemoveObjectsNavOctreeId(const UObject& Object);
 
+	UE_DEPRECATED(5.5, "This functionality should not be exposed publicly and will be removed. Derived classes should use RemoveFromNavOctree instead.")
 	NAVIGATIONSYSTEM_API void RemoveNavOctreeElementId(const FOctreeElementId2& ElementId, int32 UpdateFlags);
 
+protected:
+	NAVIGATIONSYSTEM_API void RemoveFromNavOctree(const FOctreeElementId2& ElementId, int32 UpdateFlags);
+public:
+
+	NAVIGATIONSYSTEM_API const FNavigationRelevantData* GetDataForElement(FNavigationElementHandle Element) const;
+	UE_DEPRECATED(5.5, "Use GetDataForElement instead.")
 	NAVIGATIONSYSTEM_API const FNavigationRelevantData* GetDataForObject(const UObject& Object) const;
+
+	NAVIGATIONSYSTEM_API FNavigationRelevantData* GetMutableDataForElement(FNavigationElementHandle Element);
+
+	UE_DEPRECATED(5.5, "Use GetMutableDataForElement instead.")
 	NAVIGATIONSYSTEM_API FNavigationRelevantData* GetMutableDataForObject(const UObject& Object);
 
 	/** find all elements in navigation octree within given box (intersection) */
 	NAVIGATIONSYSTEM_API void FindElementsInNavOctree(const FBox& QueryBox, const FNavigationOctreeFilter& Filter, TArray<FNavigationOctreeElement>& Elements);
 
 	/** update single element in navoctree */
-	NAVIGATIONSYSTEM_API void UpdateNavOctreeElement(UObject* ElementOwner, INavRelevantInterface* ElementInterface, int32 UpdateFlags);
+	NAVIGATIONSYSTEM_API void UpdateNavOctreeElement(FNavigationElementHandle Handle, const TSharedRef<const FNavigationElement>& Element, int32 UpdateFlags);
+	UE_DEPRECATED(5.5, "Use the version taking FNavigationElement instead.")
+	void UpdateNavOctreeElement(UObject* ElementOwner, INavRelevantInterface* ElementInterface, int32 UpdateFlags);
 
-	/** force updating parent node and all its children */
+	UE_DEPRECATED(5.5, "This method will no longer be public and should not be called directly.")
 	NAVIGATIONSYSTEM_API void UpdateNavOctreeParentChain(UObject* ElementOwner, bool bSkipElementOwnerUpdate = false);
 
+	/** update element's bounds in navigation octree and mark only specified area as dirty, doesn't re-export geometry */
+	NAVIGATIONSYSTEM_API bool UpdateNavOctreeElementBounds(FNavigationElementHandle Handle, const FBox& NewBounds, TConstArrayView<FBox> DirtyAreas);
 	UE_DEPRECATED(5.4, "Use the overloaded version with object reference and list of dirty areas as parameters instead.")
 	/** update component bounds in navigation octree and mark only specified area as dirty, doesn't re-export component geometry */
 	NAVIGATIONSYSTEM_API bool UpdateNavOctreeElementBounds(UActorComponent* Comp, const FBox& NewBounds, const FBox& DirtyArea);
-
-	/** update object bounds in navigation octree and mark only specified area as dirty, doesn't re-export geometry */
+	UE_DEPRECATED(5.5, "Use the version taking FNavigationElement instead.")
 	NAVIGATIONSYSTEM_API bool UpdateNavOctreeElementBounds(UObject& Object, const FBox& NewBounds, TConstArrayView<FBox> DirtyAreas);
 
-	/** fetched Object's data from the octree and replaces occurrences of OldArea with NewArea */
+	/** fetched element's data from the octree and replaces occurrences of OldArea with NewArea */
+	NAVIGATIONSYSTEM_API bool ReplaceAreaInOctreeData(FNavigationElementHandle Handle, TSubclassOf<UNavArea> OldArea, TSubclassOf<UNavArea> NewArea, bool bReplaceChildClasses = false);
+	UE_DEPRECATED(5.5, "Use the version taking FNavigationElement instead.")
 	NAVIGATIONSYSTEM_API bool ReplaceAreaInOctreeData(const UObject& Object, TSubclassOf<UNavArea> OldArea, TSubclassOf<UNavArea> NewArea, bool bReplaceChildClasses = false);
 
 	//----------------------------------------------------------------------//
@@ -1149,11 +1212,12 @@ public:
 
 	static FORCEINLINE bool IsNavigationSystemStatic()
 	{
-		return bStaticRuntimeNavigation
 #if WITH_EDITOR
-			&& !(GIsEditor && !GIsPlayInEditorWorld)
-#endif
-			;
+		return (bStaticRuntimeNavigation && !(GIsEditor && !GIsPlayInEditorWorld))
+			|| IsRunningCookCommandlet(); // consider navigation static when cooking to eliminate useless operations
+#else
+	return bStaticRuntimeNavigation;
+#endif // WITH_EDITOR
 	}
 
 	/**	call with bEnableStatic == true to signal the NavigationSystem it doesn't need to store
@@ -1213,6 +1277,9 @@ protected:
 
 	FNavigationDirtyAreasController DefaultDirtyAreasController;
 
+	UPROPERTY(transient)
+	TObjectPtr<UNavigationObjectRepository> Repository = nullptr;
+
 	// async queries
 	FCriticalSection NavDataRegistrationSection;
 
@@ -1280,13 +1347,18 @@ protected:
 	NAVIGATIONSYSTEM_API void UnregisterNavAreaClass(UClass* NavAreaClass);
 
 	NAVIGATIONSYSTEM_API void OnNavigationAreaEvent(UClass* AreaClass, ENavAreaEvent::Type Event);
-	
- 	NAVIGATIONSYSTEM_API FSetElementId RegisterNavOctreeElement(UObject* ElementOwner, INavRelevantInterface* ElementInterface, int32 UpdateFlags);
+
+	NAVIGATIONSYSTEM_API FSetElementId RegisterNavigationElementWithNavOctree(const TSharedRef<const FNavigationElement>& Element, int32 UpdateFlags);
+	UE_DEPRECATED(5.5, "Use the version taking FNavigationElement as parameter instead.")
+	NAVIGATIONSYSTEM_API FSetElementId RegisterNavOctreeElement(UObject* ElementOwner, INavRelevantInterface* ElementInterface, int32 UpdateFlags);
+
+	NAVIGATIONSYSTEM_API void UnregisterNavigationElementWithOctree(const TSharedRef<const FNavigationElement>& Element, int32 UpdateFlags);
+	UE_DEPRECATED(5.5, "Use the version taking FNavigationElement as parameter instead.")
 	NAVIGATIONSYSTEM_API void UnregisterNavOctreeElement(UObject* ElementOwner, INavRelevantInterface* ElementInterface, int32 UpdateFlags);
-	
+
 	/** read element data from navigation octree */
-	NAVIGATIONSYSTEM_API bool GetNavOctreeElementData(const UObject& NodeOwner, ENavigationDirtyFlag& OutDirtyFlags, FBox& OutDirtyBounds);
-	UE_DEPRECATED(5.5, "Use the version taking ENavigationDirtyFlag& instead.")
+	NAVIGATIONSYSTEM_API bool GetNavOctreeElementData(FNavigationElementHandle Element, ENavigationDirtyFlag& OutDirtyFlags, FBox& OutDirtyBounds);
+	UE_DEPRECATED(5.5, "Use the version taking FNavigationElement instead.")
 	NAVIGATIONSYSTEM_API bool GetNavOctreeElementData(const UObject& NodeOwner, int32& DirtyFlags, FBox& DirtyBounds);
 
 	/** Adds given element to NavOctree. No check for owner's validity are performed, 
@@ -1391,7 +1463,7 @@ protected:
 	//----------------------------------------------------------------------//
 public:
 	NAVIGATIONSYSTEM_API void VerifyNavigationRenderingComponents(const bool bShow);
-	/** @param if InLevel is given then only navigation bounds from that level will be considered*/
+	/** @param InLevel Optional parameter to limit the navigation bounds to consider to the ones from that level */
 	NAVIGATIONSYSTEM_API virtual int GetNavigationBoundsForNavData(const ANavigationData& NavData, TArray<FBox>& OutBounds, ULevel* InLevel = nullptr) const;
 	static NAVIGATIONSYSTEM_API INavigationDataInterface* GetNavDataForActor(const AActor& Actor);
 	NAVIGATIONSYSTEM_API virtual void Configure(const UNavigationSystemConfig& Config) override;
