[NavInvokers] Use a map instead of an array to avoid high cost as invoker usage scale up.

diff --git a/Engine/Source/Runtime/NavigationSystem/Private/NavMesh/RecastNavMesh.cpp b/Engine/Source/Runtime/NavigationSystem/Private/NavMesh/RecastNavMesh.cpp
--- a/Engine/Source/Runtime/NavigationSystem/Private/NavMesh/RecastNavMesh.cpp
+++ b/Engine/Source/Runtime/NavigationSystem/Private/NavMesh/RecastNavMesh.cpp
@@ -3888,7 +3888,7 @@ void ARecastNavMesh::UpdateActiveTiles(const TArray<FNavigationInvokerRaw>& Invo
 	}
 
 	TSet<FIntPoint>& OldActiveSet = UpdateActiveTilesWorkingMem.OldActiveSet;
-	TArray<FNavMeshDirtyTileElement>& TilesInMinDistance = UpdateActiveTilesWorkingMem.TilesInMinDistance;
+	TMap<FIntPoint, FNavMeshDirtyTileElement>& TilesInMinDistance = UpdateActiveTilesWorkingMem.TilesInMinDistanceMap;
 	TSet<FIntPoint>& TilesInMaxDistance = UpdateActiveTilesWorkingMem.TilesInMaxDistance;
 	TArray<FIntPoint>& TileToAppend = UpdateActiveTilesWorkingMem.TileToAppend;
 	
@@ -3903,7 +3903,7 @@ void ARecastNavMesh::UpdateActiveTiles(const TArray<FNavigationInvokerRaw>& Invo
 		TileToAppend.Reset();
 
 		const int32 ShrinkThreshold = static_cast<int32>(1.2 * ActiveTilesCount);
-		if (TilesInMinDistance.Max() > ShrinkThreshold)
+		if (TilesInMinDistance.GetMaxIndex() > ShrinkThreshold)
 		{
 			TilesInMinDistance.Shrink();
 		}
@@ -3962,7 +3962,7 @@ void ARecastNavMesh::UpdateActiveTiles(const TArray<FNavigationInvokerRaw>& Invo
 						if (DistanceSq < TileCenterDistanceToAddSq)
 						{
 							// Add unique tile 
-							FNavMeshDirtyTileElement* FoundTile = TilesInMinDistance.FindByPredicate([X, Y](const FNavMeshDirtyTileElement& Tile){ return Tile.Coordinates == FIntPoint(X, Y);});
+							FNavMeshDirtyTileElement* FoundTile = TilesInMinDistance.Find(FIntPoint(X, Y));
 							if (FoundTile)
 							{
 								// Update the priority if already existing
@@ -3970,7 +3970,7 @@ void ARecastNavMesh::UpdateActiveTiles(const TArray<FNavigationInvokerRaw>& Invo
 							}
 							else
 							{
-								TilesInMinDistance.Add(FNavMeshDirtyTileElement{FIntPoint(X,Y), DistanceSq, Invoker.Priority});
+								TilesInMinDistance.Add(FIntPoint(X, Y), FNavMeshDirtyTileElement{FIntPoint(X,Y), DistanceSq, Invoker.Priority});
 								TileToAppend.Add(FIntPoint(X,Y));
 							}
 						}
@@ -4005,9 +4005,10 @@ void ARecastNavMesh::UpdateActiveTiles(const TArray<FNavigationInvokerRaw>& Invo
 	// Find tiles to update
 	TArray<FNavMeshDirtyTileElement> TilesToUpdate;
 	TilesToUpdate.Reserve(ActiveTiles.Num());
-	for (const FNavMeshDirtyTileElement& Tile : TilesInMinDistance)
+	for (const TTuple<FIntPoint, FNavMeshDirtyTileElement>& Pair : TilesInMinDistance)
 	{
 		// Check if it's a new tile (not in the active set)
+		const FNavMeshDirtyTileElement& Tile = Pair.Value;
 		if (!OldActiveSet.Contains(Tile.Coordinates))
 		{
 			TilesToUpdate.Add(Tile);
diff --git a/Engine/Source/Runtime/NavigationSystem/Public/NavMesh/RecastNavMesh.h b/Engine/Source/Runtime/NavigationSystem/Public/NavMesh/RecastNavMesh.h
--- a/Engine/Source/Runtime/NavigationSystem/Public/NavMesh/RecastNavMesh.h
+++ b/Engine/Source/Runtime/NavigationSystem/Public/NavMesh/RecastNavMesh.h
@@ -1487,7 +1487,13 @@ protected:
 	struct FUpdateActiveTilesWorkingMem
 	{
 		TSet<FIntPoint> OldActiveSet;
+
+#if WITH_EDITORONLY_DATA
+		UE_DEPRECATED(5.5, "TilesInMinDistance not used anymore, use TilesInMinDistanceMap instead.")
 		TArray<FNavMeshDirtyTileElement> TilesInMinDistance;
+#endif
+		TMap<FIntPoint, FNavMeshDirtyTileElement> TilesInMinDistanceMap;
+		
 		TSet<FIntPoint> TilesInMaxDistance;
 		TArray<FIntPoint> TileToAppend;
 	};
