[NavigationSystem] prevent access to null component in an actor. Happened during undo operation in unittest from the Editor.

diff --git a/Engine/Source/Runtime/NavigationSystem/Private/NavigationSystem.cpp b/Engine/Source/Runtime/NavigationSystem/Private/NavigationSystem.cpp
--- a/Engine/Source/Runtime/NavigationSystem/Private/NavigationSystem.cpp
+++ b/Engine/Source/Runtime/NavigationSystem/Private/NavigationSystem.cpp
@@ -3709,7 +3709,8 @@ void UNavigationSystemV1::UpdateActorAndComponentsInNavOctree(AActor& Actor, con
 
 			for (UActorComponent* Component : ActorToUpdate.GetComponents())
 			{
-				if (Component->CanEverAffectNavigation()
+				if (Component != nullptr
+					&& Component->CanEverAffectNavigation()
 					&& ActorToUpdate.IsComponentRelevantForNavigation(Component))
 				{
 					const INavRelevantInterface* ComponentNavRelevantInterface = Cast<INavRelevantInterface>(Component);
