[NavLinkGeneration] Addition of area class support for generated links

diff --git a/Engine/Source/Runtime/NavigationSystem/Private/NavMesh/LinkGenerationConfig.cpp b/Engine/Source/Runtime/NavigationSystem/Private/NavMesh/LinkGenerationConfig.cpp
--- a/Engine/Source/Runtime/NavigationSystem/Private/NavMesh/LinkGenerationConfig.cpp
+++ b/Engine/Source/Runtime/NavigationSystem/Private/NavMesh/LinkGenerationConfig.cpp
@@ -2,10 +2,16 @@
 
 #include "NavMesh/LinkGenerationConfig.h"
 #include "BaseGeneratedNavLinksProxy.h"
+#include "NavAreas/NavArea_Default.h"
 
 #if WITH_RECAST
 #include "Detour/DetourNavLinkBuilderConfig.h"
 
+FNavLinkGenerationJumpDownConfig::FNavLinkGenerationJumpDownConfig()
+{
+	AreaClass = UNavArea_Default::StaticClass();
+}
+
 void FNavLinkGenerationJumpDownConfig::CopyToDetourConfig(dtNavLinkBuilderJumpDownConfig& OutDetourConfig) const
 {
 	OutDetourConfig.enabled = bEnabled;
diff --git a/Engine/Source/Runtime/NavigationSystem/Private/NavMesh/RecastNavMeshGenerator.cpp b/Engine/Source/Runtime/NavigationSystem/Private/NavMesh/RecastNavMeshGenerator.cpp
--- a/Engine/Source/Runtime/NavigationSystem/Private/NavMesh/RecastNavMeshGenerator.cpp
+++ b/Engine/Source/Runtime/NavigationSystem/Private/NavMesh/RecastNavMeshGenerator.cpp
@@ -1332,6 +1332,31 @@ struct FOffMeshData
 		LinkParams.Reserve(ElementsCount);
 	}
 
+	void AddLink(const FNavigationLink& Link, const FTransform& LocalToWorld, float DefaultSnapHeight, TFunctionRef<void(dtOffMeshLinkCreateParams&)> AddAreaID)
+	{
+		dtOffMeshLinkCreateParams NewInfo;
+		FMemory::Memzero(NewInfo);
+
+		// not doing anything to link's points order - should be already ordered properly by link processor
+		StoreUnrealPoint(NewInfo.vertsA0, LocalToWorld.TransformPosition(Link.Left));
+		StoreUnrealPoint(NewInfo.vertsB0, LocalToWorld.TransformPosition(Link.Right));
+
+		NewInfo.type = DT_OFFMESH_CON_POINT | 
+			(Link.Direction == ENavLinkDirection::BothWays ? DT_OFFMESH_CON_BIDIR : 0) |
+			(Link.bSnapToCheapestArea ? DT_OFFMESH_CON_CHEAPAREA : 0) |
+			(Link.bIsGenerated ? DT_OFFMESH_CON_GENERATED : 0);
+
+		NewInfo.snapRadius = Link.SnapRadius;
+		NewInfo.snapHeight = Link.bUseSnapHeight ? Link.SnapHeight : DefaultSnapHeight;
+		NewInfo.userID = Link.NavLinkId.GetId();
+
+		AddAreaID(NewInfo);
+
+		// snap area is currently not supported for regular (point-point) offmesh links
+
+		LinkParams.Add(NewInfo);
+	}
+	
 	void AddLinks(const TArray<FNavigationLink>& Links, const FTransform& LocalToWorld, int32 AgentIndex, float DefaultSnapHeight)
 	{
 		for (int32 LinkIndex = 0; LinkIndex < Links.Num(); ++LinkIndex)
@@ -1341,40 +1366,43 @@ struct FOffMeshData
 			{
 				continue;
 			}
-
-			dtOffMeshLinkCreateParams NewInfo;
-			FMemory::Memzero(NewInfo);
-
-			// not doing anything to link's points order - should be already ordered properly by link processor
-			StoreUnrealPoint(NewInfo.vertsA0, LocalToWorld.TransformPosition(Link.Left));
-			StoreUnrealPoint(NewInfo.vertsB0, LocalToWorld.TransformPosition(Link.Right));
-
-			NewInfo.type = DT_OFFMESH_CON_POINT | 
-				(Link.Direction == ENavLinkDirection::BothWays ? DT_OFFMESH_CON_BIDIR : 0) |
-				(Link.bSnapToCheapestArea ? DT_OFFMESH_CON_CHEAPAREA : 0) |
-				(Link.bIsGenerated ? DT_OFFMESH_CON_GENERATED : 0);
-
-			NewInfo.snapRadius = Link.SnapRadius;
-			NewInfo.snapHeight = Link.bUseSnapHeight ? Link.SnapHeight : DefaultSnapHeight;
-			NewInfo.userID = Link.NavLinkId.GetId();
-
-			UClass* AreaClass = Link.GetAreaClass();
-			const int32* AreaID = AreaClassToIdMap->Find(AreaClass);
-			if (AreaID != NULL)
+			
+			AddLink(Link, LocalToWorld, DefaultSnapHeight, [&] (dtOffMeshLinkCreateParams& NewInfo)
 			{
-				NewInfo.area = IntCastChecked<unsigned char>(*AreaID);
-				NewInfo.polyFlag = FlagsPerArea[*AreaID];
-			}
-			else
+				UClass* AreaClass = Link.GetAreaClass();
+				const int32* AreaID = AreaClassToIdMap->Find(AreaClass);
+				if (AreaID != nullptr)
+				{
+					NewInfo.area = IntCastChecked<unsigned char>(*AreaID);
+					NewInfo.polyFlag = FlagsPerArea[*AreaID];
+				}
+				else
+				{
+					UE_LOG(LogNavigation, Warning, TEXT("FRecastTileGenerator: Trying to use undefined area class while defining Off-Mesh links! (%s)"), *GetNameSafe(AreaClass));
+				}
+			});
+		}
+	}
+	
+	void AddLinks(const TArray<FGeneratedNavigationLink>& Links, const FTransform& LocalToWorld, int32 AgentIndex, float DefaultSnapHeight)
+	{
+		for (int32 LinkIndex = 0; LinkIndex < Links.Num(); ++LinkIndex)
+		{
+			const FGeneratedNavigationLink& Link = Links[LinkIndex];
+			if (!Link.SupportedAgents.Contains(AgentIndex))
 			{
-				UE_LOG(LogNavigation, Warning, TEXT("FRecastTileGenerator: Trying to use undefined area class while defining Off-Mesh links! (%s)"), *GetNameSafe(AreaClass));
+				continue;
 			}
 
-			// snap area is currently not supported for regular (point-point) offmesh links
-
-			LinkParams.Add(NewInfo);
+			AddLink(Link, LocalToWorld, DefaultSnapHeight, [&] (dtOffMeshLinkCreateParams& NewInfo)
+			{
+				// area and polyFlag have already been resolved for generated links, jut copy them
+				NewInfo.area = Link.generatedLinkArea;
+				NewInfo.polyFlag = Link.generatedLinkPolyFlag;
+			});
 		}
 	}
+	
 	void AddSegmentLinks(const TArray<FNavigationSegmentLink>& Links, const FTransform& LocalToWorld, int32 AgentIndex, float DefaultSnapHeight)
 	{
 		for (int32 LinkIndex = 0; LinkIndex < Links.Num(); ++LinkIndex)
@@ -3729,7 +3757,7 @@ struct FTileGenerationContext
 };
 
 dtStatus FRecastTileGenerator::BuildTileCacheLinks(FNavMeshBuildContext& BuildContext, dtTileCacheAlloc* alloc, const dtTileCacheLayer& layer,
-	const dtTileCacheContourSet& lcset, TArray<FNavigationLink>& OutGeneratedLinks) const
+	const dtTileCacheContourSet& lcset, TArray<FGeneratedNavigationLink>& OutGeneratedLinks) const
 {
 	TRACE_CPUPROFILER_EVENT_SCOPE(FRecastTileGenerator::BuildTileCacheLinks);
 	BuildContext.log(RC_LOG_PROGRESS, "Building Links:");
@@ -3833,7 +3861,7 @@ dtStatus FRecastTileGenerator::BuildTileCacheLinks(FNavMeshBuildContext& BuildCo
 #endif // RECAST_INTERNAL_DEBUG_DATA		
 	}
 
-	// Make FNavigationLinks
+	// Make FGeneratedNavigationLinks
 	for (const dtNavLinkBuilder::JumpLink& link : linkBuilder.m_links)
 	{
 		if (link.flags == dtNavLinkBuilder::FILTERED)
@@ -3848,11 +3876,13 @@ dtStatus FRecastTileGenerator::BuildTileCacheLinks(FNavMeshBuildContext& BuildCo
 		midA[1] -= linkBuilderConfig.agentClimb;
 		midB[1] -= linkBuilderConfig.agentClimb;
 
-		FNavigationLink& NewLink = OutGeneratedLinks.Emplace_GetRef();
+		FGeneratedNavigationLink& NewLink = OutGeneratedLinks.Emplace_GetRef();
 		NewLink.bIsGenerated = true;
 		if (link.action == DT_LINK_ACTION_JUMP_DOWN)
 		{
-			NewLink.NavLinkId = FNavLinkId(linkBuilderConfig.jumpDownConfig.linkUserId);	
+			NewLink.NavLinkId = FNavLinkId(linkBuilderConfig.jumpDownConfig.linkUserId);
+			NewLink.generatedLinkArea = linkBuilderConfig.jumpDownConfig.area;
+			NewLink.generatedLinkPolyFlag = linkBuilderConfig.jumpDownConfig.polyFlag;
 		}
 		// @todo: Set NavLinkId for jump over links if we keep them.
 		NewLink.Left = Recast2UnrealPoint(midA);
@@ -4060,7 +4090,7 @@ bool FRecastTileGenerator::GenerateNavigationDataLayer(FNavMeshBuildContext& Bui
 	}
 
 	// Build Links
-	TArray<FNavigationLink> GeneratedLinks;
+	TArray<FGeneratedNavigationLink> GeneratedLinks;
 	if (InLinkBuilderData.generatingLinks)
 	{
 		SCOPE_CYCLE_COUNTER(STAT_Navigation_RecastBuildLinks);
@@ -5077,6 +5107,9 @@ void FRecastNavMeshGenerator::Init()
 	
 	AdditionalCachedData = FRecastNavMeshCachedData::Construct(DestNavMesh);
 
+	// Must be called after AdditionalCachedData is set.
+	ResolveGeneratedLinkAreas(Config);
+
 	if (Config.MaxPolysPerTile <= 0 && DestNavMesh->HasValidNavmesh())
 	{
 		const dtNavMeshParams* SavedNavParams = DestNavMesh->GetRecastNavMeshImpl()->DetourNavMesh->getParams();
@@ -5302,6 +5335,27 @@ bool FRecastNavMeshGenerator::IsGeneratingLinks() const
 	return DestNavMesh && DestNavMesh->bGenerateNavLinks && UE::NavMesh::Private::bAllowLinkGeneration; 
 }
 
+void FRecastNavMeshGenerator::ResolveGeneratedLinkAreas(FRecastBuildConfig& OutConfig)
+{
+	if (IsGeneratingLinks())
+	{
+		const FNavLinkGenerationJumpDownConfig& JumpDown = DestNavMesh->NavLinkJumpDownConfig;
+		if (JumpDown.AreaClass)
+		{
+			const int32* AreaIDPtr = AdditionalCachedData.AreaClassToIdMap.Find(JumpDown.AreaClass);
+			if (AreaIDPtr != nullptr)
+			{
+				OutConfig.JumpDownConfig.area = IntCastChecked<unsigned char>(*AreaIDPtr);
+				OutConfig.JumpDownConfig.polyFlag = AdditionalCachedData.FlagsPerArea[*AreaIDPtr];
+			}
+			else
+			{
+				UE_LOG(LogNavigation, Warning, TEXT("FRecastTileGenerator: Trying to use undefined area class while resolving generated links areas. (%s)"), *GetNameSafe(JumpDown.AreaClass));
+			}
+		}
+	}
+}
+
 void FRecastNavMeshGenerator::CalcNavMeshProperties(int32& MaxTiles, int32& MaxPolys)
 {
 	int32 MaxTileBits = -1;
diff --git a/Engine/Source/Runtime/NavigationSystem/Public/NavMesh/LinkGenerationConfig.h b/Engine/Source/Runtime/NavigationSystem/Public/NavMesh/LinkGenerationConfig.h
--- a/Engine/Source/Runtime/NavigationSystem/Public/NavMesh/LinkGenerationConfig.h
+++ b/Engine/Source/Runtime/NavigationSystem/Public/NavMesh/LinkGenerationConfig.h
@@ -7,6 +7,7 @@
 #include "LinkGenerationConfig.generated.h"
 
 class UBaseGeneratedNavLinksProxy;
+class UNavAreaBase;
 
 #if WITH_RECAST
 struct dtNavLinkBuilderJumpDownConfig;
@@ -19,6 +20,8 @@ struct FNavLinkGenerationJumpDownConfig
 {
 	GENERATED_BODY()
 
+	FNavLinkGenerationJumpDownConfig();
+
 	/** Should this config be used to generate links. */
 	UPROPERTY(EditAnywhere, Config, Category = Settings)
 	bool bEnabled = true;
@@ -54,6 +57,10 @@ struct FNavLinkGenerationJumpDownConfig
 	UPROPERTY(EditAnywhere, Config, Category = Settings, meta=(UIMin=0, ClampMin=0))
 	float FilterDistanceThreshold = 80.f;
 
+	/** Area class for links generated by this configuration. */
+	UPROPERTY(EditAnywhere, Config, Category = Settings)
+	TSubclassOf<UNavAreaBase> AreaClass;
+
 	/** Class used to handle links made with this configuration.
 	 * Using this allows to implement custom behaviors when using navlinks, for example during the pathfollow.
 	 * Note that having a proxy is not required for successful navlink pathfinding,
diff --git a/Engine/Source/Runtime/NavigationSystem/Public/NavMesh/RecastNavMeshGenerator.h b/Engine/Source/Runtime/NavigationSystem/Public/NavMesh/RecastNavMeshGenerator.h
--- a/Engine/Source/Runtime/NavigationSystem/Public/NavMesh/RecastNavMeshGenerator.h
+++ b/Engine/Source/Runtime/NavigationSystem/Public/NavMesh/RecastNavMeshGenerator.h
@@ -347,6 +347,15 @@ struct FRecastTileTimeSliceSettings
 	int32 FilterLedgeSpansMaxYProcess = 13;
 };
 
+struct FGeneratedNavigationLink : public FNavigationLink
+{
+	/// User defined flags assigned to the polys of off-mesh connections
+	unsigned short generatedLinkPolyFlag = 0;
+
+	/// User defined area ids assigned to the off-mesh connections
+	unsigned char generatedLinkArea = 0;
+};
+
 /**
  * Class handling generation of a single tile, caching data that can speed up subsequent tile generations
  */
@@ -500,7 +509,7 @@ protected:
 
 	/** Builds navigation links */
 	dtStatus BuildTileCacheLinks(FNavMeshBuildContext& BuildContext, struct dtTileCacheAlloc* alloc, const dtTileCacheLayer& layer,
-		const struct dtTileCacheContourSet& lcset, TArray<FNavigationLink>& OutGeneratedLinks) const;
+		const struct dtTileCacheContourSet& lcset, TArray<FGeneratedNavigationLink>& OutGeneratedLinks) const;
 	
 	NAVIGATIONSYSTEM_API virtual void ApplyVoxelFilter(struct rcHeightfield* SolidHF, FVector::FReal WalkableRadius);
 
@@ -1023,6 +1032,10 @@ public:
 	bool IsGeneratingLinks() const;
 
 protected:
+	/** Resolve area class from link generation config into recast areaIds and flags.
+	 * Must be called after AdditionalCachedData is constructed. */
+	void ResolveGeneratedLinkAreas(FRecastBuildConfig& OutConfig);
+	
 	NAVIGATIONSYSTEM_API virtual void RestrictBuildingToActiveTiles(bool InRestrictBuildingToActiveTiles);
 	
 	/** Blocks until build for specified list of tiles is complete and discard results */
diff --git a/Engine/Source/Runtime/Navmesh/Public/Detour/DetourNavLinkBuilderConfig.h b/Engine/Source/Runtime/Navmesh/Public/Detour/DetourNavLinkBuilderConfig.h
--- a/Engine/Source/Runtime/Navmesh/Public/Detour/DetourNavLinkBuilderConfig.h
+++ b/Engine/Source/Runtime/Navmesh/Public/Detour/DetourNavLinkBuilderConfig.h
@@ -41,6 +41,12 @@ struct dtNavLinkBuilderJumpDownConfig
 	/// User id used to handle links made from this configuration.
 	unsigned long long linkUserId = 0;
 
+	/// User defined flags assigned to the off-mesh connections
+	unsigned short polyFlag = 0;
+
+	/// User defined area id assigned to the off-mesh connections
+	unsigned char area = 0;
+
 	/// Initialize the configuration by computing cached values.
 	NAVMESH_API void init();
 };
@@ -49,7 +55,7 @@ struct dtNavLinkBuilderJumpDownConfig
 struct dtNavLinkBuilderJumpOverConfig
 {
 	/// Should this config be used to generate links.
-	bool enabled = true;
+	bool enabled = false;	// @todo: Disabled for now since this configuration not currently exposed.
 
 	/// Maximum jumpable gap size used when matching edges to jump over. [Limit: > 0] [Units: wu]
 	float jumpGapWidth = 200.f;
