[NavigationSystem] temporarily disabling thread detector while investigating detected cases

diff --git a/Engine/Source/Runtime/NavigationSystem/Private/NavigationObjectRepository.cpp b/Engine/Source/Runtime/NavigationSystem/Private/NavigationObjectRepository.cpp
--- a/Engine/Source/Runtime/NavigationSystem/Private/NavigationObjectRepository.cpp
+++ b/Engine/Source/Runtime/NavigationSystem/Private/NavigationObjectRepository.cpp
@@ -11,7 +11,7 @@ TSharedPtr<const FNavigationElement> UNavigationObjectRepository::AddNavigationE
 {
 #if DO_ENSURE // We don't want to execute the Find at all for targets where ensures are disabled
 	{
-		UE_MT_SCOPED_READ_ACCESS(NavElementAccessDetector);
+		//UE_MT_SCOPED_READ_ACCESS(NavElementAccessDetector);
 
 		if (!ensureMsgf(NavRelevantElements.Find(Element.GetHandle()) == nullptr, TEXT("Same element can't be registered twice.")))
 		{
@@ -22,7 +22,7 @@ TSharedPtr<const FNavigationElement> UNavigationObjectRepository::AddNavigationE
 
 	const TSharedRef SharedElement(MakeShared<FNavigationElement>(MoveTemp(Element)));
 	{
-		UE_MT_SCOPED_WRITE_ACCESS(NavElementAccessDetector);
+		//UE_MT_SCOPED_WRITE_ACCESS(NavElementAccessDetector);
 		NavRelevantElements.Emplace(Element.GetHandle(), SharedElement);
 	}
 
@@ -36,7 +36,7 @@ TSharedPtr<const FNavigationElement> UNavigationObjectRepository::AddNavigationE
 
 void UNavigationObjectRepository::RemoveNavigationElement(const FNavigationElementHandle Handle)
 {
-	UE_MT_SCOPED_WRITE_ACCESS(NavElementAccessDetector);
+	//UE_MT_SCOPED_WRITE_ACCESS(NavElementAccessDetector);
 
 	TSharedPtr<const FNavigationElement> Element;
 	if (ensureMsgf(NavRelevantElements.RemoveAndCopyValue(Handle, Element),
@@ -48,7 +48,7 @@ void UNavigationObjectRepository::RemoveNavigationElement(const FNavigationEleme
 
 void UNavigationObjectRepository::ForEachNavigationElement(TFunctionRef<void(const TSharedRef<const FNavigationElement>&)> PerElementCallback) const
 {
-	UE_MT_SCOPED_READ_ACCESS(NavElementAccessDetector);
+	//UE_MT_SCOPED_READ_ACCESS(NavElementAccessDetector);
 
 	for (auto It = NavRelevantElements.CreateConstIterator(); It; ++It)
 	{
@@ -93,7 +93,7 @@ TSharedPtr<const FNavigationElement> UNavigationObjectRepository::RegisterNavRel
 	{
 		if (const TSharedPtr<const FNavigationElement> SharedElement = AddNavigationElement(FNavigationElement(NavRelevantInterface), NotifyOnSuccess))
 		{
-			UE_MT_SCOPED_WRITE_ACCESS(NavElementAccessDetector);
+			//UE_MT_SCOPED_WRITE_ACCESS(NavElementAccessDetector);
 			ObjectsToHandleMap.Emplace(FObjectKey(&NavRelevantObject), SharedElement->GetHandle());
 
 			UE_LOG(LogNavigation, Verbose, TEXT("%hs [registered] (%s:%s) Bounds: [%s]"), __FUNCTION__,
@@ -124,7 +124,7 @@ void UNavigationObjectRepository::UnregisterNavRelevantObject(const UObject* Nav
 
 	FNavigationElementHandle Handle;
 	{
-		UE_MT_SCOPED_WRITE_ACCESS(NavElementAccessDetector);
+		//UE_MT_SCOPED_WRITE_ACCESS(NavElementAccessDetector);
 		ObjectsToHandleMap.RemoveAndCopyValue(FObjectKey(NavRelevantObject), Handle);
 	}
 
@@ -136,7 +136,7 @@ void UNavigationObjectRepository::UnregisterNavRelevantObject(const UObject* Nav
 
 TSharedPtr<const FNavigationElement> UNavigationObjectRepository::GetNavigationElementForHandle(const FNavigationElementHandle Handle) const
 {
-	UE_MT_SCOPED_READ_ACCESS(NavElementAccessDetector);
+	//UE_MT_SCOPED_READ_ACCESS(NavElementAccessDetector);
 
 	if (const TSharedPtr<const FNavigationElement>* Element = NavRelevantElements.Find(Handle))
 	{
@@ -148,7 +148,7 @@ TSharedPtr<const FNavigationElement> UNavigationObjectRepository::GetNavigationE
 
 FNavigationElementHandle UNavigationObjectRepository::GetNavigationElementHandleForUObject(const UObject* NavRelevantObject)
 {
-	UE_MT_SCOPED_READ_ACCESS(NavElementAccessDetector);
+	//UE_MT_SCOPED_READ_ACCESS(NavElementAccessDetector);
 	if (const FNavigationElementHandle* Handle = ObjectsToHandleMap.Find(FObjectKey(Cast<UObject>(NavRelevantObject))))
 	{
 		return *Handle;
@@ -159,7 +159,7 @@ FNavigationElementHandle UNavigationObjectRepository::GetNavigationElementHandle
 
 TSharedPtr<const FNavigationElement> UNavigationObjectRepository::GetNavigationElementForUObject(const UObject* NavRelevantObject)
 {
-	UE_MT_SCOPED_READ_ACCESS(NavElementAccessDetector);
+	//UE_MT_SCOPED_READ_ACCESS(NavElementAccessDetector);
 
 	if (const FNavigationElementHandle* Handle = ObjectsToHandleMap.Find(FObjectKey(NavRelevantObject)))
 	{
@@ -186,7 +186,7 @@ TSharedPtr<const FNavigationElement> UNavigationObjectRepository::UpdateNavigati
 void UNavigationObjectRepository::RegisterCustomNavLinkObject(INavLinkCustomInterface& CustomNavLinkObject)
 {
 	{
-		UE_MT_SCOPED_WRITE_ACCESS(NavElementAccessDetector);
+		//UE_MT_SCOPED_WRITE_ACCESS(NavElementAccessDetector);
 
 #if DO_ENSURE // We don't want to execute the Find at all for targets where ensures are disabled
 		if (!ensureMsgf(CustomLinkObjects.Find(&CustomNavLinkObject) == INDEX_NONE, TEXT("Same interface pointer can't be registered twice.")))
@@ -204,7 +204,7 @@ void UNavigationObjectRepository::RegisterCustomNavLinkObject(INavLinkCustomInte
 void UNavigationObjectRepository::UnregisterCustomNavLinkObject(INavLinkCustomInterface& CustomNavLinkObject)
 {
 	{
-		UE_MT_SCOPED_WRITE_ACCESS(NavElementAccessDetector);
+		//UE_MT_SCOPED_WRITE_ACCESS(NavElementAccessDetector);
 		ensureMsgf(CustomLinkObjects.Remove(&CustomNavLinkObject) > 0, TEXT("Interface can't be removed since it was not registered or already unregistered)"));
 	}
 
