#UE Null protection for SourceElement in FNavigationDirtyAreasController::AddAreas

diff --git a/Engine/Source/Runtime/NavigationSystem/Private/NavigationDirtyAreasController.cpp b/Engine/Source/Runtime/NavigationSystem/Private/NavigationDirtyAreasController.cpp
--- a/Engine/Source/Runtime/NavigationSystem/Private/NavigationDirtyAreasController.cpp
+++ b/Engine/Source/Runtime/NavigationSystem/Private/NavigationDirtyAreasController.cpp
@@ -166,7 +166,7 @@ void FNavigationDirtyAreasController::AddAreas(const TConstArrayView<FBox> NewAr
 
 	if (ShouldSkipObjectPredicate.IsBound())
 	{
-		const UObject* SourceObject = SourceElement->GetWeakUObject().Get();
+		const UObject* SourceObject = SourceElement ? SourceElement->GetWeakUObject().Get() : nullptr;
 		if (SourceObject && ShouldSkipObjectPredicate.Execute(*SourceObject))
 		{
 			return;
@@ -193,7 +193,7 @@ void FNavigationDirtyAreasController::AddAreas(const TConstArrayView<FBox> NewAr
 #if !UE_BUILD_SHIPPING
 		auto DumpExtraInfo = [SourceElement, DebugReason, BoundsSize, NewArea]()
 			{
-				const UObject* SourceObject = SourceElement->GetWeakUObject().Get();
+				const UObject* SourceObject = SourceElement ? SourceElement->GetWeakUObject().Get() : nullptr;
 				FString ObjectInfo;
 				if (const UObject* ObjectOwner = (SourceObject != nullptr ? SourceObject->GetOuter() : nullptr))
 				{
