[NavigationSystem] replaced checks by warnings until the editor initializes a world properly with all its subsystems when loading from the command line.

diff --git a/Engine/Source/Runtime/NavigationSystem/Private/NavigationSystem.cpp b/Engine/Source/Runtime/NavigationSystem/Private/NavigationSystem.cpp
--- a/Engine/Source/Runtime/NavigationSystem/Private/NavigationSystem.cpp
+++ b/Engine/Source/Runtime/NavigationSystem/Private/NavigationSystem.cpp
@@ -893,7 +893,11 @@ UNavigationSystemV1::UNavigationSystemV1(const FObjectInitializer& ObjectInitial
 		{
 			Repository = World->GetSubsystem<UNavigationObjectRepository>();
 		}
-		checkf(IsRunningCookCommandlet() || Repository, TEXT("UNavigationObjectRepository is mandatory for navigation system operations."));
+
+		if (Repository == nullptr)
+		{
+			UE_LOG(LogNavigation, Warning, TEXT("UNavigationObjectRepository is required for navigation system operations."));
+		}
 	}
 	else if (GetClass() == UNavigationSystemV1::StaticClass())
 	{
@@ -1155,10 +1159,10 @@ bool UNavigationSystemV1::ConditionalPopulateNavOctree()
 #endif // WITH_RECAST
 			}
 
-			if (!DefaultOctreeController.IsNavigationOctreeLocked())
+			if (!DefaultOctreeController.IsNavigationOctreeLocked()
+				&& Repository != nullptr)
 			{
 				// Register all elements registered in the repository world subsystem.
-				checkf(Repository, TEXT("%hs is expected to be called after the repository gets cached."), __FUNCTION__);
 				Repository->ForEachNavigationElement([this](const TSharedRef<const FNavigationElement>& Element)
 				{
 					RegisterNavigationElementWithNavOctree(Element, FNavigationOctreeController::OctreeUpdate_Default);
@@ -2685,7 +2689,11 @@ void UNavigationSystemV1::ProcessRegistrationCandidates()
 
 void UNavigationSystemV1::ProcessCustomLinkPendingRegistration()
 {
-	checkf(Repository, TEXT("%hs is expected to be called after the repository gets cached."), __FUNCTION__);
+	if (Repository == nullptr)
+	{
+		return;
+	}
+
 	for (TWeakInterfacePtr<INavLinkCustomInterface> It : Repository->GetCustomLinks())
 	{
 		if (INavLinkCustomInterface* Interface = It.Get())
@@ -3528,7 +3536,11 @@ void UNavigationSystemV1::RegisterNavRelevantObjectInternal(const INavRelevantIn
 		return;
 	}
 
-	checkf(Repository, TEXT("%hs is expected to be called after the repository gets cached."), __FUNCTION__);
+	if (Repository == nullptr)
+	{
+		return;
+	}
+
 	Repository->RegisterNavRelevantObject(NavRelevantObject);
 }
 
@@ -3556,7 +3568,11 @@ void UNavigationSystemV1::UnregisterNavRelevantObjectInternal(const UObject& Obj
 		return;
 	}
 
-	checkf(Repository, TEXT("%hs is expected to be called after the repository gets cached."), __FUNCTION__);
+	if (Repository == nullptr)
+	{
+		return;
+	}
+
 	Repository->UnregisterNavRelevantObject(&Object);
 }
 
@@ -3878,7 +3894,11 @@ void UNavigationSystemV1::UpdateNavOctreeParentChain(UObject* ElementOwner, bool
 // Deprecated
 bool UNavigationSystemV1::UpdateNavOctreeElementBounds(UActorComponent* Comp, const FBox& NewBounds, const FBox& DirtyArea)
 {
-	checkf(Repository, TEXT("%hs is expected to be called after the repository gets cached."), __FUNCTION__);
+	if (Repository == nullptr)
+	{
+		return false;
+	}
+
 	const FNavigationElementHandle Handle = Repository->GetNavigationElementHandleForUObject(Comp);
 	if (Handle.IsValid())
 	{
@@ -3891,7 +3911,11 @@ bool UNavigationSystemV1::UpdateNavOctreeElementBounds(UActorComponent* Comp, co
 // Deprecated
 bool UNavigationSystemV1::UpdateNavOctreeElementBounds(UObject& Object, const FBox& NewBounds, TConstArrayView<FBox> DirtyAreas)
 {
-	checkf(Repository, TEXT("%hs is expected to be called after the repository gets cached."), __FUNCTION__);
+	if (Repository == nullptr)
+	{
+		return false;
+	}
+
 	const FNavigationElementHandle Handle = Repository->GetNavigationElementHandleForUObject(&Object);
 	if (Handle.IsValid())
 	{
@@ -3914,7 +3938,11 @@ bool UNavigationSystemV1::UpdateNavOctreeElementBounds(const FNavigationElementH
 // Deprecated
 bool UNavigationSystemV1::ReplaceAreaInOctreeData(const UObject& Object, TSubclassOf<UNavArea> OldArea, TSubclassOf<UNavArea> NewArea, bool bReplaceChildClasses)
 {
-	checkf(Repository, TEXT("%hs is expected to be called after the repository gets cached."), __FUNCTION__);
+	if (Repository == nullptr)
+	{
+		return false;
+	}
+
 	if (const FNavigationElementHandle Handle = Repository->GetNavigationElementHandleForUObject(&Object))
 	{
 		return ReplaceAreaInOctreeData(Handle, OldArea, NewArea, bReplaceChildClasses);
@@ -5296,7 +5324,11 @@ bool UNavigationSystemV1::K2_ReplaceAreaInOctreeData(const UObject* Object, TSub
 {
 	SCOPE_CYCLE_COUNTER(STAT_NavOctreeBookkeeping);
 
-	checkf(Repository, TEXT("%hs is expected to be called after the repository gets cached."), __FUNCTION__);
+	if (Repository == nullptr)
+	{
+		return false;
+	}
+
 	if (const FNavigationElementHandle Handle = Repository->GetNavigationElementHandleForUObject(Object))
 	{
 		return ReplaceAreaInOctreeData(Handle, OldArea, NewArea);
@@ -5621,7 +5653,11 @@ void UNavigationSystemV1::UnregisterInvoker_Internal(const UObject& Invoker)
 
 void UNavigationSystemV1::RegisterToRepositoryDelegates()
 {
-	checkf(Repository, TEXT("%hs is expected to be called after the repository gets cached."), __FUNCTION__);
+	if (Repository == nullptr)
+	{
+		return;
+	}
+
 	Repository->OnCustomNavLinkObjectRegistered.BindWeakLambda(this, [this](INavLinkCustomInterface& CustomLink)
 		{
 			RegisterCustomLink(CustomLink);
@@ -5645,7 +5681,11 @@ void UNavigationSystemV1::RegisterToRepositoryDelegates()
 
 void UNavigationSystemV1::UnregisterFromRepositoryDelegates() const
 {
-	checkf(Repository, TEXT("%hs is expected to be called after the repository gets cached."), __FUNCTION__);
+	if (Repository == nullptr)
+	{
+		return;
+	}
+
 	Repository->OnCustomNavLinkObjectRegistered = nullptr;
 	Repository->OnCustomNavLinkObjectUnregistered = nullptr;
 	Repository->OnNavigationElementAddedDelegate = nullptr;
