#UE Restore IsRunningCookCommandlet portion of condition

diff --git a/Engine/Source/Runtime/NavigationSystem/Private/NavigationSystem.cpp b/Engine/Source/Runtime/NavigationSystem/Private/NavigationSystem.cpp
--- a/Engine/Source/Runtime/NavigationSystem/Private/NavigationSystem.cpp
+++ b/Engine/Source/Runtime/NavigationSystem/Private/NavigationSystem.cpp
@@ -894,7 +894,7 @@ UNavigationSystemV1::UNavigationSystemV1(const FObjectInitializer& ObjectInitial
 			Repository = World->GetSubsystem<UNavigationObjectRepository>();
 		}
 
-		if (Repository == nullptr)
+		if (Repository == nullptr && !IsRunningCookCommandlet())
 		{
 			UE_LOG(LogNavigation, Warning, TEXT("UNavigationObjectRepository is required for navigation system operations."));
 		}
