[NavigationSystem] removed a check while investigating a flow causing the check to fire

diff --git a/Engine/Source/Runtime/NavigationSystem/Private/NavigationSystem.cpp b/Engine/Source/Runtime/NavigationSystem/Private/NavigationSystem.cpp
--- a/Engine/Source/Runtime/NavigationSystem/Private/NavigationSystem.cpp
+++ b/Engine/Source/Runtime/NavigationSystem/Private/NavigationSystem.cpp
@@ -3596,10 +3596,12 @@ void UNavigationSystemV1::UpdateNavRelevantObjectInNavOctreeStatic(
 
 	if (UNavigationSystemV1* NavSys = InNavigationSystem ? InNavigationSystem : FNavigationSystem::GetCurrent<UNavigationSystemV1>(InObject.GetWorld()))
 	{
-		checkf(NavSys->Repository, TEXT("%hs is expected to be called after the repository gets cached."), __FUNCTION__);
-		if (const TSharedPtr<const FNavigationElement> SharedElement = NavSys->Repository->UpdateNavigationElementForUObject(InNavRelevantObject, InObject))
+		if (NavSys->Repository)
 		{
-			InCallback(*NavSys, SharedElement.ToSharedRef());
+			if (const TSharedPtr<const FNavigationElement> SharedElement = NavSys->Repository->UpdateNavigationElementForUObject(InNavRelevantObject, InObject))
+			{
+				InCallback(*NavSys, SharedElement.ToSharedRef());
+			}
 		}
 	}
 	else
