[NavigationSystem] fixed regression when updating an actor with all its components

diff --git a/Engine/Source/Runtime/NavigationSystem/Private/NavigationSystem.cpp b/Engine/Source/Runtime/NavigationSystem/Private/NavigationSystem.cpp
--- a/Engine/Source/Runtime/NavigationSystem/Private/NavigationSystem.cpp
+++ b/Engine/Source/Runtime/NavigationSystem/Private/NavigationSystem.cpp
@@ -3739,6 +3739,7 @@ void UNavigationSystemV1::UpdateActorAndComponentsInNavOctree(AActor& Actor, con
 							{
 								InNavSys.UpdateNavOctreeElement(InElement->GetHandle(), InElement, FNavigationOctreeController::OctreeUpdate_Default);
 							});
+						continue;
 					}
 				}
 
