[NavigationSystem] fixed missing NavigationData cleanup when cleaning up the whole system.
When a NavigationData is either EndPlay or Destroyed it unregisters itself from the NavigationSystem and perform its cleanup.
Problem was that when the whole system is doing a cleanup in response to `OnWorldBeginTearDown` it was only unregistering without performing the cleanup

diff --git a/Engine/Source/Runtime/NavigationSystem/Private/NavigationData.cpp b/Engine/Source/Runtime/NavigationSystem/Private/NavigationData.cpp
--- a/Engine/Source/Runtime/NavigationSystem/Private/NavigationData.cpp
+++ b/Engine/Source/Runtime/NavigationSystem/Private/NavigationData.cpp
@@ -485,9 +485,10 @@ void ANavigationData::UnregisterAndCleanUp()
 		{
 			NavSys->UnregisterNavData(this);
 		}
-
-		CleanUp();
 	}
+
+	// Cleanup is not tied to the registration state.
+	CleanUp();
 }
 
 void ANavigationData::CleanUp()
