[NavigationSystem] do not create repository on client if navigation is not allowed

diff --git a/Engine/Source/Runtime/NavigationSystem/Private/NavigationObjectRepository.cpp b/Engine/Source/Runtime/NavigationSystem/Private/NavigationObjectRepository.cpp
--- a/Engine/Source/Runtime/NavigationSystem/Private/NavigationObjectRepository.cpp
+++ b/Engine/Source/Runtime/NavigationSystem/Private/NavigationObjectRepository.cpp
@@ -1,6 +1,7 @@
 ﻿// Copyright Epic Games, Inc. All Rights Reserved.
 
 #include "NavigationObjectRepository.h"
+#include "NavigationSystem.h"
 #include "NavLinkCustomInterface.h"
 #include "AI/NavigationSystemBase.h"
 #include "AI/Navigation/NavigationElement.h"
@@ -64,6 +65,12 @@ TSharedPtr<const FNavigationElement> UNavigationObjectRepository::RegisterNavRel
 	return RegisterNavRelevantObjectInternal(NavRelevantObject, *Cast<UObject>(&NavRelevantObject), ENotifyOnSuccess::Yes);
 }
 
+bool UNavigationObjectRepository::ShouldCreateSubsystem(UObject* Outer) const
+{
+	return (Super::ShouldCreateSubsystem(Outer))
+		&& UNavigationSystemV1::IsNavigationAllowed(Cast<UWorld>(Outer), /** bRequiresNavigationSystemInstance */ false);
+}
+
 TSharedPtr<const FNavigationElement> UNavigationObjectRepository::RegisterNavRelevantObjectInternal(
 	const INavRelevantInterface& NavRelevantInterface,
 	const UObject& NavRelevantObject,
diff --git a/Engine/Source/Runtime/NavigationSystem/Private/NavigationObjectRepository.h b/Engine/Source/Runtime/NavigationSystem/Private/NavigationObjectRepository.h
--- a/Engine/Source/Runtime/NavigationSystem/Private/NavigationObjectRepository.h
+++ b/Engine/Source/Runtime/NavigationSystem/Private/NavigationObjectRepository.h
@@ -163,6 +163,9 @@ public:
 	/** Delegate executed when a custom navigation link is unregistered with the repository. */
 	FOnCustomNavLinkObjectRegistrationEvent OnCustomNavLinkObjectUnregistered;
 
+protected:
+	virtual bool ShouldCreateSubsystem(UObject* Outer) const override;
+
 private:
 
 	TSharedPtr<const FNavigationElement> RegisterNavRelevantObjectInternal(const INavRelevantInterface& NavRelevantInterface, const UObject& NavRelevantObject, ENotifyOnSuccess NotifyOnSuccess);
diff --git a/Engine/Source/Runtime/NavigationSystem/Private/NavigationSystem.cpp b/Engine/Source/Runtime/NavigationSystem/Private/NavigationSystem.cpp
--- a/Engine/Source/Runtime/NavigationSystem/Private/NavigationSystem.cpp
+++ b/Engine/Source/Runtime/NavigationSystem/Private/NavigationSystem.cpp
@@ -758,10 +758,7 @@ UNavigationSystemV1::UNavigationSystemV1(const FObjectInitializer& ObjectInitial
 			UNavigationSystemBase::OnComponentTransformChangedDelegate().BindLambda([](USceneComponent& Comp) {
 				if (UNavigationSystemV1::ShouldUpdateNavOctreeOnComponentChange())
 				{
-					UWorld* World = Comp.GetWorld();
-					UNavigationSystemV1* NavSys = FNavigationSystem::GetCurrent<UNavigationSystemV1>(World);
-					if (NavSys != nullptr
-						&& (NavSys->ShouldAllowClientSideNavigation() || !World->IsNetMode(ENetMode::NM_Client)))
+					if (UNavigationSystemV1::IsNavigationAllowed(Comp.GetWorld(), /** bRequiresNavigationSystemInstance */ true))
 					{
 						// use propagated component's transform update in editor OR server game with additional navsys check
 						UNavigationSystemV1::UpdateNavOctreeAfterMove(&Comp);
@@ -3230,13 +3227,34 @@ void UNavigationSystemV1::ResetCachedFilter(TSubclassOf<UNavigationQueryFilter>
 	}
 }
 
+bool UNavigationSystemV1::IsNavigationAllowed(const UWorld* World, const bool bRequiresNavigationSystemInstance)
+{
+	if (World)
+	{
+		if (World->GetNetMode() != NM_Client)
+		{
+			return true;
+		}
+
+		if (const UNavigationSystemV1* NavSys = FNavigationSystem::GetCurrent<UNavigationSystemV1>(World))
+		{
+			return NavSys->ShouldAllowClientSideNavigation();
+		}
+
+		return bRequiresNavigationSystemInstance == false
+			&& (*GEngine->NavigationSystemClass != nullptr)
+			&& GEngine->NavigationSystemClass->GetDefaultObject<UNavigationSystemV1>()->ShouldAllowClientSideNavigation();
+	}
+
+	return false;
+}
+
 UNavigationSystemV1* UNavigationSystemV1::CreateNavigationSystem(UWorld* WorldOwner)
 {
 	UNavigationSystemV1* NavSys = nullptr;
 
 	// create navigation system for editor and server targets, but remove it from game clients
-	if (WorldOwner && (*GEngine->NavigationSystemClass != nullptr) 
-		&& (GEngine->NavigationSystemClass->GetDefaultObject<UNavigationSystemV1>()->bAllowClientSideNavigation || WorldOwner->GetNetMode() != NM_Client))
+	if (UNavigationSystemV1::IsNavigationAllowed(WorldOwner, /** bRequiresNavigationSystemInstance */ false))
 	{
 		AWorldSettings* WorldSettings = WorldOwner->GetWorldSettings();
 		if (WorldSettings == nullptr || WorldSettings->IsNavigationSystemEnabled())
diff --git a/Engine/Source/Runtime/NavigationSystem/Public/NavigationSystem.h b/Engine/Source/Runtime/NavigationSystem/Public/NavigationSystem.h
--- a/Engine/Source/Runtime/NavigationSystem/Public/NavigationSystem.h
+++ b/Engine/Source/Runtime/NavigationSystem/Public/NavigationSystem.h
@@ -1210,6 +1210,15 @@ public:
 			;
 	}
 
+	/**
+	 * Indicates if navigation is allowed for a given world.
+	 * @param World The world in which the navigation would be used
+	 * @param bRequiresNavigationSystemInstance Indicates if the method relies only on an existing navigation system instance ('true')
+	 * or if it should also look in the engine default setting in case the system is missing ('false').
+	 * @return True if the navigation system can be used in the work, false otherwise
+	 */
+	static NAVIGATIONSYSTEM_API bool IsNavigationAllowed(const UWorld* World, bool bRequiresNavigationSystemInstance);
+
 	static FORCEINLINE bool IsNavigationSystemStatic()
 	{
 #if WITH_EDITOR
