Use FTransactionallySafeScopeLock for ANavigationData::ActivePathsLock

diff --git a/Engine/Source/Runtime/NavigationSystem/Private/NavMesh/RecastNavMesh.cpp b/Engine/Source/Runtime/NavigationSystem/Private/NavMesh/RecastNavMesh.cpp
--- a/Engine/Source/Runtime/NavigationSystem/Private/NavMesh/RecastNavMesh.cpp
+++ b/Engine/Source/Runtime/NavigationSystem/Private/NavMesh/RecastNavMesh.cpp
@@ -24,6 +24,7 @@
 #include "Math/Color.h"
 #include "NavigationDataHandler.h"
 #include "BaseGeneratedNavLinksProxy.h"
+#include "Misc/TransactionallySafeScopeLock.h"
 
 #if WITH_EDITOR
 #include "EditorSupportDelegates.h"
@@ -2769,7 +2770,7 @@ void ARecastNavMesh::InvalidateAffectedPaths(const TArray<FNavTileRef>& ChangedT
 	// before starting async queries task but protecting ActivePaths will make
 	// the system safer in case of future timing changes.
 	{
-		FScopeLock PathLock(&ActivePathsLock);
+		FTransactionallySafeScopeLock PathLock(&ActivePathsLock);
 
 		for (int32 PathIndex = PathsCount - 1; PathIndex >= 0; --PathIndex)
 		{
diff --git a/Engine/Source/Runtime/NavigationSystem/Private/NavigationData.cpp b/Engine/Source/Runtime/NavigationSystem/Private/NavigationData.cpp
--- a/Engine/Source/Runtime/NavigationSystem/Private/NavigationData.cpp
+++ b/Engine/Source/Runtime/NavigationSystem/Private/NavigationData.cpp
@@ -13,6 +13,7 @@
 #include "NavFilters/NavigationQueryFilter.h"
 #include "NavigationSystem.h"
 #include "VisualLogger/VisualLogger.h"
+#include "Misc/TransactionallySafeScopeLock.h"
 
 #include UE_INLINE_GENERATED_CPP_BY_NAME(NavigationData)
 
@@ -423,7 +424,7 @@ void ANavigationData::PurgeUnusedPaths()
 
 	// Paths can be registered from async pathfinding thread 
 	// while unused paths are purged in main thread (actor tick)
-	FScopeLock PathLock(&ActivePathsLock);
+	FTransactionallySafeScopeLock PathLock(&ActivePathsLock);
 
 	const int32 Count = ActivePaths.Num();
 	for (int32 PathIndex = Count - 1; PathIndex >= 0; --PathIndex)
@@ -436,6 +437,13 @@ void ANavigationData::PurgeUnusedPaths()
 	}
 }
 
+void ANavigationData::RegisterActivePath(FNavPathSharedPtr SharedPath)
+{
+	// Paths can be registered from main thread and async pathfinding thread
+	FTransactionallySafeScopeLock PathLock(&ActivePathsLock);
+	ActivePaths.Add(SharedPath);
+}
+
 #if WITH_EDITOR
 void ANavigationData::PostEditUndo()
 {
@@ -853,7 +861,7 @@ uint32 ANavigationData::LogMemUsed() const
 	{
 		// Paths can be registered from async pathfinding thread
 		// while logging is requested on main thread (console command)
-		FScopeLock PathLock(&ActivePathsLock);
+		FTransactionallySafeScopeLock PathLock(&ActivePathsLock);
 		ActivePathsMemSize = ActivePaths.GetAllocatedSize();
 	}
 
diff --git a/Engine/Source/Runtime/NavigationSystem/Public/NavigationData.h b/Engine/Source/Runtime/NavigationSystem/Public/NavigationData.h
--- a/Engine/Source/Runtime/NavigationSystem/Public/NavigationData.h
+++ b/Engine/Source/Runtime/NavigationSystem/Public/NavigationData.h
@@ -18,6 +18,8 @@
 #include "NavigationSystemTypes.h"
 #include "EngineDefines.h"
 #include "AI/Navigation/NavigationDataInterface.h"
+#include "Misc/TransactionallySafeScopeLock.h"
+
 #include "NavigationData.generated.h"
 
 class ANavigationData;
@@ -732,12 +734,7 @@ protected:
 	/** removes from ActivePaths all paths that no longer have shared references (and are invalid in fact) */
 	NAVIGATIONSYSTEM_API void PurgeUnusedPaths();
 
-	void RegisterActivePath(FNavPathSharedPtr SharedPath)
-	{
-		// Paths can be registered from main thread and async pathfinding thread
-		FScopeLock PathLock(&ActivePathsLock);
-		ActivePaths.Add(SharedPath);
-	}
+	NAVIGATIONSYSTEM_API void RegisterActivePath(FNavPathSharedPtr SharedPath);
 
 public:
 	/** Returns bounding box for the navmesh. */
@@ -1041,7 +1038,7 @@ protected:
 	TArray<FNavPathWeakPtr> ActivePaths;
 
 	/** Synchronization object for paths registration from main thread and async pathfinding thread */
-	mutable FCriticalSection ActivePathsLock;
+	mutable FTransactionallySafeCriticalSection ActivePathsLock;
 
 	/**
 	 *	Contains paths that requested observing its goal's location. These paths will be 
