[AutoRTFM] Use a transactionally-safe RW access detector for nav elements.

diff --git a/Engine/Source/Runtime/NavigationSystem/Private/NavigationObjectRepository.h b/Engine/Source/Runtime/NavigationSystem/Private/NavigationObjectRepository.h
--- a/Engine/Source/Runtime/NavigationSystem/Private/NavigationObjectRepository.h
+++ b/Engine/Source/Runtime/NavigationSystem/Private/NavigationObjectRepository.h
@@ -3,6 +3,7 @@
 #pragma once
 
 #include "AI/Navigation/NavigationElement.h" // required since we can't fwd declare types used as TMap keys
+#include "Misc/MTTransactionallySafeAccessDetector.h"
 #include "Subsystems/WorldSubsystem.h"
 #include "UObject/ObjectKey.h"
 #include "UObject/WeakInterfacePtr.h"
@@ -180,7 +181,7 @@ private:
 	TArray<TWeakInterfacePtr<INavLinkCustomInterface>> CustomLinkObjects;
 
 	/** Multi thread access detector used to validate accesses to the maps of registered UObjects and FNavigationElement */
-	UE_MT_DECLARE_RW_ACCESS_DETECTOR(NavElementAccessDetector);
+	UE_MT_DECLARE_TS_RW_ACCESS_DETECTOR(NavElementAccessDetector);
 
 	//----------------------------------------------------------------------//
 	// DEPRECATED
