[NavigationSystem] removed some client specific code from delegates listeners and let the code rely on the presence of the navigation system/repository so all flows are consistent. Replaced recently added IsNavigationAllowed by virtual method 'ShouldCreateNavigationSystemInstance' to allow Navigation configs to know if an instance should be created or not.

diff --git a/Engine/Source/Runtime/Engine/Classes/AI/NavigationSystemBase.h b/Engine/Source/Runtime/Engine/Classes/AI/NavigationSystemBase.h
--- a/Engine/Source/Runtime/Engine/Classes/AI/NavigationSystemBase.h
+++ b/Engine/Source/Runtime/Engine/Classes/AI/NavigationSystemBase.h
@@ -283,6 +283,13 @@ public:
 
 	ENGINE_API virtual bool IsWorldInitDone() const PURE_VIRTUAL(UNavigationSystemBase::IsWorldInitDone, return false;);
 
+	/**
+	 * Indicates whether a navigation system instance should be created for a given world.
+	 * @param World The world in which the navigation would be used
+	 * @return whether the navigation system should be created for the specified world or not
+	 */
+	ENGINE_API virtual bool ShouldCreateNavigationSystemInstance(const UWorld* World) const PURE_VIRTUAL(UNavigationSystemBase::ShouldCreateNavigationSystemInstance, return true;);
+
 	static ENGINE_API FNavigationSystem::FOnNavigationInitSignature& OnNavigationInitStartStaticDelegate();
 	static ENGINE_API FNavigationSystem::FOnNavigationInitSignature& OnNavigationInitDoneStaticDelegate();
 	static ENGINE_API FNavigationSystem::FOnNavAreaGenericEvent& OnNavAreaRegisteredDelegate();
diff --git a/Engine/Source/Runtime/Engine/Private/AI/NavigationSystemConfig.cpp b/Engine/Source/Runtime/Engine/Private/AI/NavigationSystemConfig.cpp
--- a/Engine/Source/Runtime/Engine/Private/AI/NavigationSystemConfig.cpp
+++ b/Engine/Source/Runtime/Engine/Private/AI/NavigationSystemConfig.cpp
@@ -26,10 +26,11 @@ UNavigationSystemBase* UNavigationSystemConfig::CreateAndConfigureNavigationSyst
 {
 	UNavigationSystemBase* NavSys = nullptr;
 
-	TSubclassOf<UNavigationSystemBase> NavSysClass = NavigationSystemClass.ResolveClass();
+	const TSubclassOf<UNavigationSystemBase> NavSysClass = NavigationSystemClass.ResolveClass();
 	ensure(NavSysClass != nullptr || NavigationSystemClass.IsValid() == false);
 	
-	if (NavSysClass)
+	if (NavSysClass
+		&& NavSysClass->GetDefaultObject<UNavigationSystemBase>()->ShouldCreateNavigationSystemInstance(&World))
 	{
 		NavSys = NewObject<UNavigationSystemBase>(&World, NavSysClass);
 		if (NavSys)
@@ -38,7 +39,7 @@ UNavigationSystemBase* UNavigationSystemConfig::CreateAndConfigureNavigationSyst
 		}
 	}
 
-	return NavSys;	
+	return NavSys;
 }
 
 #if WITH_EDITOR
diff --git a/Engine/Source/Runtime/NavigationSystem/Private/NavigationObjectRepository.cpp b/Engine/Source/Runtime/NavigationSystem/Private/NavigationObjectRepository.cpp
--- a/Engine/Source/Runtime/NavigationSystem/Private/NavigationObjectRepository.cpp
+++ b/Engine/Source/Runtime/NavigationSystem/Private/NavigationObjectRepository.cpp
@@ -68,7 +68,7 @@ TSharedPtr<const FNavigationElement> UNavigationObjectRepository::RegisterNavRel
 bool UNavigationObjectRepository::ShouldCreateSubsystem(UObject* Outer) const
 {
 	return (Super::ShouldCreateSubsystem(Outer))
-		&& UNavigationSystemV1::IsNavigationAllowed(Cast<UWorld>(Outer), /** bRequiresNavigationSystemInstance */ false);
+		&& GetDefault<UNavigationSystemV1>()->ShouldCreateNavigationSystemInstance(Cast<UWorld>(Outer));
 }
 
 TSharedPtr<const FNavigationElement> UNavigationObjectRepository::RegisterNavRelevantObjectInternal(
diff --git a/Engine/Source/Runtime/NavigationSystem/Private/NavigationSystem.cpp b/Engine/Source/Runtime/NavigationSystem/Private/NavigationSystem.cpp
--- a/Engine/Source/Runtime/NavigationSystem/Private/NavigationSystem.cpp
+++ b/Engine/Source/Runtime/NavigationSystem/Private/NavigationSystem.cpp
@@ -755,14 +755,11 @@ UNavigationSystemV1::UNavigationSystemV1(const FObjectInitializer& ObjectInitial
 					UNavigationSystemV1::UpdateActorAndComponentsInNavOctree(Actor);
 				}
 			});
-			UNavigationSystemBase::OnComponentTransformChangedDelegate().BindLambda([](USceneComponent& Comp) {
-				if (UNavigationSystemV1::ShouldUpdateNavOctreeOnComponentChange())
+			UNavigationSystemBase::OnComponentTransformChangedDelegate().BindLambda([](USceneComponent& Comp)
+			{
+				if (ShouldUpdateNavOctreeOnComponentChange())
 				{
-					if (UNavigationSystemV1::IsNavigationAllowed(Comp.GetWorld(), /** bRequiresNavigationSystemInstance */ true))
-					{
-						// use propagated component's transform update in editor OR server game with additional navsys check
-						UNavigationSystemV1::UpdateNavOctreeAfterMove(&Comp);
-					}
+					UpdateNavOctreeAfterMove(&Comp);
 				}
 			});
 			UNavigationSystemBase::OnActorRegisteredDelegate().BindLambda([](AActor& Actor) { UNavigationSystemV1::OnActorRegistered(&Actor); });
@@ -3227,34 +3224,22 @@ void UNavigationSystemV1::ResetCachedFilter(TSubclassOf<UNavigationQueryFilter>
 	}
 }
 
-bool UNavigationSystemV1::IsNavigationAllowed(const UWorld* World, const bool bRequiresNavigationSystemInstance)
+bool UNavigationSystemV1::ShouldCreateNavigationSystemInstance(const UWorld* World) const
 {
-	if (World)
-	{
-		if (World->GetNetMode() != NM_Client)
-		{
-			return true;
-		}
+	ensureMsgf(IsTemplate(), TEXT("This method is expected to only be called on Template objects"
+							   " to determine if an instance of this type should be created."));
 
-		if (const UNavigationSystemV1* NavSys = FNavigationSystem::GetCurrent<UNavigationSystemV1>(World))
-		{
-			return NavSys->ShouldAllowClientSideNavigation();
-		}
-
-		return bRequiresNavigationSystemInstance == false
-			&& (*GEngine->NavigationSystemClass != nullptr)
-			&& GEngine->NavigationSystemClass->GetDefaultObject<UNavigationSystemV1>()->ShouldAllowClientSideNavigation();
-	}
-
-	return false;
+	return (World != nullptr)
+		&& (ShouldAllowClientSideNavigation() || (World->GetNetMode() != NM_Client));
 }
 
+// Deprecated
 UNavigationSystemV1* UNavigationSystemV1::CreateNavigationSystem(UWorld* WorldOwner)
 {
 	UNavigationSystemV1* NavSys = nullptr;
 
 	// create navigation system for editor and server targets, but remove it from game clients
-	if (UNavigationSystemV1::IsNavigationAllowed(WorldOwner, /** bRequiresNavigationSystemInstance */ false))
+	if (GetDefault<UNavigationSystemV1>()->ShouldCreateNavigationSystemInstance(WorldOwner))
 	{
 		AWorldSettings* WorldSettings = WorldOwner->GetWorldSettings();
 		if (WorldSettings == nullptr || WorldSettings->IsNavigationSystemEnabled())
@@ -6225,6 +6210,8 @@ void UNavigationSystemModuleConfig::UpdateWithNavSysCDO(const UNavigationSystemV
 
 UNavigationSystemBase* UNavigationSystemModuleConfig::CreateAndConfigureNavigationSystem(UWorld& World) const
 {
+	// This should be handled by ShouldCreateNavigationSystemInstance
+	// called from the base class below but this is an early out.
 	if (bCreateOnClient == false && World.GetNetMode() == NM_Client)
 	{
 		return nullptr;
@@ -6233,7 +6220,7 @@ UNavigationSystemBase* UNavigationSystemModuleConfig::CreateAndConfigureNavigati
 	UNavigationSystemBase* NewNavSys = Super::CreateAndConfigureNavigationSystem(World);
 	UNavigationSystemV1* NavSysInstance = Cast<UNavigationSystemV1>(NewNavSys);
 	UE_CLOG(NavSysInstance == nullptr && NewNavSys != nullptr, LogNavigation, Error
-		, TEXT("Unable to spawn navsys instance of class %s - unable to cast to UNavigationSystemV1")
+		, TEXT("Unable to spawn navigation system instance of class %s - unable to cast to UNavigationSystemV1")
 		, *NavigationSystemClass.GetAssetName()
 	);
 	
diff --git a/Engine/Source/Runtime/NavigationSystem/Public/NavigationSystem.h b/Engine/Source/Runtime/NavigationSystem/Public/NavigationSystem.h
--- a/Engine/Source/Runtime/NavigationSystem/Public/NavigationSystem.h
+++ b/Engine/Source/Runtime/NavigationSystem/Public/NavigationSystem.h
@@ -1191,6 +1191,7 @@ public:
 
 	FORCEINLINE bool IsSetUpForLazyGeometryExporting() const { return bGenerateNavigationOnlyAroundNavigationInvokers; }
 
+	UE_DEPRECATED(5.5, "This creation flow is deprecation. Use FNavigationSystem::AddNavigationSystemToWorld instead.")
 	static NAVIGATIONSYSTEM_API UNavigationSystemV1* CreateNavigationSystem(UWorld* WorldOwner);
 
 	static NAVIGATIONSYSTEM_API UNavigationSystemV1* GetCurrent(UWorld* World);
@@ -1211,13 +1212,11 @@ public:
 	}
 
 	/**
-	 * Indicates if navigation is allowed for a given world.
+	 * Indicates whether navigation is allowed for a given world.
 	 * @param World The world in which the navigation would be used
-	 * @param bRequiresNavigationSystemInstance Indicates if the method relies only on an existing navigation system instance ('true')
-	 * or if it should also look in the engine default setting in case the system is missing ('false').
-	 * @return True if the navigation system can be used in the work, false otherwise
+	 * @return whether the navigation system can be used in the specified world or not
 	 */
-	static NAVIGATIONSYSTEM_API bool IsNavigationAllowed(const UWorld* World, bool bRequiresNavigationSystemInstance);
+	NAVIGATIONSYSTEM_API virtual bool ShouldCreateNavigationSystemInstance(const UWorld* World) const override;
 
 	static FORCEINLINE bool IsNavigationSystemStatic()
 	{
