Fixes for when bCompileRecast is false (for programs, especially)

diff --git a/Engine/Source/Runtime/NavigationSystem/Private/NavMesh/LinkGenerationConfig.cpp b/Engine/Source/Runtime/NavigationSystem/Private/NavMesh/LinkGenerationConfig.cpp
--- a/Engine/Source/Runtime/NavigationSystem/Private/NavMesh/LinkGenerationConfig.cpp
+++ b/Engine/Source/Runtime/NavigationSystem/Private/NavMesh/LinkGenerationConfig.cpp
@@ -1,4 +1,4 @@
-﻿// Copyright Epic Games, Inc. All Rights Reserved.
+// Copyright Epic Games, Inc. All Rights Reserved.
 
 #include "NavMesh/LinkGenerationConfig.h"
 #include "BaseGeneratedNavLinksProxy.h"
@@ -6,12 +6,15 @@
 
 #if WITH_RECAST
 #include "Detour/DetourNavLinkBuilderConfig.h"
+#endif //WITH_RECAST
 
 FNavLinkGenerationJumpDownConfig::FNavLinkGenerationJumpDownConfig()
 {
 	AreaClass = UNavArea_Default::StaticClass();
 }
 
+#if WITH_RECAST
+
 void FNavLinkGenerationJumpDownConfig::CopyToDetourConfig(dtNavLinkBuilderJumpDownConfig& OutDetourConfig) const
 {
 	OutDetourConfig.enabled = bEnabled;
diff --git a/Engine/Source/Runtime/NavigationSystem/Private/NavigationDataHandler.cpp b/Engine/Source/Runtime/NavigationSystem/Private/NavigationDataHandler.cpp
--- a/Engine/Source/Runtime/NavigationSystem/Private/NavigationDataHandler.cpp
+++ b/Engine/Source/Runtime/NavigationSystem/Private/NavigationDataHandler.cpp
@@ -3,13 +3,16 @@
 #include "NavigationDataHandler.h"
 #include "AI/Navigation/NavigationDirtyElement.h"
 #include "AI/Navigation/NavRelevantInterface.h"
-#include "DetourCrowd/DetourCrowd.h"
 #include "Engine/Level.h"
 #include "GameFramework/Actor.h"
 #include "NavAreas/NavArea.h"
 #include "NavMesh/RecastGeometryExport.h"
 #include "VisualLogger/VisualLogger.h"
 
+#if WITH_RECAST
+#include "DetourCrowd/DetourCrowd.h"
+#endif // WITH_RECAST
+
 DEFINE_LOG_CATEGORY_STATIC(LogNavOctree, Warning, All);
 
 namespace UE::NavigationHelper::Private
