Fix crash in ARecastNavMesh::FindEdges when RecastNavMeshImpl is not valid anymore.

diff --git a/Engine/Source/Runtime/NavigationSystem/Private/NavMesh/RecastNavMesh.cpp b/Engine/Source/Runtime/NavigationSystem/Private/NavMesh/RecastNavMesh.cpp
--- a/Engine/Source/Runtime/NavigationSystem/Private/NavMesh/RecastNavMesh.cpp
+++ b/Engine/Source/Runtime/NavigationSystem/Private/NavMesh/RecastNavMesh.cpp
@@ -2153,6 +2153,11 @@ bool ARecastNavMesh::GetPolysInBox(const FBox& Box, TArray<FNavPoly>& Polys, FSh
 
 bool ARecastNavMesh::FindEdges(const NavNodeRef CenterNodeRef, const FVector Center, const FVector::FReal Radius, const FSharedConstNavQueryFilter Filter, TArray<FNavigationWallEdge>& OutEdges) const
 {
+	if (HasValidNavmesh() == false)
+	{
+		return false;
+	}
+	
 	const FNavigationQueryFilter& FilterToUse = GetRightFilterRef(Filter);
 	INITIALIZE_NAVQUERY(NavQuery, FilterToUse.GetMaxSearchNodes());
 	const dtQueryFilter* QueryFilter = static_cast<const FRecastQueryFilter*>(FilterToUse.GetImplementation())->GetAsDetourQueryFilter();
