Change RecastNavmesh modifier sorting to make sure modifiers that replace a nav area are processed at the end

diff --git a/Engine/Source/Runtime/NavigationSystem/Private/NavMesh/RecastNavMesh.cpp b/Engine/Source/Runtime/NavigationSystem/Private/NavMesh/RecastNavMesh.cpp
--- a/Engine/Source/Runtime/NavigationSystem/Private/NavMesh/RecastNavMesh.cpp
+++ b/Engine/Source/Runtime/NavigationSystem/Private/NavMesh/RecastNavMesh.cpp
@@ -1012,7 +1012,8 @@ void ARecastNavMesh::SortAreasForGenerator(TArray<FRecastAreaNavModifierElement>
 			const bool bIsBReplacing = (B.GetAreaClassToReplace() != NULL);
 			if (bIsAReplacing != bIsBReplacing)
 			{
-				return bIsAReplacing;
+				// We want the "Replace" modifiers last because we want them to replace any NavArea from the "Apply" modifiers.
+				return !bIsAReplacing;
 			}
 
 			if (A.Cost != B.Cost)
