Let NavigationModifierVolume be used in Replace mode

diff --git a/Engine/Source/Runtime/Engine/Private/AI/Navigation/NavigationModifier.cpp b/Engine/Source/Runtime/Engine/Private/AI/Navigation/NavigationModifier.cpp
--- a/Engine/Source/Runtime/Engine/Private/AI/Navigation/NavigationModifier.cpp
+++ b/Engine/Source/Runtime/Engine/Private/AI/Navigation/NavigationModifier.cpp
@@ -922,7 +922,7 @@ FCompositeNavModifier FCompositeNavModifier::GetInstantiatedMetaModifier(const F
 	return Result;
 }
 
-void FCompositeNavModifier::CreateAreaModifiers(const UPrimitiveComponent* PrimComp, const TSubclassOf<UNavAreaBase> AreaClass)
+void FCompositeNavModifier::CreateAreaModifiers(const UPrimitiveComponent* PrimComp, const TSubclassOf<UNavAreaBase> AreaClass, const TSubclassOf<UNavAreaBase> AreaClassToReplace)
 {
 	UBodySetup* BodySetup = PrimComp ? ((UPrimitiveComponent*)PrimComp)->GetBodySetup() : nullptr;
 	if (BodySetup == nullptr)
@@ -936,6 +936,11 @@ void FCompositeNavModifier::CreateAreaModifiers(const UPrimitiveComponent* PrimC
 		const FBox BoxSize = BoxElem.CalcAABB(FTransform::Identity, 1.0f);
 
 		FAreaNavModifier AreaMod(BoxSize, PrimComp->GetComponentTransform(), AreaClass);
+		if (AreaClassToReplace)
+		{
+			AreaMod.SetAreaClassToReplace(AreaClassToReplace);
+			AreaMod.SetApplyMode(ENavigationAreaMode::Replace);
+		}
 		Add(AreaMod);
 	}
 
@@ -945,6 +950,11 @@ void FCompositeNavModifier::CreateAreaModifiers(const UPrimitiveComponent* PrimC
 		const FTransform AreaOffset(FVector(0, 0, -SphylElem.Length));
 
 		FAreaNavModifier AreaMod(SphylElem.Radius, SphylElem.Length * 2.0f, AreaOffset * PrimComp->GetComponentTransform(), AreaClass);
+		if (AreaClassToReplace)
+		{
+			AreaMod.SetAreaClassToReplace(AreaClassToReplace);
+			AreaMod.SetApplyMode(ENavigationAreaMode::Replace);
+		}
 		Add(AreaMod);
 	}
 
@@ -954,6 +964,11 @@ void FCompositeNavModifier::CreateAreaModifiers(const UPrimitiveComponent* PrimC
 		if (ConvexElem.VertexData.Num() > 0)
 		{
 			FAreaNavModifier AreaMod(UE::LWC::ConvertArrayType<FVector>(ConvexElem.VertexData), 0, ConvexElem.VertexData.Num(), ENavigationCoordSystem::Unreal, PrimComp->GetComponentTransform(), AreaClass);
+			if (AreaClassToReplace)
+			{
+				AreaMod.SetAreaClassToReplace(AreaClassToReplace);
+				AreaMod.SetApplyMode(ENavigationAreaMode::Replace);
+			}
 			Add(AreaMod);
 		}
 		else
@@ -968,6 +983,11 @@ void FCompositeNavModifier::CreateAreaModifiers(const UPrimitiveComponent* PrimC
 		const FTransform AreaOffset(FVector(0, 0, -SphereElem.Radius));
 
 		FAreaNavModifier AreaMod(SphereElem.Radius, SphereElem.Radius * 2.0f, AreaOffset * PrimComp->GetComponentTransform(), AreaClass);
+		if (AreaClassToReplace)
+		{
+			AreaMod.SetAreaClassToReplace(AreaClassToReplace);
+			AreaMod.SetApplyMode(ENavigationAreaMode::Replace);
+		}
 		Add(AreaMod);
 	}
 }
diff --git a/Engine/Source/Runtime/Engine/Public/AI/NavigationModifier.h b/Engine/Source/Runtime/Engine/Public/AI/NavigationModifier.h
--- a/Engine/Source/Runtime/Engine/Public/AI/NavigationModifier.h
+++ b/Engine/Source/Runtime/Engine/Public/AI/NavigationModifier.h
@@ -337,7 +337,7 @@ struct FCompositeNavModifier : public FNavigationModifier
 		}
 	}
 
-	ENGINE_API void CreateAreaModifiers(const UPrimitiveComponent* PrimComp, const TSubclassOf<UNavAreaBase> AreaClass);
+	ENGINE_API void CreateAreaModifiers(const UPrimitiveComponent* PrimComp, const TSubclassOf<UNavAreaBase> AreaClass, const TSubclassOf<UNavAreaBase> AreaClassToReplace = {});
 	ENGINE_API void CreateAreaModifiers(const FCollisionShape& CollisionShape, const FTransform& LocalToWorld, const TSubclassOf<UNavAreaBase> AreaClass, const bool bIncludeAgentHeight = false);
 
 	FORCEINLINE const TArray<FAreaNavModifier>& GetAreas() const { return Areas; }
diff --git a/Engine/Source/Runtime/NavigationSystem/Private/NavModifierVolume.cpp b/Engine/Source/Runtime/NavigationSystem/Private/NavModifierVolume.cpp
--- a/Engine/Source/Runtime/NavigationSystem/Private/NavModifierVolume.cpp
+++ b/Engine/Source/Runtime/NavigationSystem/Private/NavModifierVolume.cpp
@@ -19,7 +19,7 @@ namespace UE::Navigation::ModVolume::Private
 	void OnNavAreaRegistrationChanged(ANavModifierVolume& ModifierVolume, const UWorld& World, const UClass* NavAreaClass)
 	{
 		if (NavAreaClass
-			&& NavAreaClass == ModifierVolume.GetAreaClass()
+			&& (NavAreaClass == ModifierVolume.GetAreaClass() || NavAreaClass == ModifierVolume.GetAreaClassToReplace())
 			&& &World == ModifierVolume.GetWorld()
 			&& ModifierVolume.HasActorRegisteredAllComponents()) // Update only required after initial registration was completed
 		{
@@ -35,6 +35,7 @@ namespace UE::Navigation::ModVolume::Private
 ANavModifierVolume::ANavModifierVolume(const FObjectInitializer& ObjectInitializer)
 	: Super(ObjectInitializer)
 	, AreaClass(UNavArea_Null::StaticClass())
+	, AreaClassToReplace(nullptr)
 	, NavMeshResolution(ENavigationDataResolution::Invalid)
 {
 	if (GetBrushComponent())
@@ -110,9 +111,13 @@ void ANavModifierVolume::OnNavAreaUnregistered(const UWorld& World, const UClass
 
 void ANavModifierVolume::GetNavigationData(FNavigationRelevantData& Data) const
 {
-	if (Brush && AreaClass && AreaClass != FNavigationSystem::GetDefaultWalkableArea())
+	if (Brush && AreaClass)
 	{
-		Data.Modifiers.CreateAreaModifiers(GetBrushComponent(), AreaClass);
+		// No need to create modifiers if the AreaClass we want to set is the default one unless we want to replace a NavArea to default.
+		if (AreaClass != FNavigationSystem::GetDefaultWalkableArea() || AreaClassToReplace)
+		{
+			Data.Modifiers.CreateAreaModifiers(GetBrushComponent(), AreaClass, AreaClassToReplace);
+		}
 	}
 
 	if (GetBrushComponent()->Brush != nullptr)
@@ -120,7 +125,12 @@ void ANavModifierVolume::GetNavigationData(FNavigationRelevantData& Data) const
 		if (bMaskFillCollisionUnderneathForNavmesh)
 		{
 			const FBox& Box = GetBrushComponent()->Brush->Bounds.GetBox();
-			const FAreaNavModifier AreaMod(Box, GetBrushComponent()->GetComponentTransform(), AreaClass);
+			FAreaNavModifier AreaMod(Box, GetBrushComponent()->GetComponentTransform(), AreaClass);
+			if (AreaClassToReplace)
+			{
+				AreaMod.SetAreaClassToReplace(AreaClassToReplace);
+				AreaMod.SetApplyMode(ENavigationAreaMode::Replace);
+			}
 			Data.Modifiers.SetMaskFillCollisionUnderneathForNavmesh(true);
 			Data.Modifiers.Add(AreaMod);
 		}
@@ -147,6 +157,16 @@ void ANavModifierVolume::SetAreaClass(TSubclassOf<UNavArea> NewAreaClass)
 	}
 }
 
+void ANavModifierVolume::SetAreaClassToReplace(TSubclassOf<UNavArea> NewAreaClassToReplace)
+{
+	if (NewAreaClassToReplace != AreaClassToReplace)
+	{
+		AreaClassToReplace = NewAreaClassToReplace;
+
+		FNavigationSystem::UpdateActorData(*this);
+	}
+}
+
 void ANavModifierVolume::RebuildNavigationData()
 {
 	FNavigationSystem::UpdateActorData(*this);
@@ -168,13 +188,14 @@ void ANavModifierVolume::PostEditUndo()
 void ANavModifierVolume::PostEditChangeProperty(FPropertyChangedEvent& PropertyChangedEvent)
 {
 	static const FName NAME_AreaClass = GET_MEMBER_NAME_CHECKED(ANavModifierVolume, AreaClass);
+	static const FName NAME_AreaClassToReplace = GET_MEMBER_NAME_CHECKED(ANavModifierVolume, AreaClassToReplace);
 	static const FName NAME_BrushComponent = TEXT("BrushComponent");
 
 	Super::PostEditChangeProperty(PropertyChangedEvent);
 	
 	const FName PropName = PropertyChangedEvent.Property ? PropertyChangedEvent.Property->GetFName() : NAME_None;
 
-	if (PropName == NAME_AreaClass)
+	if (PropName == NAME_AreaClass || PropName == NAME_AreaClassToReplace)
 	{
 		FNavigationSystem::UpdateActorData(*this);
 	}
diff --git a/Engine/Source/Runtime/NavigationSystem/Public/NavModifierVolume.h b/Engine/Source/Runtime/NavigationSystem/Public/NavModifierVolume.h
--- a/Engine/Source/Runtime/NavigationSystem/Public/NavModifierVolume.h
+++ b/Engine/Source/Runtime/NavigationSystem/Public/NavModifierVolume.h
@@ -26,9 +26,14 @@ class ANavModifierVolume : public AVolume, public INavRelevantInterface
 	GENERATED_BODY()
 
 protected:
+	/** NavArea to apply inside the defined volume. */
 	UPROPERTY(EditAnywhere, BlueprintReadOnly, Category=Default)
 	TSubclassOf<UNavArea> AreaClass;
 
+	/** When setting this value, the modifier volume behavior changes : it will now replace any surface marked by AreaClassToReplace in the volume and replace it with AreaClass. */ 
+	UPROPERTY(EditAnywhere, BlueprintReadOnly, Category=Default)
+	TSubclassOf<UNavArea> AreaClassToReplace;
+
 	/** Experimental: if set, the 2D space occupied by the volume box will ignore FillCollisionUnderneathForNavmesh */
 	UPROPERTY(EditAnywhere, Category = Default, AdvancedDisplay)
 	bool bMaskFillCollisionUnderneathForNavmesh;
@@ -48,8 +53,11 @@ public:
 
 	UFUNCTION(BlueprintCallable, Category = "AI|Navigation")
 	NAVIGATIONSYSTEM_API void SetAreaClass(TSubclassOf<UNavArea> NewAreaClass = nullptr);
+	UFUNCTION(BlueprintCallable, Category = "AI|Navigation")
+	NAVIGATIONSYSTEM_API void SetAreaClassToReplace(TSubclassOf<UNavArea> NewAreaClassToReplace = nullptr);
 
 	TSubclassOf<UNavArea> GetAreaClass() const { return AreaClass; }
+	TSubclassOf<UNavArea> GetAreaClassToReplace() const { return AreaClassToReplace; }
 
 	NAVIGATIONSYSTEM_API virtual void GetNavigationData(FNavigationRelevantData& Data) const override;
 	NAVIGATIONSYSTEM_API virtual FBox GetNavigationBounds() const override;
