Let NavigationModifierComponent be used in Replace mode

diff --git a/Engine/Source/Runtime/NavigationSystem/Private/NavModifierComponent.cpp b/Engine/Source/Runtime/NavigationSystem/Private/NavModifierComponent.cpp
--- a/Engine/Source/Runtime/NavigationSystem/Private/NavModifierComponent.cpp
+++ b/Engine/Source/Runtime/NavigationSystem/Private/NavModifierComponent.cpp
@@ -14,22 +14,35 @@
 
 #include UE_INLINE_GENERATED_CPP_BY_NAME(NavModifierComponent)
 
-#if WITH_EDITOR
 namespace UE::Navigation::ModComponent::Private
 {
+#if WITH_EDITOR
 	void OnNavAreaRegistrationChanged(UNavModifierComponent& ModifierComponent, const UWorld& World, const UClass* NavAreaClass)
 	{
-		if (NavAreaClass && NavAreaClass == ModifierComponent.AreaClass && &World == ModifierComponent.GetWorld())
+		if (NavAreaClass && (NavAreaClass == ModifierComponent.AreaClass || NavAreaClass == ModifierComponent.AreaClassToReplace) && &World == ModifierComponent.GetWorld())
 		{
 			ModifierComponent.RefreshNavigationModifiers();
 		}
 	}
-} // UE::Navigation::ModComponent::Private
 #endif // WITH_EDITOR
 
+	FAreaNavModifier CreateAreaModifier(const FBox& Box, const FQuat& Quat, const TSubclassOf<UNavArea>& AreaClass, const TSubclassOf<UNavArea>& AreaClassToReplace, const bool bIncludeAgentHeight)
+	{
+		FAreaNavModifier AreaNavModifier(Box, FTransform(Quat), AreaClass);
+		AreaNavModifier.SetIncludeAgentHeight(bIncludeAgentHeight);
+		if (AreaClassToReplace)
+		{
+			AreaNavModifier.SetAreaClassToReplace(AreaClassToReplace);
+			AreaNavModifier.SetApplyMode(ENavigationAreaMode::Replace);
+		}
+		return AreaNavModifier;
+	}
+} // UE::Navigation::ModComponent::Private
+
 UNavModifierComponent::UNavModifierComponent(const FObjectInitializer& ObjectInitializer) : Super(ObjectInitializer)
 {
 	AreaClass = UNavArea_Null::StaticClass();
+	AreaClassToReplace = nullptr;
 	FailsafeExtent = FVector(100, 100, 100);
 	bIncludeAgentHeight = true;
 	NavMeshResolution = ENavigationDataResolution::Invalid;
@@ -211,7 +224,7 @@ void UNavModifierComponent::CalcAndCacheBounds() const
 		TArray<FAreaNavModifier> Areas;
 		for (int32 Idx = 0; Idx < ComponentBounds.Num(); Idx++)
 		{
-			Areas.Add(FAreaNavModifier(ComponentBounds[Idx].Box, FTransform(ComponentBounds[Idx].Quat), AreaClass));
+			Areas.Add(UE::Navigation::ModComponent::Private::CreateAreaModifier(ComponentBounds[Idx].Box, ComponentBounds[Idx].Quat, AreaClass, AreaClassToReplace, bIncludeAgentHeight));
 		}
 
 		for(const FAreaNavModifier& Modifier : Areas)
@@ -225,7 +238,7 @@ void UNavModifierComponent::GetNavigationData(FNavigationRelevantData& Data) con
 {
 	for (int32 Idx = 0; Idx < ComponentBounds.Num(); Idx++)
 	{
-		Data.Modifiers.Add(FAreaNavModifier(ComponentBounds[Idx].Box, FTransform(ComponentBounds[Idx].Quat), AreaClass).SetIncludeAgentHeight(bIncludeAgentHeight));
+		Data.Modifiers.Add(UE::Navigation::ModComponent::Private::CreateAreaModifier(ComponentBounds[Idx].Box, ComponentBounds[Idx].Quat, AreaClass, AreaClassToReplace, bIncludeAgentHeight));
 	}
 
 	Data.Modifiers.SetNavMeshResolution(NavMeshResolution);
@@ -240,6 +253,15 @@ void UNavModifierComponent::SetAreaClass(TSubclassOf<UNavArea> NewAreaClass)
 	}
 }
 
+void UNavModifierComponent::SetAreaClassToReplace(TSubclassOf<UNavArea> NewAreaClassToReplace)
+{
+	if (AreaClassToReplace != NewAreaClassToReplace)
+	{
+		AreaClassToReplace = NewAreaClassToReplace;
+		RefreshNavigationModifiers();
+	}
+}
+
 void UNavModifierComponent::OnTransformUpdated(USceneComponent* RootComponent, EUpdateTransformFlags UpdateTransformFlags, ETeleportType Teleport)
 {
 	// force bounds recaching next time GetNavigationBounds gets called.
diff --git a/Engine/Source/Runtime/NavigationSystem/Public/NavModifierComponent.h b/Engine/Source/Runtime/NavigationSystem/Public/NavModifierComponent.h
--- a/Engine/Source/Runtime/NavigationSystem/Public/NavModifierComponent.h
+++ b/Engine/Source/Runtime/NavigationSystem/Public/NavModifierComponent.h
@@ -20,9 +20,14 @@ class UNavModifierComponent : public UNavRelevantComponent
 {
 	GENERATED_UCLASS_BODY()
 
+	/** NavArea to apply inside the defined volume. */
 	UPROPERTY(EditAnywhere, BlueprintReadOnly, Category = Navigation)
 	TSubclassOf<UNavArea> AreaClass;
 
+	/** When setting this value, the modifier behavior changes : it will now replace any surface marked by AreaClassToReplace in the volume and replace it with AreaClass. */ 
+	UPROPERTY(EditAnywhere, BlueprintReadOnly, Category = Navigation)
+	TSubclassOf<UNavArea> AreaClassToReplace;
+
 	/** box extent used ONLY when owning actor doesn't have collision component */
 	UPROPERTY(EditAnywhere, Category = Navigation)
 	FVector FailsafeExtent;
@@ -45,6 +50,8 @@ class UNavModifierComponent : public UNavRelevantComponent
 
 	UFUNCTION(BlueprintCallable, Category = "AI|Navigation")
 	NAVIGATIONSYSTEM_API void SetAreaClass(TSubclassOf<UNavArea> NewAreaClass);
+	UFUNCTION(BlueprintCallable, Category = "AI|Navigation")
+	NAVIGATIONSYSTEM_API void SetAreaClassToReplace(TSubclassOf<UNavArea> NewAreaClassToReplace);
 
 protected:
 	NAVIGATIONSYSTEM_API void OnTransformUpdated(USceneComponent* RootComponent, EUpdateTransformFlags UpdateTransformFlags, ETeleportType Teleport);
