Hide the Enabled property from the navlink config since it currently doesn't have a meaningfull usage.

diff --git a/Engine/Source/Runtime/NavigationSystem/Public/NavMesh/LinkGenerationConfig.h b/Engine/Source/Runtime/NavigationSystem/Public/NavMesh/LinkGenerationConfig.h
--- a/Engine/Source/Runtime/NavigationSystem/Public/NavMesh/LinkGenerationConfig.h
+++ b/Engine/Source/Runtime/NavigationSystem/Public/NavMesh/LinkGenerationConfig.h
@@ -23,7 +23,7 @@ struct FNavLinkGenerationJumpDownConfig
 	FNavLinkGenerationJumpDownConfig();
 
 	/** Should this config be used to generate links. */
-	UPROPERTY(EditAnywhere, Config, Category = Settings)
+	UPROPERTY(Config)		// @todo Hidden for now. Since there is currently only one config, it has no real usage.
 	bool bEnabled = true;
 	
 	/** Horizontal length of the jump.
