Fix issue where AgentToNavDataMap entries were not cleaned up properly due to ARecastNavMesh updating its config to have a valid agent step height if the step height was not overridden.

diff --git a/Engine/Source/Runtime/NavigationSystem/Private/NavMesh/RecastNavMesh.cpp b/Engine/Source/Runtime/NavigationSystem/Private/NavMesh/RecastNavMesh.cpp
--- a/Engine/Source/Runtime/NavigationSystem/Private/NavMesh/RecastNavMesh.cpp
+++ b/Engine/Source/Runtime/NavigationSystem/Private/NavMesh/RecastNavMesh.cpp
@@ -1215,6 +1215,10 @@ void ARecastNavMesh::SetConfig(const FNavDataConfig& Src)
 {
 	// Step 1: set NavDataConfig
 	NavDataConfig = Src;
+
+	// Set the NavAgentProperties that are being used to configure navmmesh
+	SetNavAgentProperties(Src);
+
 	if (!Src.HasStepHeightOverride())
 	{
 		// If there is no override, use the navmesh value.
diff --git a/Engine/Source/Runtime/NavigationSystem/Private/NavigationData.cpp b/Engine/Source/Runtime/NavigationSystem/Private/NavigationData.cpp
--- a/Engine/Source/Runtime/NavigationSystem/Private/NavigationData.cpp
+++ b/Engine/Source/Runtime/NavigationSystem/Private/NavigationData.cpp
@@ -876,4 +876,10 @@ uint32 ANavigationData::LogMemUsed() const
 	}
 
 	return MemUsed;
+}
+
+void ANavigationData::SetConfig(const FNavDataConfig& Src)
+{
+	SetNavAgentProperties(Src);
+	NavDataConfig = Src;
 }
\ No newline at end of file
diff --git a/Engine/Source/Runtime/NavigationSystem/Private/NavigationSystem.cpp b/Engine/Source/Runtime/NavigationSystem/Private/NavigationSystem.cpp
--- a/Engine/Source/Runtime/NavigationSystem/Private/NavigationSystem.cpp
+++ b/Engine/Source/Runtime/NavigationSystem/Private/NavigationSystem.cpp
@@ -2835,7 +2835,7 @@ void UNavigationSystemV1::UnregisterNavData(ANavigationData* NavData)
 		return;
 	}
 
-    AgentToNavDataMap.Remove(NavData->GetConfig());
+    AgentToNavDataMap.Remove(NavData->GetNavAgentProperties());
 
 	FScopeLock Lock(&NavDataRegistration);
 	NavDataRegistrationQueue.Remove(NavData);
diff --git a/Engine/Source/Runtime/NavigationSystem/Public/NavigationData.h b/Engine/Source/Runtime/NavigationSystem/Public/NavigationData.h
--- a/Engine/Source/Runtime/NavigationSystem/Public/NavigationData.h
+++ b/Engine/Source/Runtime/NavigationSystem/Public/NavigationData.h
@@ -625,10 +625,17 @@ public:
 public:
 	FORCEINLINE const FNavDataConfig& GetConfig() const { return NavDataConfig; }
 	FORCEINLINE ERuntimeGenerationType GetRuntimeGenerationMode() const { return RuntimeGeneration; }
-	virtual void SetConfig(const FNavDataConfig& Src) { NavDataConfig = Src; }
+	/** Populates NavDataConfig and sets NavAgentProperties with the Src config
+	 *  Should be used when initially configuring the navdata and not when updating
+	 *  values used in NavDataConfig */
+	virtual void SetConfig(const FNavDataConfig& Src);
 
 	void SetSupportsDefaultAgent(bool bIsDefault) { bSupportsDefaultAgent = bIsDefault; SetNavRenderingEnabled(bIsDefault); }
 	bool IsSupportingDefaultAgent() const { return bSupportsDefaultAgent; }
+	const FNavAgentProperties& GetNavAgentProperties() const
+	{
+		return NavAgentProperties;
+	}
 
 	NAVIGATIONSYSTEM_API virtual bool DoesSupportAgent(const FNavAgentProperties& AgentProps) const;
 
@@ -1083,6 +1090,17 @@ private:
 	uint16 NavDataUniqueID;
 
 	static NAVIGATIONSYSTEM_API uint16 GetNextUniqueID();
+
+protected:
+	/** The exact nav agent properties used when registering this navdata
+	 *  in the navigation system. If navdata configuration changes are
+	 *  needed, NavDataConfig can be modified and used */
+	FNavAgentProperties NavAgentProperties;
+
+	void SetNavAgentProperties(const FNavAgentProperties& InNavAgentProperties)
+	{
+		NavAgentProperties = InNavAgentProperties;
+	}
 };
 
 struct FAsyncPathFindingQuery : public FPathFindingQuery
