[Navigation] fixed regression for processing Level BSP
In case of the level geometry we export the data before adding the node so the octree needed to skip the default export

diff --git a/Engine/Source/Runtime/NavigationSystem/Private/NavigationOctree.cpp b/Engine/Source/Runtime/NavigationSystem/Private/NavigationOctree.cpp
--- a/Engine/Source/Runtime/NavigationSystem/Private/NavigationOctree.cpp
+++ b/Engine/Source/Runtime/NavigationSystem/Private/NavigationOctree.cpp
@@ -217,54 +217,59 @@ void FNavigationOctree::AddNode(const FBox& Bounds, FNavigationOctreeElement& Oc
 	OctreeElement.Bounds = Bounds;
 	OctreeElement.Data->bShouldSkipDirtyAreaOnAddOrRemove = !SourceElement.GetDirtyAreaOnRegistration();
 
-	const bool bDoInstantGathering = !IsLazyGathering(SourceElement);
-
-	if (bGatherGeometry)
+	// Only gather geometry and navigation data if not already provided.
+	// We don't want to use the default geometry export since it will clear the navigation data.
+	if (OctreeElement.Data->IsEmpty())
 	{
-		if (bDoInstantGathering)
-		{
-			GeometryExportDelegate.ExecuteIfBound(SourceElement, *OctreeElement.Data);
-		}
-		else
+		const bool bDoInstantGathering = !IsLazyGathering(SourceElement);
+
+		if (bGatherGeometry)
 		{
-			OctreeElement.Data->bPendingLazyGeometryGathering = true;
-			OctreeElement.Data->bSupportsGatheringGeometrySlices = SourceElement.GeometrySliceExportDelegate.IsBound();
+			if (bDoInstantGathering)
+			{
+				GeometryExportDelegate.ExecuteIfBound(SourceElement, *OctreeElement.Data);
+			}
+			else
+			{
+				OctreeElement.Data->bPendingLazyGeometryGathering = true;
+				OctreeElement.Data->bSupportsGatheringGeometrySlices = SourceElement.GeometrySliceExportDelegate.IsBound();
+			}
 		}
-	}
 
-	SCOPE_CYCLE_COUNTER(STAT_Navigation_GatheringNavigationModifiersSync);
-	if (bDoInstantGathering)
-	{
+		SCOPE_CYCLE_COUNTER(STAT_Navigation_GatheringNavigationModifiersSync);
+		if (bDoInstantGathering)
+		{
 #if !UE_BUILD_SHIPPING
-		const bool bCanOutputDurationWarning = GatheringNavModifiersTimeLimitWarning >= 0.0f;
-		const double StartTime = bCanOutputDurationWarning ? FPlatformTime::Seconds() : 0.0f;
+			const bool bCanOutputDurationWarning = GatheringNavModifiersTimeLimitWarning >= 0.0f;
+			const double StartTime = bCanOutputDurationWarning ? FPlatformTime::Seconds() : 0.0f;
 #endif //!UE_BUILD_SHIPPING
 
-		SourceElement.NavigationDataExportDelegate.ExecuteIfBound(SourceElement, *OctreeElement.Data);
+			SourceElement.NavigationDataExportDelegate.ExecuteIfBound(SourceElement, *OctreeElement.Data);
 
 #if !UE_BUILD_SHIPPING
-		// If GatheringNavModifiersWarningLimitTime is positive, it will print a Warning if the time taken to call GetNavigationData is more than GatheringNavModifiersWarningLimitTime			
-		if (bCanOutputDurationWarning)
-		{
-			if (const double DeltaTime = FPlatformTime::Seconds() - StartTime; DeltaTime > GatheringNavModifiersTimeLimitWarning)
+			// If GatheringNavModifiersWarningLimitTime is positive, it will print a Warning if the time taken to call GetNavigationData is more than GatheringNavModifiersWarningLimitTime			
+			if (bCanOutputDurationWarning)
 			{
-				UE_LOG(LogNavigation, Warning, TEXT("The time (%f sec) for gathering navigation data on a navigation element exceeded the time limit (%f sec) | Element = %s"),
-					DeltaTime,
-					GatheringNavModifiersTimeLimitWarning,
-					*SourceElement.GetName());
+				if (const double DeltaTime = FPlatformTime::Seconds() - StartTime; DeltaTime > GatheringNavModifiersTimeLimitWarning)
+				{
+					UE_LOG(LogNavigation, Warning, TEXT("The time (%f sec) for gathering navigation data on a navigation element exceeded the time limit (%f sec) | Element = %s"),
+						DeltaTime,
+						GatheringNavModifiersTimeLimitWarning,
+						*SourceElement.GetName());
+				}
 			}
-		}
 #endif //!UE_BUILD_SHIPPING
-	}
-	else
-	{
-		OctreeElement.Data->bPendingLazyModifiersGathering = true;
-	}
+		}
+		else
+		{
+			OctreeElement.Data->bPendingLazyModifiersGathering = true;
+		}
 
-	// validate exported data
-	// shrink arrays before counting memory
-	// it will be reallocated when adding to octree and RemoveNode will have different value returned by GetAllocatedSize()
-	OctreeElement.ValidateAndShrink();
+		// validate exported data
+		// shrink arrays before counting memory
+		// it will be reallocated when adding to octree and RemoveNode will have different value returned by GetAllocatedSize()
+		OctreeElement.ValidateAndShrink();
+	}
 
 	const int32 ElementMemory = OctreeElement.GetAllocatedSize();
 	NodesMemory += ElementMemory;
