Fix crash in RecastNamesh::GetPolysInBox

diff --git a/Engine/Source/Runtime/NavigationSystem/Private/NavMesh/RecastNavMesh.cpp b/Engine/Source/Runtime/NavigationSystem/Private/NavMesh/RecastNavMesh.cpp
--- a/Engine/Source/Runtime/NavigationSystem/Private/NavMesh/RecastNavMesh.cpp
+++ b/Engine/Source/Runtime/NavigationSystem/Private/NavMesh/RecastNavMesh.cpp
@@ -2098,7 +2098,8 @@ void ARecastNavMesh::BatchProjectPoints(TArray<FNavigationProjectionWork>& Workl
 bool ARecastNavMesh::GetPolysInBox(const FBox& Box, TArray<FNavPoly>& Polys, FSharedConstNavQueryFilter Filter, const UObject* InOwner) const
 {
 	// sanity check
-	if (RecastNavMeshImpl->GetRecastMesh() == NULL)
+	const dtNavMesh* NavMesh = GetRecastMesh();
+	if (NavMesh == nullptr)
 	{
 		return false;
 	}
@@ -2133,7 +2134,7 @@ bool ARecastNavMesh::GetPolysInBox(const FBox& Box, TArray<FNavPoly>& Polys, FSh
 			{
 				dtPoly const* Poly;
 				dtMeshTile const* Tile;
-				dtStatus Status = RecastNavMeshImpl->GetRecastMesh()->getTileAndPolyByRef(HitPolys[i], &Tile, &Poly);
+				dtStatus Status = NavMesh->getTileAndPolyByRef(HitPolys[i], &Tile, &Poly);
 				if (dtStatusSucceed(Status))
 				{
 					FVector PolyCenter(0);
