Addition of LinkBuilderFlags to the generation config to control wich point links will be added from the created links (center, extremity or both).

diff --git a/Engine/Source/Runtime/NavigationSystem/Private/NavMesh/LinkGenerationConfig.cpp b/Engine/Source/Runtime/NavigationSystem/Private/NavMesh/LinkGenerationConfig.cpp
--- a/Engine/Source/Runtime/NavigationSystem/Private/NavMesh/LinkGenerationConfig.cpp
+++ b/Engine/Source/Runtime/NavigationSystem/Private/NavMesh/LinkGenerationConfig.cpp
@@ -25,6 +25,7 @@ void FNavLinkGenerationJumpDownConfig::CopyToDetourConfig(dtNavLinkBuilderJumpDo
 	OutDetourConfig.jumpEndsHeightTolerance	= JumpEndsHeightTolerance;
 	OutDetourConfig.samplingSeparationFactor = SamplingSeparationFactor;
 	OutDetourConfig.filterDistanceThreshold = FilterDistanceThreshold;
+	OutDetourConfig.linkBuilderFlags = LinkBuilderFlags;
 
 	if (LinkProxy)
 	{
diff --git a/Engine/Source/Runtime/NavigationSystem/Private/NavMesh/RecastNavMeshGenerator.cpp b/Engine/Source/Runtime/NavigationSystem/Private/NavMesh/RecastNavMeshGenerator.cpp
--- a/Engine/Source/Runtime/NavigationSystem/Private/NavMesh/RecastNavMeshGenerator.cpp
+++ b/Engine/Source/Runtime/NavigationSystem/Private/NavMesh/RecastNavMeshGenerator.cpp
@@ -3863,32 +3863,58 @@ dtStatus FRecastTileGenerator::BuildTileCacheLinks(FNavMeshBuildContext& BuildCo
 #endif // RECAST_INTERNAL_DEBUG_DATA		
 	}
 
+	auto AddLinkLambda = [&OutGeneratedLinks](const dtNavLinkBuilderJumpDownConfig& config, const dtReal* posA, const dtReal* posB)
+	{
+		FGeneratedNavigationLink& NewLink = OutGeneratedLinks.Emplace_GetRef();
+		NewLink.bIsGenerated = true;
+		NewLink.NavLinkId = FNavLinkId(config.linkUserId);
+		NewLink.generatedLinkArea = config.area;
+		NewLink.generatedLinkPolyFlag = config.polyFlag;
+		NewLink.Left = Recast2UnrealPoint(posA);
+		NewLink.Right = Recast2UnrealPoint(posB);	
+	};
+
 	// Make FGeneratedNavigationLinks
 	for (const dtNavLinkBuilder::JumpLink& link : linkBuilder.m_links)
 	{
 		if (link.flags == dtNavLinkBuilder::FILTERED)
 			continue;
 
-		// For now, just make a link using the center of the range.
-		dtReal midA[3];
-		dtVlerp(midA, &link.spine0[0], &link.spine1[0], 0.5);
-		dtReal midB[3];
-		dtVlerp(midB, &link.spine0[(link.nspine-1)*3], &link.spine1[(link.nspine-1)*3], 0.5);
-		// Since trajectory validation starts at agentClimb height to ignore small bumps, remove the offset for the actual link height.
-		midA[1] -= linkBuilderConfig.agentClimb;
-		midB[1] -= linkBuilderConfig.agentClimb;
+		// Only "JumpDownConfig are expected for now
+		if (link.action != DT_LINK_ACTION_JUMP_DOWN)
+		{
+			continue;
+		}
+
+		if (TileConfig.JumpDownConfig.linkBuilderFlags & DT_NAVLINK_CREATE_CENTER_POINT_LINK)
+		{
+			// Make a link using the center of the range.
+			dtReal midA[3];
+			dtVlerp(midA, &link.spine0[0], &link.spine1[0], 0.5);
+			dtReal midB[3];
+			dtVlerp(midB, &link.spine0[(link.nspine-1)*3], &link.spine1[(link.nspine-1)*3], 0.5);
+			// Since trajectory validation starts at agentClimb height to ignore small bumps, remove the offset for the actual link height.
+			midA[1] -= linkBuilderConfig.agentClimb;
+			midB[1] -= linkBuilderConfig.agentClimb;
+			AddLinkLambda(linkBuilderConfig.jumpDownConfig, midA, midB);
+		}
 
-		FGeneratedNavigationLink& NewLink = OutGeneratedLinks.Emplace_GetRef();
-		NewLink.bIsGenerated = true;
-		if (link.action == DT_LINK_ACTION_JUMP_DOWN)
+		if (TileConfig.JumpDownConfig.linkBuilderFlags & DT_NAVLINK_CREATE_EXTREMITY_LINKS)
 		{
-			NewLink.NavLinkId = FNavLinkId(linkBuilderConfig.jumpDownConfig.linkUserId);
-			NewLink.generatedLinkArea = linkBuilderConfig.jumpDownConfig.area;
-			NewLink.generatedLinkPolyFlag = linkBuilderConfig.jumpDownConfig.polyFlag;
+			dtReal posA[3];
+			dtReal posB[3];
+			dtVcopy(posA, &link.spine0[0]);
+			dtVcopy(posB, &link.spine0[(link.nspine-1)*3]);
+			posA[1] -= linkBuilderConfig.agentClimb;
+			posB[1] -= linkBuilderConfig.agentClimb;
+			AddLinkLambda(linkBuilderConfig.jumpDownConfig, posA, posB);
+
+			dtVcopy(posA, &link.spine1[0]);
+			dtVcopy(posB, &link.spine1[(link.nspine-1)*3]);
+			posA[1] -= linkBuilderConfig.agentClimb;
+			posB[1] -= linkBuilderConfig.agentClimb;
+			AddLinkLambda(linkBuilderConfig.jumpDownConfig, posA, posB);
 		}
-		// @todo: Set NavLinkId for jump over links if we keep them.
-		NewLink.Left = Recast2UnrealPoint(midA);
-		NewLink.Right = Recast2UnrealPoint(midB);
 	}
 	
 	LogOnExit();
diff --git a/Engine/Source/Runtime/NavigationSystem/Public/NavMesh/LinkGenerationConfig.h b/Engine/Source/Runtime/NavigationSystem/Public/NavMesh/LinkGenerationConfig.h
--- a/Engine/Source/Runtime/NavigationSystem/Public/NavMesh/LinkGenerationConfig.h
+++ b/Engine/Source/Runtime/NavigationSystem/Public/NavMesh/LinkGenerationConfig.h
@@ -14,6 +14,14 @@ struct dtNavLinkBuilderJumpDownConfig;
 struct dtNavLinkBuilderJumpOverConfig;
 #endif //WITH_RECAST
 
+UENUM(meta = (Bitflags, UseEnumValuesAsMaskValuesInEditor = "true"))
+enum class ENavLinkBuilderFlags : uint16
+{
+	CreateCenterPointLink		= 1 << 0,
+	CreateExtremityLink			= 1 << 1,
+};
+ENUM_CLASS_FLAGS(ENavLinkBuilderFlags);
+
 /** Experimental configuration to generate vertical links. */
 USTRUCT()
 struct FNavLinkGenerationJumpDownConfig
@@ -57,6 +65,10 @@ struct FNavLinkGenerationJumpDownConfig
 	UPROPERTY(EditAnywhere, Config, Category = Settings, meta=(UIMin=0, ClampMin=0))
 	float FilterDistanceThreshold = 80.f;
 
+	/** Flags used to indicate how links will be added. */
+	UPROPERTY(EditAnywhere, Config, Category = Settings, meta = (Bitmask, BitmaskEnum = "/Script/NavigationSystem.ENavLinkBuilderFlags"))
+	uint16 LinkBuilderFlags = (uint16)ENavLinkBuilderFlags::CreateCenterPointLink;
+
 	/** Area class for links generated by this configuration. */
 	UPROPERTY(EditAnywhere, Config, Category = Settings)
 	TSubclassOf<UNavAreaBase> AreaClass;
diff --git a/Engine/Source/Runtime/Navmesh/Public/Detour/DetourNavLinkBuilderConfig.h b/Engine/Source/Runtime/Navmesh/Public/Detour/DetourNavLinkBuilderConfig.h
--- a/Engine/Source/Runtime/Navmesh/Public/Detour/DetourNavLinkBuilderConfig.h
+++ b/Engine/Source/Runtime/Navmesh/Public/Detour/DetourNavLinkBuilderConfig.h
@@ -3,6 +3,12 @@
 
 #include "CoreTypes.h"
 
+enum dtNavLinkBuilderFlags : unsigned short
+{
+	DT_NAVLINK_CREATE_CENTER_POINT_LINK = 1 << 0,
+	DT_NAVLINK_CREATE_EXTREMITY_LINKS	= 1 << 1,
+};
+
 /** Configuration for generated jump down links. */
 struct dtNavLinkBuilderJumpDownConfig
 {
@@ -32,6 +38,9 @@ struct dtNavLinkBuilderJumpDownConfig
 	/// Use greater distance for more filtering (0 to deactivate filtering). [Limit: > 0] [Units: wu]
 	float filterDistanceThreshold = 80.f;
 
+	/// Flags used to indicate how links will be added.
+	unsigned short linkBuilderFlags = DT_NAVLINK_CREATE_CENTER_POINT_LINK;
+
 	/// Cached parabola constant fitting the configuration parameters. 
 	float cachedParabolaConstant = 0;
 
