Temporary fix for backward compatibility crash in standalone caused by deprecated property not being set from the configuration (introduced at CL 34739547)

diff --git a/Engine/Source/Runtime/NavigationSystem/Private/NavMesh/RecastNavMeshGenerator.cpp b/Engine/Source/Runtime/NavigationSystem/Private/NavMesh/RecastNavMeshGenerator.cpp
--- a/Engine/Source/Runtime/NavigationSystem/Private/NavMesh/RecastNavMeshGenerator.cpp
+++ b/Engine/Source/Runtime/NavigationSystem/Private/NavMesh/RecastNavMeshGenerator.cpp
@@ -4990,6 +4990,7 @@ void FRecastNavMeshGenerator::ConfigureBuildProperties(FRecastBuildConfig& OutCo
 {
 	// @TODO those variables should be tweakable per navmesh actor
 	const float CellSize = DestNavMesh->GetCellSize(ENavigationDataResolution::Default);
+	ensure(CellSize != 0.f);
 	const float CellHeight = DestNavMesh->GetCellHeight(ENavigationDataResolution::Default);
 	const float AgentHeight = DestNavMesh->AgentHeight;
 	const float AgentMaxSlope = DestNavMesh->AgentMaxSlope;
diff --git a/Engine/Source/Runtime/NavigationSystem/Public/NavMesh/RecastNavMesh.h b/Engine/Source/Runtime/NavigationSystem/Public/NavMesh/RecastNavMesh.h
--- a/Engine/Source/Runtime/NavigationSystem/Public/NavMesh/RecastNavMesh.h
+++ b/Engine/Source/Runtime/NavigationSystem/Public/NavMesh/RecastNavMesh.h
@@ -661,7 +661,6 @@ class ARecastNavMesh : public ANavigationData
 	UPROPERTY(EditAnywhere, Category=Generation, config, meta=(ClampMin = "300.0"))
 	float TileSizeUU;
 
-#if WITH_EDITORONLY_DATA
 	/**
 	 * Note that we are not using _DEPRECATED on the following deprecated properties
 	 * since it prevents the property from being serialized back which can break the
@@ -679,7 +678,6 @@ class ARecastNavMesh : public ANavigationData
 	UE_DEPRECATED(all, "Use NavMeshResolutionParams to set AgentMaxStepHeight for the different resolutions instead")
 	UPROPERTY(config, meta = (DeprecatedProperty, DeprecationMessage = "Use NavMeshResolutionParams to set AgentMaxStepHeight for the different resolutions instead"))
 	float AgentMaxStepHeight;
-#endif // WITH_EDITORONLY_DATA
 
 	/** Resolution params 
 	 * If using multiple resolutions, it's recommended to chose the highest resolution first and 
