Add tooltip for radii in Navigation Invoker Component

diff --git a/Engine/Source/Runtime/NavigationSystem/Public/NavigationInvokerComponent.h b/Engine/Source/Runtime/NavigationSystem/Public/NavigationInvokerComponent.h
--- a/Engine/Source/Runtime/NavigationSystem/Public/NavigationInvokerComponent.h
+++ b/Engine/Source/Runtime/NavigationSystem/Public/NavigationInvokerComponent.h
@@ -21,9 +21,12 @@ class UNavigationInvokerComponent : public UActorComponent
 
 protected:
 
+	/** Navigation data is requested within a TileGenerationRadius circle around the component owner. */
 	UPROPERTY(EditAnywhere, Category = Navigation, meta = (ClampMin = "0.1", ClampMax = "6400000", UIMin = "0.1", UIMax = "6400000"))
 	float TileGenerationRadius;
 
+	/** Navigation data can be discarded when outside of a TileRemovalRadius circle around the component owner.
+	 * This is computed for all navigation invokers, so a navigation data must be outside of all navigation invokers TileRemovalRadius circles to be discarded. */
 	UPROPERTY(EditAnywhere, Category = Navigation, meta = (ClampMin = "0.1", ClampMax = "6400000", UIMin = "0.1", UIMax = "6400000"))
 	float TileRemovalRadius;
 
