Addition of a boolean to track if a navmesh has been loaded without it's tiles and update the display so it shows as needs to be rebuilt.
Addition of a "Navmesh needs to be rebuilt" message to the navigation display since the stat display is not always enabled.

diff --git a/Engine/Source/Runtime/NavigationSystem/Private/NavMesh/NavMeshRenderingComponent.cpp b/Engine/Source/Runtime/NavigationSystem/Private/NavMesh/NavMeshRenderingComponent.cpp
--- a/Engine/Source/Runtime/NavigationSystem/Private/NavMesh/NavMeshRenderingComponent.cpp
+++ b/Engine/Source/Runtime/NavigationSystem/Private/NavMesh/NavMeshRenderingComponent.cpp
@@ -593,6 +593,11 @@ void FNavMeshSceneProxyData::GatherData(const ARecastNavMesh* NavMesh, int32 InN
 			const ERuntimeGenerationType Mode = NavMesh->GetRuntimeGenerationMode();
 			const FString GenerationMode = Mode == ERuntimeGenerationType::Static ? TEXT("Static") :
 				(Mode == ERuntimeGenerationType::Dynamic ? TEXT("Dynamic") : (Mode == ERuntimeGenerationType::DynamicModifiersOnly ? TEXT("DynamicModifersOnly") : TEXT("Unknown")));
+			
+			if (NavMesh->NeedsRebuild())
+			{
+				DebugLabels.Add(FDebugText(FString::Printf(TEXT("%s"), TEXT("*** NAVMESH NEEDS TO BE REBUILT ***"))));	
+			}
 			DebugLabels.Add(FDebugText(FString::Printf(TEXT("%s (%s%s)"), *NavMesh->GetName(), NavMesh->bIsWorldPartitioned ? TEXT("WP ") : TEXT(""), *GenerationMode)));
 			DebugLabels.Add(FDebugText(FString::Printf(TEXT("AgentRadius %0.1f, AgentHeight %0.1f"), NavMesh->AgentRadius, NavMesh->AgentHeight)));
 			DebugLabels.Add(FDebugText(FString::Printf(
diff --git a/Engine/Source/Runtime/NavigationSystem/Private/NavMesh/PImplRecastNavMesh.cpp b/Engine/Source/Runtime/NavigationSystem/Private/NavMesh/PImplRecastNavMesh.cpp
--- a/Engine/Source/Runtime/NavigationSystem/Private/NavMesh/PImplRecastNavMesh.cpp
+++ b/Engine/Source/Runtime/NavigationSystem/Private/NavMesh/PImplRecastNavMesh.cpp
@@ -661,6 +661,8 @@ void FPImplRecastNavMesh::Serialize( FArchive& Ar, int32 NavMeshVersion )
 			{
 				UE_VLOG(NavMeshOwner, LogNavigation, Error, TEXT("Failed to initialize NavMesh"));
 			}
+			
+			NavMeshOwner->bHasNoTileData = (NumTiles == 0);	
 
 #if WITH_NAVMESH_SEGMENT_LINKS			
 			// Store the tileRefs when iterating over the tile list.
diff --git a/Engine/Source/Runtime/NavigationSystem/Private/NavMesh/RecastNavMesh.cpp b/Engine/Source/Runtime/NavigationSystem/Private/NavMesh/RecastNavMesh.cpp
--- a/Engine/Source/Runtime/NavigationSystem/Private/NavMesh/RecastNavMesh.cpp
+++ b/Engine/Source/Runtime/NavigationSystem/Private/NavMesh/RecastNavMesh.cpp
@@ -593,6 +593,7 @@ void ARecastNavMesh::LoadBeforeGeneratorRebuild()
 void ARecastNavMesh::CleanUp()
 {
 	Super::CleanUp();
+	bHasNoTileData = false;
 	if (NavDataGenerator.IsValid())
 	{
 		NavDataGenerator->CancelBuild();
@@ -3737,7 +3738,7 @@ void ARecastNavMesh::PostEditUndo()
 
 bool ARecastNavMesh::NeedsRebuild() const
 {
-	bool bLooksLikeNeeded = !RecastNavMeshImpl || RecastNavMeshImpl->GetRecastMesh() == 0;
+	bool bLooksLikeNeeded = !RecastNavMeshImpl || RecastNavMeshImpl->GetRecastMesh() == 0 || bHasNoTileData;
 	if (NavDataGenerator.IsValid())
 	{
 		return bLooksLikeNeeded || NavDataGenerator->GetNumRemaningBuildTasks() > 0;
diff --git a/Engine/Source/Runtime/NavigationSystem/Private/NavMesh/RecastNavMeshGenerator.cpp b/Engine/Source/Runtime/NavigationSystem/Private/NavMesh/RecastNavMeshGenerator.cpp
--- a/Engine/Source/Runtime/NavigationSystem/Private/NavMesh/RecastNavMeshGenerator.cpp
+++ b/Engine/Source/Runtime/NavigationSystem/Private/NavMesh/RecastNavMeshGenerator.cpp
@@ -5428,6 +5428,7 @@ bool FRecastNavMeshGenerator::RebuildAll()
 	
 	// Recreate recast navmesh
 	DestNavMesh->GetRecastNavMeshImpl()->ReleaseDetourNavMesh();
+	DestNavMesh->bHasNoTileData = false;
 
 	RcNavMeshOrigin = Unreal2RecastPoint(DestNavMesh->NavMeshOriginOffset);
 
diff --git a/Engine/Source/Runtime/NavigationSystem/Public/NavMesh/RecastNavMesh.h b/Engine/Source/Runtime/NavigationSystem/Public/NavMesh/RecastNavMesh.h
--- a/Engine/Source/Runtime/NavigationSystem/Public/NavMesh/RecastNavMesh.h
+++ b/Engine/Source/Runtime/NavigationSystem/Public/NavMesh/RecastNavMesh.h
@@ -876,6 +876,9 @@ protected:
 	FNavLinkGenerationJumpDownConfig NavLinkJumpDownConfig;
 	
 private:
+	/** @returns true if there were no tiles when the navmesh was loaded. */
+	bool bHasNoTileData : 1 = false;
+	
 	/** Cache rasterized voxels instead of just collision vertices/indices in navigation octree */
 	UPROPERTY(config)
 	uint32 bUseVoxelCache : 1;
