Add way to log all the invokers via a command line

diff --git a/Engine/Source/Runtime/NavigationSystem/Private/NavigationSystem.cpp b/Engine/Source/Runtime/NavigationSystem/Private/NavigationSystem.cpp
--- a/Engine/Source/Runtime/NavigationSystem/Private/NavigationSystem.cpp
+++ b/Engine/Source/Runtime/NavigationSystem/Private/NavigationSystem.cpp
@@ -247,6 +247,20 @@ static FAutoConsoleCommandWithWorldArgsAndOutputDevice CmdDumpOctreeElements(
 		}
 	));
 
+static FAutoConsoleCommandWithWorldArgsAndOutputDevice CmdLogInvokers(
+	TEXT("ai.debug.nav.LogInvokers"),
+	TEXT("Iterate through all the navigation invokers and log details about each of them to the output device."),
+	FConsoleCommandWithWorldArgsAndOutputDeviceDelegate::CreateLambda([](const TArray<FString>& Args, const UWorld* World, FOutputDevice& OutputDevice)
+		{
+#if !UE_BUILD_SHIPPING
+			if (const UNavigationSystemV1* NavSys = FNavigationSystem::GetCurrent<UNavigationSystemV1>(World))
+			{
+				NavSys->DebugLogInvokers(OutputDevice);
+			}
+#endif // !UE_BUILD_SHIPPING
+		}
+	));
+
 	const FNavDataConfig& GetFallbackNavDataConfig()
 	{
 		static FNavDataConfig FallbackNavDataConfig(FNavigationSystem::FallbackAgentRadius, FNavigationSystem::FallbackAgentHeight);
@@ -5831,6 +5845,23 @@ void UNavigationSystemV1::UpdateInvokers()
 #endif // !UE_BUILD_SHIPPING
 }
 
+#if !UE_BUILD_SHIPPING
+void UNavigationSystemV1::DebugLogInvokers(FOutputDevice& OutputDevice) const
+{
+	OutputDevice.Logf(ELogVerbosity::Log, TEXT("Logging %d Invokers:"), Invokers.Num());
+	for (auto It = Invokers.CreateConstIterator(); It; ++It)
+	{
+		const FNavigationInvoker& Invoker = It.Value();
+		OutputDevice.Logf(ELogVerbosity::Log, TEXT("- %s: Radius[Generation:%s Removal:%s] Agents:%d Priority:%s"),
+			*GetNameSafe(It.Key()),
+			*FString::SanitizeFloat(Invoker.GenerationRadius),
+			*FString::SanitizeFloat(Invoker.RemovalRadius),
+			Invoker.SupportedAgents.GetAgentBits(),
+			*UEnum::GetValueAsString(Invoker.Priority));
+	}
+}
+#endif // !UE_BUILD_SHIPPING
+
 void UNavigationSystemV1::UpdateNavDataActiveTiles()
 {
 #if WITH_RECAST
diff --git a/Engine/Source/Runtime/NavigationSystem/Public/NavigationSystem.h b/Engine/Source/Runtime/NavigationSystem/Public/NavigationSystem.h
--- a/Engine/Source/Runtime/NavigationSystem/Public/NavigationSystem.h
+++ b/Engine/Source/Runtime/NavigationSystem/Public/NavigationSystem.h
@@ -459,7 +459,12 @@ private:
 
 	double NextInvokersUpdateTime;
 	NAVIGATIONSYSTEM_API void UpdateInvokers();
-	
+
+public:
+#if !UE_BUILD_SHIPPING
+	void DebugLogInvokers(FOutputDevice& OutputDevice) const;
+#endif // !UE_BUILD_SHIPPING
+
 protected:
 	NAVIGATIONSYSTEM_API virtual void UpdateNavDataActiveTiles();
 
