Only register generated navlink proxies if bGenerateNavLinks is enabled.

diff --git a/Engine/Source/Runtime/NavigationSystem/Private/NavMesh/RecastNavMesh.cpp b/Engine/Source/Runtime/NavigationSystem/Private/NavMesh/RecastNavMesh.cpp
--- a/Engine/Source/Runtime/NavigationSystem/Private/NavMesh/RecastNavMesh.cpp
+++ b/Engine/Source/Runtime/NavigationSystem/Private/NavMesh/RecastNavMesh.cpp
@@ -762,7 +762,7 @@ void ARecastNavMesh::PostRegisterAllComponents()
 	}
 
 	// Create and register link proxy if enabled.
-	if (NavLinkJumpDownConfig.bEnabled && NavLinkJumpDownConfig.LinkProxyClass.Get() != nullptr)
+	if (bGenerateNavLinks && NavLinkJumpDownConfig.bEnabled && NavLinkJumpDownConfig.LinkProxyClass.Get() != nullptr)
 	{
 		// If it's the first time, LinkProxyId will have the default invalid value so we need to generate an Id.
 		const FNavLinkId Id = NavLinkJumpDownConfig.LinkProxyId.IsValid() ? NavLinkJumpDownConfig.LinkProxyId :  FNavLinkId::GenerateUniqueId(); 
