[AutoRTFM] Migrate to UE::TScopeLock instead of FTransactionallySafeScopeLock.

These are functionally identical.

diff --git a/Engine/Source/Runtime/Core/Private/Features/ModularFeatures.cpp b/Engine/Source/Runtime/Core/Private/Features/ModularFeatures.cpp
--- a/Engine/Source/Runtime/Core/Private/Features/ModularFeatures.cpp
+++ b/Engine/Source/Runtime/Core/Private/Features/ModularFeatures.cpp
@@ -4,7 +4,7 @@
 
 #include "CoreTypes.h"
 #include "Misc/AssertionMacros.h"
-#include "Misc/TransactionallySafeScopeLock.h"
+#include "Misc/ScopeLock.h"
 
 IModularFeatures& IModularFeatures::Get()
 {
@@ -54,7 +54,7 @@ IModularFeature* FModularFeatures::GetModularFeatureImplementation( const FName
 
 void FModularFeatures::RegisterModularFeature( const FName Type, IModularFeature* ModularFeature )
 {
-	FTransactionallySafeScopeLock ScopeLock(&ModularFeaturesMapCriticalSection);
+	UE::TScopeLock ScopeLock(ModularFeaturesMapCriticalSection);
 
 	ModularFeaturesMap.AddUnique( Type, ModularFeature );
 	ModularFeatureRegisteredEvent.Broadcast( Type, ModularFeature );
@@ -63,7 +63,7 @@ void FModularFeatures::RegisterModularFeature( const FName Type, IModularFeature
 
 void FModularFeatures::UnregisterModularFeature( const FName Type, IModularFeature* ModularFeature )
 {
-	FTransactionallySafeScopeLock ScopeLock(&ModularFeaturesMapCriticalSection);
+	UE::TScopeLock ScopeLock(ModularFeaturesMapCriticalSection);
 
 	ModularFeaturesMap.RemoveSingle( Type, ModularFeature );
 	ModularFeatureUnregisteredEvent.Broadcast( Type, ModularFeature );
diff --git a/Engine/Source/Runtime/Core/Private/Misc/ConfigCacheIni.cpp b/Engine/Source/Runtime/Core/Private/Misc/ConfigCacheIni.cpp
--- a/Engine/Source/Runtime/Core/Private/Misc/ConfigCacheIni.cpp
+++ b/Engine/Source/Runtime/Core/Private/Misc/ConfigCacheIni.cpp
@@ -22,7 +22,7 @@
 #include "Misc/DataDrivenPlatformInfoRegistry.h"
 #include "Misc/StringBuilder.h"
 #include "Misc/PathViews.h"
-#include "Misc/TransactionallySafeScopeLock.h"
+#include "Misc/ScopeLock.h"
 #include "ProfilingDebugging/CpuProfilerTrace.h"
 #include "ProfilingDebugging/AssetMetadataTrace.h"
 #include "Serialization/MemoryReader.h"
@@ -6278,17 +6278,15 @@ void FConfigCacheIni::RemoveTagFromAllBranches(FName Tag, FConfigModificationTra
 		Pair.Value->RemoveTagFromBranches(Tag, ModificationTracker);
 	}
 #endif
-
 }
 
-
 void FConfigCacheIni::AddPluginToBranches(FName PluginName, FConfigModificationTracker* ModificationTracker)
 {
 	// @todo make sure we are still pending
 	
 	FPluginInfo* PluginInfo = nullptr;
 	{
-		FTransactionallySafeScopeLock Lock(&RegisteredPluginsLock);
+		UE::TScopeLock Lock(RegisteredPluginsLock);
 		
 		PluginInfo = RegisteredPlugins.FindRef(PluginName);
 		if (PluginInfo == nullptr)
@@ -6547,7 +6545,7 @@ void FConfigCacheIni::RegisterPlugin(FName PluginName, const FString& PluginDir,
 	Info->Priority = Priority;
 	Info->bIncludePluginNameInBranchName = bIncludePluginNameInBranchName;
 
-	FTransactionallySafeScopeLock Lock(&RegisteredPluginsLock);
+	UE::TScopeLock Lock(RegisteredPluginsLock);
 	RegisteredPlugins.Add(PluginName, Info);
 }
 
diff --git a/Engine/Source/Runtime/Core/Private/Misc/ConfigContext.cpp b/Engine/Source/Runtime/Core/Private/Misc/ConfigContext.cpp
--- a/Engine/Source/Runtime/Core/Private/Misc/ConfigContext.cpp
+++ b/Engine/Source/Runtime/Core/Private/Misc/ConfigContext.cpp
@@ -2,14 +2,15 @@
 
 #include "Misc/ConfigContext.h"
 #include "Misc/App.h"
+#include "Misc/CommandLine.h"
 #include "Misc/ConfigCacheIni.h"
 #include "Misc/ConfigHierarchy.h"
 #include "Misc/CoreMisc.h"
+#include "Misc/DataDrivenPlatformInfoRegistry.h"
 #include "Misc/MessageDialog.h"
-#include "Misc/CommandLine.h"
-#include "Misc/RemoteConfigIni.h"
 #include "Misc/Paths.h"
-#include "Misc/DataDrivenPlatformInfoRegistry.h"
+#include "Misc/RemoteConfigIni.h"
+#include "Misc/ScopeLock.h"
 #include "ProfilingDebugging/CpuProfilerTrace.h"
 #include "ProfilingDebugging/AssetMetadataTrace.h"
 #include "HAL/LowLevelMemStats.h"
@@ -83,7 +84,7 @@ void FConfigContext::CachePaths()
 		// are we loading a plugin?
 		if (ConfigSystem != nullptr)
 		{
-			FTransactionallySafeScopeLock Lock(&FConfigCacheIni::RegisteredPluginsLock);
+			UE::TScopeLock Lock(FConfigCacheIni::RegisteredPluginsLock);
 			
 			FConfigCacheIni::FPluginInfo* PluginInfo = ConfigSystem->RegisteredPlugins.FindRef(ConfigFileTag == NAME_None ? FName(*BaseIniName) : ConfigFileTag);
 			if (PluginInfo != nullptr)
diff --git a/Engine/Source/Runtime/Core/Public/Delegates/DelegateAccessHandler.h b/Engine/Source/Runtime/Core/Public/Delegates/DelegateAccessHandler.h
--- a/Engine/Source/Runtime/Core/Public/Delegates/DelegateAccessHandler.h
+++ b/Engine/Source/Runtime/Core/Public/Delegates/DelegateAccessHandler.h
@@ -5,9 +5,8 @@
 #include "CoreTypes.h"
 #include "HAL/CriticalSection.h"  // jira SOL-6812: Remove this.
 #include "Misc/MTAccessDetector.h"
+#include "Misc/ScopeLock.h"
 #include "Misc/TransactionallySafeCriticalSection.h"
-#include "Misc/TransactionallySafeScopeLock.h"
-#include "Misc/ScopeLock.h"  // jira SOL-6812: Remove this.
 #include "AutoRTFM/AutoRTFM.h"
 
 //#define UE_DETECT_DELEGATES_RACE_CONDITIONS 0
@@ -74,17 +73,17 @@ class TDelegateAccessHandlerBase<FThreadSafeDelegateMode>
 	friend struct TWriteLockedDelegateAllocation<FThreadSafeDelegateMode>;
 
 protected:
-	struct FReadAccessScope { FTransactionallySafeScopeLock Lock; };
-	struct FWriteAccessScope { FTransactionallySafeScopeLock Lock; };
+	struct FReadAccessScope { UE::TScopeLock<FTransactionallySafeCriticalSection> Lock; };
+	struct FWriteAccessScope { UE::TScopeLock<FTransactionallySafeCriticalSection> Lock; };
 
 	[[nodiscard]] FReadAccessScope GetReadAccessScope() const
 	{
-		return { &Mutex };
+		return FReadAccessScope{Mutex};
 	}
 
 	[[nodiscard]] FWriteAccessScope GetWriteAccessScope()
 	{
-		return { &Mutex };
+		return FWriteAccessScope{Mutex};
 	}
 
 private:
diff --git a/Engine/Source/Runtime/CoreUObject/Private/AssetRegistry/AssetDataTagMap.cpp b/Engine/Source/Runtime/CoreUObject/Private/AssetRegistry/AssetDataTagMap.cpp
--- a/Engine/Source/Runtime/CoreUObject/Private/AssetRegistry/AssetDataTagMap.cpp
+++ b/Engine/Source/Runtime/CoreUObject/Private/AssetRegistry/AssetDataTagMap.cpp
@@ -7,7 +7,6 @@
 #include "Misc/PackageName.h"
 #include "Misc/Optional.h"
 #include "Misc/TransactionallySafeCriticalSection.h"
-#include "Misc/TransactionallySafeScopeLock.h"
 #include "Misc/ScopeLock.h"
 #include "Serialization/ArrayWriter.h"
 #include "Serialization/MemoryReader.h"
@@ -377,7 +376,7 @@ namespace FixedTagPrivate
 
 		FStore* CreateAndRegister()
 		{
-			FTransactionallySafeScopeLock Lock(&Mutex);
+			UE::TScopeLock Lock(Mutex);
 
 			static constexpr uint32 Mask = Capacity - 1; 
 			for (uint32 RawIndex = ProbeIndex, RawEnd = ProbeIndex + Capacity; RawIndex < RawEnd + Capacity; ++RawIndex)
@@ -399,7 +398,7 @@ namespace FixedTagPrivate
 
 		void Unregister(FStore& Store)
 		{
-			FTransactionallySafeScopeLock Lock(&Mutex);
+			UE::TScopeLock Lock(Mutex);
 
 			check(Store.Index < Capacity);
 			check(Data[Store.Index] == &Store);
diff --git a/Engine/Source/Runtime/Engine/Private/StreamableManager.cpp b/Engine/Source/Runtime/Engine/Private/StreamableManager.cpp
--- a/Engine/Source/Runtime/Engine/Private/StreamableManager.cpp
+++ b/Engine/Source/Runtime/Engine/Private/StreamableManager.cpp
@@ -6,8 +6,8 @@
 #include "UObject/Package.h"
 #include "Misc/PackageAccessTrackingOps.h"
 #include "Misc/PackageName.h"
+#include "Misc/ScopeLock.h"
 #include "Misc/TransactionallySafeCriticalSection.h"
-#include "Misc/TransactionallySafeScopeLock.h"
 #include "UObject/UObjectThreadContext.h"
 #include "HAL/IConsoleManager.h"
 #include "Tickable.h"
@@ -61,7 +61,7 @@ public:
 		FPendingDelegate* PendingDelegate = new FPendingDelegate(Forward<InStreamableDelegate>(Delegate), Forward<InStreamableDelegate>(CancelDelegate), AssociatedHandle);
 
 		{
-			FTransactionallySafeScopeLock Lock(&DataLock);
+			UE::TScopeLock Lock(DataLock);
 
 			FPendingDelegateList& DelegatesForHandle = PendingDelegatesByHandle.FindOrAdd(AssociatedHandle);
 			FPendingDelegateList::AddTail(PendingDelegate, PendingDelegates, DelegatesForHandle);
@@ -74,7 +74,7 @@ public:
 		FPendingDelegateList PendingDelegatesToDelete;
 
 		{
-			FTransactionallySafeScopeLock _(&DataLock);
+			UE::TScopeLock _(DataLock);
 
 			FPendingDelegateList* DelegatesForHandle = PendingDelegatesByHandle.Find(AssociatedHandle);
 			if (!DelegatesForHandle)
@@ -144,7 +144,7 @@ public:
 
 		FPendingDelegateList DelegatesToCall;
 		{
-			FTransactionallySafeScopeLock Lock(&DataLock);
+			UE::TScopeLock Lock(DataLock);
 
 			FPendingDelegate* CurrentNode = PendingDelegates.Head;
 			while (CurrentNode)
diff --git a/Engine/Source/Runtime/NavigationSystem/Private/NavMesh/RecastNavMesh.cpp b/Engine/Source/Runtime/NavigationSystem/Private/NavMesh/RecastNavMesh.cpp
--- a/Engine/Source/Runtime/NavigationSystem/Private/NavMesh/RecastNavMesh.cpp
+++ b/Engine/Source/Runtime/NavigationSystem/Private/NavMesh/RecastNavMesh.cpp
@@ -24,7 +24,7 @@
 #include "Math/Color.h"
 #include "NavigationDataHandler.h"
 #include "BaseGeneratedNavLinksProxy.h"
-#include "Misc/TransactionallySafeScopeLock.h"
+#include "Misc/ScopeLock.h"
 
 #if WITH_EDITOR
 #include "EditorSupportDelegates.h"
@@ -2863,7 +2863,7 @@ void ARecastNavMesh::InvalidateAffectedPaths(const TArray<FNavTileRef>& ChangedT
 	// before starting async queries task but protecting ActivePaths will make
 	// the system safer in case of future timing changes.
 	{
-		FTransactionallySafeScopeLock PathLock(&ActivePathsLock);
+		UE::TScopeLock PathLock(ActivePathsLock);
 
 		for (int32 PathIndex = PathsCount - 1; PathIndex >= 0; --PathIndex)
 		{
diff --git a/Engine/Source/Runtime/NavigationSystem/Private/NavigationData.cpp b/Engine/Source/Runtime/NavigationSystem/Private/NavigationData.cpp
--- a/Engine/Source/Runtime/NavigationSystem/Private/NavigationData.cpp
+++ b/Engine/Source/Runtime/NavigationSystem/Private/NavigationData.cpp
@@ -13,7 +13,7 @@
 #include "NavFilters/NavigationQueryFilter.h"
 #include "NavigationSystem.h"
 #include "VisualLogger/VisualLogger.h"
-#include "Misc/TransactionallySafeScopeLock.h"
+#include "Misc/ScopeLock.h"
 
 #include UE_INLINE_GENERATED_CPP_BY_NAME(NavigationData)
 
@@ -424,7 +424,7 @@ void ANavigationData::PurgeUnusedPaths()
 
 	// Paths can be registered from async pathfinding thread 
 	// while unused paths are purged in main thread (actor tick)
-	FTransactionallySafeScopeLock PathLock(&ActivePathsLock);
+	UE::TScopeLock PathLock(ActivePathsLock);
 
 	const int32 Count = ActivePaths.Num();
 	for (int32 PathIndex = Count - 1; PathIndex >= 0; --PathIndex)
@@ -440,7 +440,7 @@ void ANavigationData::PurgeUnusedPaths()
 void ANavigationData::RegisterActivePath(FNavPathSharedPtr SharedPath)
 {
 	// Paths can be registered from main thread and async pathfinding thread
-	FTransactionallySafeScopeLock PathLock(&ActivePathsLock);
+	UE::TScopeLock PathLock(ActivePathsLock);
 	ActivePaths.Add(SharedPath);
 }
 
@@ -861,7 +861,7 @@ uint32 ANavigationData::LogMemUsed() const
 	{
 		// Paths can be registered from async pathfinding thread
 		// while logging is requested on main thread (console command)
-		FTransactionallySafeScopeLock PathLock(&ActivePathsLock);
+		UE::TScopeLock PathLock(ActivePathsLock);
 		ActivePathsMemSize = ActivePaths.GetAllocatedSize();
 	}
 
diff --git a/Engine/Source/Runtime/NavigationSystem/Public/NavigationData.h b/Engine/Source/Runtime/NavigationSystem/Public/NavigationData.h
--- a/Engine/Source/Runtime/NavigationSystem/Public/NavigationData.h
+++ b/Engine/Source/Runtime/NavigationSystem/Public/NavigationData.h
@@ -18,7 +18,6 @@
 #include "NavigationSystemTypes.h"
 #include "EngineDefines.h"
 #include "AI/Navigation/NavigationDataInterface.h"
-#include "Misc/TransactionallySafeScopeLock.h"
 
 #include "NavigationData.generated.h"
 
