Prevent generation of areas when the Bounds contains a NaN

diff --git a/Engine/Source/Runtime/NavigationSystem/Private/NavMesh/RecastNavMeshGenerator.cpp b/Engine/Source/Runtime/NavigationSystem/Private/NavMesh/RecastNavMeshGenerator.cpp
--- a/Engine/Source/Runtime/NavigationSystem/Private/NavMesh/RecastNavMeshGenerator.cpp
+++ b/Engine/Source/Runtime/NavigationSystem/Private/NavMesh/RecastNavMeshGenerator.cpp
@@ -6370,8 +6370,8 @@ void FRecastNavMeshGenerator::MarkDirtyTiles(const TArray<FNavigationDirtyArea>&
 		{
 			QUICK_SCOPE_CYCLE_COUNTER(STAT_RecastNavMeshGenerator_DirtyArea);
 
-			if (!ensureMsgf(DirtyArea.Bounds.IsValid, TEXT("%hs Attempting to use DirtyArea.Bounds which are not valid. Source: %s"),
-					__FUNCTION__, *DirtyArea.GetSourceDescription()))
+			if (!ensureMsgf(DirtyArea.Bounds.IsValid && !DirtyArea.Bounds.ContainsNaN(), TEXT("%hs Attempting to use DirtyArea.Bounds which are not valid%s. Bounds: %s, Source: %s"),
+					__FUNCTION__, DirtyArea.Bounds.ContainsNaN() ? TEXT(" (contains NaN)") : TEXT(""), *DirtyArea.Bounds.ToString(), *DirtyArea.GetSourceDescription()))
 			{
 				continue;
 			}
