Recast Navmesh: Add information when a path find request fails because either its start or end location is off-navmesh

diff --git a/Engine/Source/Runtime/Engine/Public/AI/Navigation/NavQueryFilter.h b/Engine/Source/Runtime/Engine/Public/AI/Navigation/NavQueryFilter.h
--- a/Engine/Source/Runtime/Engine/Public/AI/Navigation/NavQueryFilter.h
+++ b/Engine/Source/Runtime/Engine/Public/AI/Navigation/NavQueryFilter.h
@@ -5,7 +5,7 @@
 #include "CoreMinimal.h"
 
 // For now leave all the real numbers here as floats instead of conversion to FVector::FReal. Its not worth changing the virtual functions and breaking
-// any existing code. Its not required for heuristic scale or area costs to be larger or have greater precison than a float.
+// any existing code. Its not required for heuristic scale or area costs to be larger or have greater precision than a float.
 class INavigationQueryFilterInterface
 {
 public:
diff --git a/Engine/Source/Runtime/NavigationSystem/Private/NavMesh/PImplRecastNavMesh.cpp b/Engine/Source/Runtime/NavigationSystem/Private/NavMesh/PImplRecastNavMesh.cpp
--- a/Engine/Source/Runtime/NavigationSystem/Private/NavMesh/PImplRecastNavMesh.cpp
+++ b/Engine/Source/Runtime/NavigationSystem/Private/NavMesh/PImplRecastNavMesh.cpp
@@ -1201,7 +1201,7 @@ ENavigationQueryResult::Type FPImplRecastNavMesh::FindPath(const FVector& StartL
 
 	FVector RecastStartPos, RecastEndPos;
 	NavNodeRef StartPolyID, EndPolyID;
-	const bool bCanSearch = InitPathfinding(StartLoc, EndLoc, NavQuery, QueryFilter, RecastStartPos, StartPolyID, RecastEndPos, EndPolyID);
+	const bool bCanSearch = InitPathfinding(StartLoc, EndLoc, NavQuery, QueryFilter, RecastStartPos, StartPolyID, RecastEndPos, EndPolyID, &Path);
 	if (!bCanSearch)
 	{
 		return ENavigationQueryResult::Error;
@@ -1339,7 +1339,7 @@ ENavigationQueryResult::Type FPImplRecastNavMesh::TestClusterPath(const FVector&
 bool FPImplRecastNavMesh::InitPathfinding(const FVector& UnrealStart, const FVector& UnrealEnd,
 	const dtNavMeshQuery& Query, const dtQueryFilter* Filter,
 	FVector& RecastStart, dtPolyRef& StartPoly,
-	FVector& RecastEnd, dtPolyRef& EndPoly) const
+	FVector& RecastEnd, dtPolyRef& EndPoly, FNavMeshPath* OutPath) const
 {
 	const FVector NavExtent = NavMeshOwner->GetModifiedQueryExtent(NavMeshOwner->GetDefaultQueryExtent());
 	const FVector::FReal Extent[3] = { NavExtent.X, NavExtent.Z, NavExtent.Y };
@@ -1356,6 +1356,10 @@ bool FPImplRecastNavMesh::InitPathfinding(const FVector& UnrealStart, const FVec
 		UE_VLOG_LOCATION(NavMeshOwner, LogNavigation, Warning, UnrealStart, 15, FColor::Red, TEXT("Start failed"));
 		UE_VLOG_BOX(NavMeshOwner, LogNavigation, Warning, FBox(UnrealStart - NavExtent, UnrealStart + NavExtent), FColor::Red, TEXT_EMPTY);
 
+		if (OutPath)
+		{
+			OutPath->SetErrorStartLocationNonNavigable(true);
+		}
 		return false;
 	}
 
@@ -1370,6 +1374,10 @@ bool FPImplRecastNavMesh::InitPathfinding(const FVector& UnrealStart, const FVec
 			UE_VLOG_LOCATION(NavMeshOwner, LogNavigation, Warning, UnrealEnd, 15, FColor::Red, TEXT("End failed"));
 			UE_VLOG_BOX(NavMeshOwner, LogNavigation, Warning, FBox(UnrealEnd - NavExtent, UnrealEnd + NavExtent), FColor::Red, TEXT_EMPTY);
 
+			if (OutPath)
+			{
+				OutPath->SetErrorEndLocationNonNavigable(true);
+			}
 			return false;
 		}
 
diff --git a/Engine/Source/Runtime/NavigationSystem/Public/NavMesh/PImplRecastNavMesh.h b/Engine/Source/Runtime/NavigationSystem/Public/NavMesh/PImplRecastNavMesh.h
--- a/Engine/Source/Runtime/NavigationSystem/Public/NavMesh/PImplRecastNavMesh.h
+++ b/Engine/Source/Runtime/NavigationSystem/Public/NavMesh/PImplRecastNavMesh.h
@@ -289,11 +289,22 @@ public:
 	/** Helper function to serialize a Recast tile compressed data. */
 	static NAVIGATIONSYSTEM_API void SerializeCompressedTileCacheData(FArchive& Ar, int32 NavMeshVersion, unsigned char*& CompressedData, int32& CompressedDataSize);
 
-	/** Initialize data for pathfinding */
+	/** Initialize data for pathfinding
+	* @param UnrealStart Start location for the desired path
+	* @param UnrealEnd End location for the desired path
+	* @param Query Query to use to project start and end position to the navmesh
+	* @param Filter Filters which polygons are acceptable to project start and end locations
+	* @param RecastStart Outputs the start location snapped to the navmesh
+	* @param StartPoly Outputs the navmesh poly at the start location
+	* @param RecastEnd Outputs the end location snapped to the navmesh. Equal to UnrealEnd when there's no navmesh at the end location.
+	* @param EndPoly Outputs the navmesh poly at the end location. Invalid when there's no navmesh at the end location
+	* @param OutPath [optional] When provided, Path will be used to store relevant error codes that this method can produce
+	*/
 	NAVIGATIONSYSTEM_API bool InitPathfinding(const FVector& UnrealStart, const FVector& UnrealEnd, 
 		const dtNavMeshQuery& Query, const dtQueryFilter* Filter,
 		FVector& RecastStart, dtPolyRef& StartPoly,
-		FVector& RecastEnd, dtPolyRef& EndPoly) const;
+		FVector& RecastEnd, dtPolyRef& EndPoly,
+		FNavMeshPath* OutPath = nullptr) const;
 
 	/** Marks path flags, perform string pulling if needed */
 	NAVIGATIONSYSTEM_API void PostProcessPath(dtStatus PathfindResult, FNavMeshPath& Path,
@@ -357,7 +368,7 @@ ENavigationQueryResult::Type FPImplRecastNavMesh::FindPathCustomAStar(TRecastASt
 
 	FVector RecastStartPos, RecastEndPos;
 	NavNodeRef StartPolyID, EndPolyID;
-	const bool bCanSearch = InitPathfinding(StartLoc, EndLoc, RecastGraphWrapper.GetRecastQuery(), QueryFilter, RecastStartPos, StartPolyID, RecastEndPos, EndPolyID);
+	const bool bCanSearch = InitPathfinding(StartLoc, EndLoc, RecastGraphWrapper.GetRecastQuery(), QueryFilter, RecastStartPos, StartPolyID, RecastEndPos, EndPolyID, &Path);
 	if (!bCanSearch)
 	{
 		return ENavigationQueryResult::Error;
diff --git a/Engine/Source/Runtime/NavigationSystem/Public/NavigationData.h b/Engine/Source/Runtime/NavigationSystem/Public/NavigationData.h
--- a/Engine/Source/Runtime/NavigationSystem/Public/NavigationData.h
+++ b/Engine/Source/Runtime/NavigationSystem/Public/NavigationData.h
@@ -117,6 +117,20 @@ struct FNavigationPath : public TSharedFromThis<FNavigationPath, ESPMode::Thread
 	{
 		return bReachedSearchLimit;
 	}
+
+	// True when the path request failed because the start location couldn't be projected on the NavData.
+	// False doesn't guarantee that the location is navigable, it's also possible a previous error skipped the projection altogether
+	bool IsErrorStartLocationNonNavigable() const
+	{
+		return bErrorStartLocationNonNavigable;
+	}
+
+	// True when the path request failed because the end location couldn't be projected on the NavData.
+	// False doesn't guarantee that the location is navigable, it's also possible a previous error skipped the projection altogether
+	bool IsErrorEndLocationNonNavigable() const
+	{
+		return bErrorEndLocationNonNavigable;
+	}
 	FORCEINLINE bool IsWaitingForRepath() const
 	{
 		return bWaitingForRepath;
@@ -188,7 +202,14 @@ struct FNavigationPath : public TSharedFromThis<FNavigationPath, ESPMode::Thread
 	{
 		bReachedSearchLimit = bLimited;
 	}
-
+	void SetErrorStartLocationNonNavigable(const bool bErrorNonNavigable)
+	{
+		bErrorStartLocationNonNavigable = bErrorNonNavigable;
+	}
+	void SetErrorEndLocationNonNavigable(const bool bErrorNonNavigable)
+	{
+		bErrorEndLocationNonNavigable = bErrorNonNavigable;
+	}
 	FORCEINLINE void SetFilter(FSharedConstNavQueryFilter InFilter)
 	{
 		PathFindingQueryData.QueryFilter = InFilter;
@@ -449,6 +470,12 @@ protected:
 	*	although it might lead closer to destination. */
 	uint32 bReachedSearchLimit : 1;
 
+	/** "true" when the start location could not be projected on the navigation data */
+	uint32 bErrorStartLocationNonNavigable : 1 = false;
+
+	/** "true" when the end location could not be projected on the navigation data */
+	uint32 bErrorEndLocationNonNavigable : 1 = false;
+
 	/** if true path will request re-pathing if it gets invalidated due to underlying navigation changed */
 	uint32 bDoAutoUpdateOnInvalidation : 1;
 
diff --git a/Engine/Source/Runtime/NavigationSystem/Public/NavigationSystem.h b/Engine/Source/Runtime/NavigationSystem/Public/NavigationSystem.h
--- a/Engine/Source/Runtime/NavigationSystem/Public/NavigationSystem.h
+++ b/Engine/Source/Runtime/NavigationSystem/Public/NavigationSystem.h
@@ -645,10 +645,8 @@ public:
 	/** 
 	 *	Asynchronously looks for a path from @StartLocation to @EndLocation for agent with properties @AgentProperties. NavData actor appropriate for specified 
 	 *	FNavAgentProperties will be found automatically
+	 *	@param AgentProperties optional agent properties used to determine the NavData to use, only when the NavData inside the FPathFindingQuery was left empty.
 	 *	@param ResultDelegate delegate that will be called once query has been processed and finished. Will be called even if query fails - in such case see comments for delegate's params
-	 *	@param NavData optional navigation data that will be used instead of the one that would be deducted from AgentProperties
-	 *	@param PathToFill if points to an actual navigation path instance than this instance will be filled with resulting path. Otherwise a new instance will be created and 
-	 *		used in call to ResultDelegate
 	 *  @param Mode switch between normal and hierarchical path finding algorithms
 	 *	@return request ID
 	 */
