Filled convex rasterization: add a new rasterization mode for collision that consider the geometry as a single filled vertical convex volume. This way the interior is filled and won't generate navmesh

diff --git a/Engine/Source/Runtime/Engine/Classes/Components/PrimitiveComponent.h b/Engine/Source/Runtime/Engine/Classes/Components/PrimitiveComponent.h
--- a/Engine/Source/Runtime/Engine/Classes/Components/PrimitiveComponent.h
+++ b/Engine/Source/Runtime/Engine/Classes/Components/PrimitiveComponent.h
@@ -630,6 +630,12 @@ public:
 	UPROPERTY(EditAnywhere, Category = Navigation)
 	uint8 bFillCollisionUnderneathForNavmesh:1;
 
+	/** If set, the geometry gathered for navigation data generation will be converted into a filled vertical convex volume.
+	 * It means that all collision geometries of the asset are merged vertically resulting in a grid of vertical columns that encompass the asset.
+	 * This can be useful to represent the interior of the surface and prevent navmesh inside. */
+	UPROPERTY(EditAnywhere, Category = Navigation)
+	uint8 bRasterizeAsFilledConvexVolume:1;
+
 	// General flags.
 	
 	/** If this is True, this component must always be loaded on clients, even if Hidden and CollisionEnabled is NoCollision. */
@@ -3092,6 +3098,9 @@ public:
 
 	/** Whether the physics is static for this component. */
 	FORCEINLINE bool IsStaticPhysics() const { return Mobility != EComponentMobility::Movable && bStaticWhenNotMoveable; }
+
+protected:
+	void AddNavigationModifier(FNavigationRelevantData& OutData) const;
 };
 
 /** 
diff --git a/Engine/Source/Runtime/Engine/Private/AI/Navigation/NavigationModifier.cpp b/Engine/Source/Runtime/Engine/Private/AI/Navigation/NavigationModifier.cpp
--- a/Engine/Source/Runtime/Engine/Private/AI/Navigation/NavigationModifier.cpp
+++ b/Engine/Source/Runtime/Engine/Private/AI/Navigation/NavigationModifier.cpp
@@ -777,6 +777,7 @@ void FCompositeNavModifier::Reset()
 	bIsPerInstanceModifier = false;
 	bFillCollisionUnderneathForNavmesh = false;
 	bMaskFillCollisionUnderneathForNavmesh = false;
+	bRasterizeAsFilledConvexVolume = false;
 	NavMeshResolution = ENavigationDataResolution::Invalid;
 }
 
@@ -789,6 +790,7 @@ void FCompositeNavModifier::Empty()
 	bAdjustHeight = false;
 	bFillCollisionUnderneathForNavmesh = false;
 	bMaskFillCollisionUnderneathForNavmesh = false;
+	bRasterizeAsFilledConvexVolume = false;
 	NavMeshResolution = ENavigationDataResolution::Invalid;
 }
 
diff --git a/Engine/Source/Runtime/Engine/Private/Components/PrimitiveComponent.cpp b/Engine/Source/Runtime/Engine/Private/Components/PrimitiveComponent.cpp
--- a/Engine/Source/Runtime/Engine/Private/Components/PrimitiveComponent.cpp
+++ b/Engine/Source/Runtime/Engine/Private/Components/PrimitiveComponent.cpp
@@ -325,6 +325,7 @@ UPrimitiveComponent::UPrimitiveComponent(const FObjectInitializer& ObjectInitial
 	bConsiderForActorPlacementWhenHidden = false;
 #endif // WITH_EDITORONLY_DATA
 	bFillCollisionUnderneathForNavmesh = false;
+	bRasterizeAsFilledConvexVolume = false;
 	AlwaysLoadOnClient = true;
 	AlwaysLoadOnServer = true;
 	SetCollisionProfileName(UCollisionProfile::BlockAll_ProfileName);
@@ -3140,10 +3141,16 @@ bool UPrimitiveComponent::ShouldDispatchWakeEvents(FName BoneName) const
 
 void UPrimitiveComponent::GetNavigationData(FNavigationRelevantData& OutData) const
 {
-	if (bFillCollisionUnderneathForNavmesh)
+	AddNavigationModifier(OutData);
+}
+
+void UPrimitiveComponent::AddNavigationModifier(FNavigationRelevantData& OutData) const
+{
+	if (bFillCollisionUnderneathForNavmesh || bRasterizeAsFilledConvexVolume)
 	{
 		FCompositeNavModifier CompositeNavModifier;
 		CompositeNavModifier.SetFillCollisionUnderneathForNavmesh(bFillCollisionUnderneathForNavmesh);
+		CompositeNavModifier.SetRasterizeAsFilledConvexVolume(bRasterizeAsFilledConvexVolume);
 		OutData.Modifiers.Add(CompositeNavModifier);
 	}
 }
diff --git a/Engine/Source/Runtime/Engine/Private/ISMPartition/ISMComponentDescriptor.cpp b/Engine/Source/Runtime/Engine/Private/ISMPartition/ISMComponentDescriptor.cpp
--- a/Engine/Source/Runtime/Engine/Private/ISMPartition/ISMComponentDescriptor.cpp
+++ b/Engine/Source/Runtime/Engine/Private/ISMPartition/ISMComponentDescriptor.cpp
@@ -134,6 +134,7 @@ void FISMComponentDescriptorBase::InitFrom(const UStaticMeshComponent* Template,
 	bOverrideNavigationExport = Template->bOverrideNavigationExport;
 	bForceNavigationObstacle = Template->bForceNavigationObstacle;
 	bFillCollisionUnderneathForNavmesh = Template->bFillCollisionUnderneathForNavmesh;
+	bRasterizeAsFilledConvexVolume = Template->bRasterizeAsFilledConvexVolume;
 
 	if (GISMComponentDescriptorCopiesCanEverAffectNavigation)
 	{
@@ -254,6 +255,7 @@ bool FISMComponentDescriptorBase::operator==(const FISMComponentDescriptorBase&
 	bForceNavigationObstacle == Other.bForceNavigationObstacle &&
 	bCanEverAffectNavigation == Other.bCanEverAffectNavigation &&
 	bFillCollisionUnderneathForNavmesh == Other.bFillCollisionUnderneathForNavmesh &&
+	bRasterizeAsFilledConvexVolume == Other.bRasterizeAsFilledConvexVolume &&
 	WorldPositionOffsetDisableDistance == Other.WorldPositionOffsetDisableDistance &&
 	ShadowCacheInvalidationBehavior == Other.ShadowCacheInvalidationBehavior &&
 	DetailMode == Other.DetailMode &&
@@ -381,6 +383,7 @@ void FISMComponentDescriptorBase::InitComponent(UInstancedStaticMeshComponent* I
 	ISMComponent->bOverrideNavigationExport = bOverrideNavigationExport;
 	ISMComponent->bForceNavigationObstacle = bForceNavigationObstacle;
 	ISMComponent->bFillCollisionUnderneathForNavmesh = bFillCollisionUnderneathForNavmesh;
+	ISMComponent->bRasterizeAsFilledConvexVolume = bRasterizeAsFilledConvexVolume;
 	ISMComponent->WorldPositionOffsetDisableDistance = WorldPositionOffsetDisableDistance;
 	ISMComponent->ShadowCacheInvalidationBehavior = ShadowCacheInvalidationBehavior;
 	ISMComponent->DetailMode = DetailMode;
diff --git a/Engine/Source/Runtime/Engine/Private/InstancedStaticMesh.cpp b/Engine/Source/Runtime/Engine/Private/InstancedStaticMesh.cpp
--- a/Engine/Source/Runtime/Engine/Private/InstancedStaticMesh.cpp
+++ b/Engine/Source/Runtime/Engine/Private/InstancedStaticMesh.cpp
@@ -5479,12 +5479,7 @@ bool UInstancedStaticMeshComponent::ComponentOverlapMultiImpl(TArray<struct FOve
 
 void UInstancedStaticMeshComponent::GetNavigationData(FNavigationRelevantData& Data) const
 {
-	if (bFillCollisionUnderneathForNavmesh)
-	{
-		FCompositeNavModifier CompositeNavModifier;
-		CompositeNavModifier.SetFillCollisionUnderneathForNavmesh(bFillCollisionUnderneathForNavmesh);
-		Data.Modifiers.Add(CompositeNavModifier);
-	}
+	AddNavigationModifier(Data);
 
 	// Navigation data will get refreshed once async compilation finishes
 	if (GetStaticMesh() && !GetStaticMesh()->IsCompiling() && GetStaticMesh()->GetNavCollision())
diff --git a/Engine/Source/Runtime/Engine/Public/AI/NavigationModifier.h b/Engine/Source/Runtime/Engine/Public/AI/NavigationModifier.h
--- a/Engine/Source/Runtime/Engine/Public/AI/NavigationModifier.h
+++ b/Engine/Source/Runtime/Engine/Public/AI/NavigationModifier.h
@@ -280,6 +280,7 @@ struct FCompositeNavModifier : public FNavigationModifier
 		, bIsPerInstanceModifier(false)
 		, bFillCollisionUnderneathForNavmesh(false)
 		, bMaskFillCollisionUnderneathForNavmesh(false)
+		, bRasterizeAsFilledConvexVolume(false)
 		, NavMeshResolution(ENavigationDataResolution::Invalid)
 	{}
 
@@ -292,6 +293,7 @@ struct FCompositeNavModifier : public FNavigationModifier
 		return (Areas.Num() == 0) && (SimpleLinks.Num() == 0) && (CustomLinks.Num() == 0) && 
 			!bFillCollisionUnderneathForNavmesh && 
 			!bMaskFillCollisionUnderneathForNavmesh &&
+			!bRasterizeAsFilledConvexVolume &&
 			NavMeshResolution == ENavigationDataResolution::Invalid;
 	}
 
@@ -331,6 +333,7 @@ struct FCompositeNavModifier : public FNavigationModifier
 		bHasLowAreaModifiers |= Modifiers.HasLowAreaModifiers();
 		bFillCollisionUnderneathForNavmesh |= Modifiers.GetFillCollisionUnderneathForNavmesh();
 		bMaskFillCollisionUnderneathForNavmesh |= Modifiers.GetMaskFillCollisionUnderneathForNavmesh();
+		bRasterizeAsFilledConvexVolume |= Modifiers.GetRasterizeAsFilledConvexVolume();
 		if (Modifiers.GetNavMeshResolution() != ENavigationDataResolution::Invalid)
 		{
 			NavMeshResolution = FMath::Max(NavMeshResolution, Modifiers.GetNavMeshResolution());	// Pick the highest resolution
@@ -354,6 +357,8 @@ struct FCompositeNavModifier : public FNavigationModifier
 	FORCEINLINE void SetFillCollisionUnderneathForNavmesh(bool bValue) { bFillCollisionUnderneathForNavmesh = bValue; }
 	FORCEINLINE bool GetMaskFillCollisionUnderneathForNavmesh() const { return bMaskFillCollisionUnderneathForNavmesh; }
 	FORCEINLINE void SetMaskFillCollisionUnderneathForNavmesh(bool bValue) { bMaskFillCollisionUnderneathForNavmesh = bValue; }
+	FORCEINLINE bool GetRasterizeAsFilledConvexVolume() const { return bRasterizeAsFilledConvexVolume; }
+	FORCEINLINE void SetRasterizeAsFilledConvexVolume(bool bValue) { bRasterizeAsFilledConvexVolume = bValue; }
 	FORCEINLINE ENavigationDataResolution GetNavMeshResolution() const { return bEnableNavMeshResolutions ? NavMeshResolution : ENavigationDataResolution::Invalid; }
 	FORCEINLINE void SetNavMeshResolution(ENavigationDataResolution Resolution) { NavMeshResolution = Resolution; }
 	FORCEINLINE void ReserveForAdditionalAreas(int32 AdditionalElementsCount) { Areas.Reserve(Areas.Num() + AdditionalElementsCount); }
@@ -390,5 +395,6 @@ private:
     uint32 bIsPerInstanceModifier : 1;
 	uint32 bFillCollisionUnderneathForNavmesh : 1;
 	uint32 bMaskFillCollisionUnderneathForNavmesh : 1;
+	uint32 bRasterizeAsFilledConvexVolume : 1;
 	ENavigationDataResolution NavMeshResolution;
 };
diff --git a/Engine/Source/Runtime/Engine/Public/ISMPartition/ISMComponentDescriptor.h b/Engine/Source/Runtime/Engine/Public/ISMPartition/ISMComponentDescriptor.h
--- a/Engine/Source/Runtime/Engine/Public/ISMPartition/ISMComponentDescriptor.h
+++ b/Engine/Source/Runtime/Engine/Public/ISMPartition/ISMComponentDescriptor.h
@@ -204,6 +204,9 @@ public:
 	UPROPERTY(EditAnywhere, Category = "Component Settings")
 	uint8 bFillCollisionUnderneathForNavmesh : 1;
 
+	UPROPERTY(EditAnywhere, Category = "Component Settings")
+	uint8 bRasterizeAsFilledConvexVolume : 1;
+
 	UPROPERTY(EditAnywhere, Category = "Component Settings")
 	int32 WorldPositionOffsetDisableDistance;
 
diff --git a/Engine/Source/Runtime/NavigationSystem/Private/NavMesh/RecastNavMeshGenerator.cpp b/Engine/Source/Runtime/NavigationSystem/Private/NavMesh/RecastNavMeshGenerator.cpp
--- a/Engine/Source/Runtime/NavigationSystem/Private/NavMesh/RecastNavMeshGenerator.cpp
+++ b/Engine/Source/Runtime/NavigationSystem/Private/NavMesh/RecastNavMeshGenerator.cpp
@@ -275,6 +275,25 @@ FRecastGeometryCache::FRecastGeometryCache(const uint8* Memory)
 	Indices = (int32*)(Memory + sizeof(FRecastGeometryCache) + (sizeof(FVector::FReal) * Header.NumVerts * 3));
 }
 
+rcRasterizationFlags FRecastRawGeometryElement::GetRasterizationFlags(const FCompositeNavModifier& InModifier)
+{
+	int32 Flags = 0;
+	if (InModifier.GetFillCollisionUnderneathForNavmesh())
+	{
+		// To prevent navmesh generation under the geometry, set the RC_PROJECT_TO_BOTTOM flag to true.
+		// This rasterize triangles as filled columns down to the HF lower bound.
+		Flags |= RC_PROJECT_TO_BOTTOM;
+	}
+	if (InModifier.GetRasterizeAsFilledConvexVolume())
+	{
+		// To rasterized as filled convex volume, set RC_RASTERIZE_AS_FILLED_CONVEX flag to true.
+		// It will rasterize all the triangles of an ElementData into a 2d grid of 1 span per cell thus making sure that the interior is full
+		Flags |= RC_RASTERIZE_AS_FILLED_CONVEX;
+	}
+
+	return rcRasterizationFlags(Flags);
+}
+
 namespace RecastGeometryExport {
 
 #if SHOW_NAV_EXPORT_PREVIEW
@@ -1640,9 +1659,10 @@ struct FVoxelCacheRasterizeContext
 		if (RasterizeHF == NULL)
 		{
 			const FVector::FReal DummyBounds[3] = { 0 };
+			const bool bAllocateTempSpanColumns = true;	// we can't know here if at some point a NavElementData will need it, so we need to allocate it
 
 			RasterizeHF = rcAllocHeightfield();
-			rcCreateHeightfield(NULL, *RasterizeHF, FieldSize, FieldSize, DummyBounds, DummyBounds, CellSize, CellHeight);
+			rcCreateHeightfield(NULL, *RasterizeHF, FieldSize, FieldSize, DummyBounds, DummyBounds, CellSize, CellHeight, bAllocateTempSpanColumns);
 		}
 	}
 
@@ -1686,9 +1706,11 @@ uint32 GetTileCacheSizeHelper(TArray<FNavMeshTileData>& CompressedTiles)
 
 FRecastTileGenerator::FRecastTileGenerator(FRecastNavMeshGenerator& ParentGenerator, const FIntPoint& Location, const double PendingTileCreationTime)
 	: TimeSliceManager(ParentGenerator.GetTimeSliceManager())
+	, RasterizeGeometryWorldRecastCoordsBounds(ForceInit)
 {
 	bUpdateGeometry = true;
 	bHasLowAreaModifiers = false;
+	bHasGeometryToRasterizeAsFilledConvexVolume = false;
 
 	TileX = Location.X;
 	TileY = Location.Y;
@@ -2182,7 +2204,7 @@ void FRecastTileGenerator::GatherNavigationDataGeometry(const TSharedRef<FNaviga
 			{
 
 				// rasterize
-				PrepareVoxelCache(ElementData.CollisionData, ModifierInstance, SpanData);
+				PrepareVoxelCache(ElementData.CollisionData, Unreal2RecastBox(ElementData.Bounds), ModifierInstance, SpanData);
 				CachedVoxels = SpanData.GetData();
 				NumCachedVoxels = SpanData.Num();
 
@@ -2363,7 +2385,13 @@ void FRecastTileGenerator::InitRasterizationMaskArray(const rcHeightfield* InSol
 	FMemory::Memset(OutRasterizationMasks.GetData(), AllowAllFlags, CellCount*sizeof(TInlineMaskArray::ElementType));
 }
 
+// Deprecated
 void FRecastTileGenerator::PrepareVoxelCache(const TNavStatArray<uint8>& RawCollisionCache, const FCompositeNavModifier& InModifier, TNavStatArray<rcSpanCache>& SpanData)
+{
+	return PrepareVoxelCache(RawCollisionCache, FBox(ForceInit), InModifier, SpanData);
+}
+
+void FRecastTileGenerator::PrepareVoxelCache(const TNavStatArray<uint8>& RawCollisionCache, const FBox& RecastBounds, const FCompositeNavModifier& InModifier, TNavStatArray<rcSpanCache>& SpanData)
 {
 	// tile's geometry: voxel cache (only for synchronous rebuilds)
 	const int32 WalkableClimbVX = TileConfig.walkableClimb;
@@ -2397,14 +2425,14 @@ void FRecastTileGenerator::PrepareVoxelCache(const TNavStatArray<uint8>& RawColl
 		}
 	}
 
-	// To prevent navmesh generation under the triangles, set the RC_PROJECT_TO_BOTTOM flag to true.
-	// This rasterize triangles as filled columns down to the HF lower bound.
-	const rcRasterizationFlags Flags = InModifier.GetFillCollisionUnderneathForNavmesh() ? RC_PROJECT_TO_BOTTOM : rcRasterizationFlags(0);
+	const rcRasterizationFlags Flags = FRecastRawGeometryElement::GetRasterizationFlags(InModifier);
 
+	const FVector::FReal* BoundsMin = RecastBounds.IsValid ? &RecastBounds.Min.X : nullptr;
+	const FVector::FReal* BoundsMax = RecastBounds.IsValid ? &RecastBounds.Max.X : nullptr;
 	TInlineMaskArray::ElementType* MaskArray = RasterizationMasks.Num() > 0 ? RasterizationMasks.GetData() : nullptr;
 	rcRasterizeTriangles(0, CachedCollisions.Verts, CachedCollisions.Header.NumVerts,
 		CachedCollisions.Indices, TriAreas.GetData(), CachedCollisions.Header.NumFaces,
-		*VoxelCacheContext.RasterizeHF, WalkableClimbVX, Flags, MaskArray);
+		*VoxelCacheContext.RasterizeHF, WalkableClimbVX, Flags, MaskArray, BoundsMin, BoundsMax);
 
 	const int32 NumSpans = rcCountSpans(0, *VoxelCacheContext.RasterizeHF);
 	if (NumSpans > 0)
@@ -2467,6 +2495,8 @@ void FRecastTileGenerator::AppendModifier(const FCompositeNavModifier& Modifier,
 		OffmeshLinks.Add(SimpleLinkCollection);
 	}
 
+	bHasGeometryToRasterizeAsFilledConvexVolume = bHasGeometryToRasterizeAsFilledConvexVolume || Modifier.GetRasterizeAsFilledConvexVolume();
+
 	// Navmesh resolutions is a modifier without area, if present, it must not be skipped.
 	if (Modifier.GetAreas().Num() == 0 && Modifier.GetNavMeshResolution() == ENavigationDataResolution::Invalid)
 	{
@@ -2517,10 +2547,7 @@ void FRecastTileGenerator::AppendGeometry(const FNavigationRelevantData& Element
 	}
 	
 	FRecastRawGeometryElement GeometryElement;
-
-	// To prevent navmesh generation under the geometry, set the RC_PROJECT_TO_BOTTOM flag to true.
-	// This rasterize triangles as filled columns down to the HF lower bound.
-	GeometryElement.RasterizationFlags = InModifier.GetFillCollisionUnderneathForNavmesh() ? RC_PROJECT_TO_BOTTOM : rcRasterizationFlags(0);
+	GeometryElement.RasterizationFlags = FRecastRawGeometryElement::GetRasterizationFlags(InModifier);
 
 	FRecastGeometryCache CollisionCache(RawCollisionCache.GetData());
 	
@@ -2547,6 +2574,8 @@ void FRecastTileGenerator::AppendGeometry(const FNavigationRelevantData& Element
 		FMemory::Memcpy(GeometryElement.GeomCoords.GetData(), CollisionCache.Verts, sizeof(FVector::FReal) * NumCoords);
 		FMemory::Memcpy(GeometryElement.GeomIndices.GetData(), CollisionCache.Indices, sizeof(int32) * NumIndices);
 
+		GeometryElement.GeomCoordsBounds = Unreal2RecastBox(ElementData.Bounds); 
+
 		RawGeometry.Add(MoveTemp(GeometryElement));
 	}	
 }
@@ -2726,7 +2755,7 @@ bool FRecastTileGenerator::CreateHeightField(FNavMeshBuildContext& BuildContext)
 			BuildContext.log(RC_LOG_ERROR, "CreateHeightField: Out of memory 'SolidHF'.");
 			return false;
 		}
-		if (!rcCreateHeightfield(&BuildContext, *SolidHF, TileConfig.width, TileConfig.height, TileConfig.bmin, TileConfig.bmax, TileConfig.cs, TileConfig.ch))
+		if (!rcCreateHeightfield(&BuildContext, *SolidHF, TileConfig.width, TileConfig.height, TileConfig.bmin, TileConfig.bmax, TileConfig.cs, TileConfig.ch, bHasGeometryToRasterizeAsFilledConvexVolume))
 		{
 			BuildContext.log(RC_LOG_ERROR, "CreateHeightField: Could not create solid heightfield.");
 			return false;
@@ -2735,7 +2764,13 @@ bool FRecastTileGenerator::CreateHeightField(FNavMeshBuildContext& BuildContext)
 	return true;
 }
 
+// Deprecated
 ETimeSliceWorkResult FRecastTileGenerator::RasterizeGeometryRecastTimeSliced(FNavMeshBuildContext& BuildContext, const TArray<FVector::FReal>& Coords, const TArray<int32>& Indices, const rcRasterizationFlags RasterizationFlags, FTileRasterizationContext& RasterContext)
+{
+	return RasterizeGeometryRecastTimeSliced(BuildContext, Coords, Indices, FBox(ForceInit), RasterizationFlags, RasterContext);
+}
+
+ETimeSliceWorkResult FRecastTileGenerator::RasterizeGeometryRecastTimeSliced(FNavMeshBuildContext& BuildContext, const TArray<FVector::FReal>& Coords, const TArray<int32>& Indices, const FBox& CoordsBounds, const rcRasterizationFlags RasterizationFlags, FTileRasterizationContext& RasterContext)
 {
 	QUICK_SCOPE_CYCLE_COUNTER(STAT_Navigation_RasterizeGeometryRecast);
 
@@ -2771,11 +2806,13 @@ ETimeSliceWorkResult FRecastTileGenerator::RasterizeGeometryRecastTimeSliced(FNa
 
 		QUICK_SCOPE_CYCLE_COUNTER(STAT_Navigation_RasterizeGeomRecastRasterizeTriangles);
 
+		const FVector::FReal* BoundsMin = CoordsBounds.IsValid ? &CoordsBounds.Min.X : nullptr;
+		const FVector::FReal* BoundsMax = CoordsBounds.IsValid ? &CoordsBounds.Max.X : nullptr;
 		const TInlineMaskArray::ElementType* MaskArray = RasterContext.RasterizationMasks.Num() > 0 ? RasterContext.RasterizationMasks.GetData() : nullptr;
 		rcRasterizeTriangles(&BuildContext,
 			Coords.GetData(), NumVerts,
 			Indices.GetData(), RasterizeGeomRecastTriAreas.GetData(), NumFaces,
-			*SolidHF, TileConfig.walkableClimb, RasterizationFlags, MaskArray);
+			*SolidHF, TileConfig.walkableClimb, RasterizationFlags, MaskArray, BoundsMin, BoundsMax);
 
 #if RECAST_INTERNAL_DEBUG_DATA	
 		BuildContext.InternalDebugData.TriangleCount += NumFaces;
@@ -2800,7 +2837,13 @@ ETimeSliceWorkResult FRecastTileGenerator::RasterizeGeometryRecastTimeSliced(FNa
 	return ETimeSliceWorkResult::Succeeded;
 }
 
+// Deprecated
 void FRecastTileGenerator::RasterizeGeometryRecast(FNavMeshBuildContext& BuildContext, const TArray<FVector::FReal>& Coords, const TArray<int32>& Indices, const rcRasterizationFlags RasterizationFlags, FTileRasterizationContext& RasterContext)
+{
+	return RasterizeGeometryRecast(BuildContext, Coords, Indices, FBox(ForceInit), RasterizationFlags, RasterContext);
+}
+
+void FRecastTileGenerator::RasterizeGeometryRecast(FNavMeshBuildContext& BuildContext, const TArray<FVector::FReal>& Coords, const TArray<int32>& Indices, const FBox& CoordsBounds, const rcRasterizationFlags RasterizationFlags, FTileRasterizationContext& RasterContext)
 {
 	QUICK_SCOPE_CYCLE_COUNTER(STAT_Navigation_RasterizeGeometryRecast);
 
@@ -2820,11 +2863,13 @@ void FRecastTileGenerator::RasterizeGeometryRecast(FNavMeshBuildContext& BuildCo
 	{
 		QUICK_SCOPE_CYCLE_COUNTER(STAT_Navigation_RasterizeGeomRecastRasterizeTriangles);
 
+		const FVector::FReal* BoundsMin = CoordsBounds.IsValid ? &CoordsBounds.Min.X : nullptr;
+		const FVector::FReal* BoundsMax = CoordsBounds.IsValid ? &CoordsBounds.Max.X : nullptr;
 		const TInlineMaskArray::ElementType* MaskArray = RasterContext.RasterizationMasks.Num() > 0 ? RasterContext.RasterizationMasks.GetData() : nullptr;
 		rcRasterizeTriangles(&BuildContext,
 			Coords.GetData(), NumVerts,
 			Indices.GetData(), RasterizeGeomRecastTriAreas.GetData(), NumFaces,
-			*SolidHF, TileConfig.walkableClimb, RasterizationFlags, MaskArray);
+			*SolidHF, TileConfig.walkableClimb, RasterizationFlags, MaskArray, BoundsMin, BoundsMax);
 	}
 
 #if RECAST_INTERNAL_DEBUG_DATA
@@ -2847,11 +2892,18 @@ void FRecastTileGenerator::RasterizeGeometryRecast(FNavMeshBuildContext& BuildCo
 	RasterizeGeomRecastTriAreas.Reset();
 }
 
+// Deprecated
 void FRecastTileGenerator::RasterizeGeometryTransformCoordsAndFlipIndices(const TArray<FVector::FReal>& Coords, const TArray<int32>& Indices, const FTransform& LocalToWorld)
+{
+	return RasterizeGeometryTransformCoordsAndFlipIndices(Coords, Indices, FBox(ForceInit), LocalToWorld);
+}
+
+void FRecastTileGenerator::RasterizeGeometryTransformCoordsAndFlipIndices(const TArray<FVector::FReal>& Coords, const TArray<int32>& Indices, const FBox& CoordsBounds, const FTransform& LocalToWorld)
 {
 	QUICK_SCOPE_CYCLE_COUNTER(STAT_Navigation_RasterizeGeometryTransformCoordsAndFlipIndices);
 
 	RasterizeGeometryWorldRecastCoords.SetNumUninitialized(Coords.Num(), EAllowShrinking::No);
+	RasterizeGeometryWorldRecastCoordsBounds = FBox(ForceInit);
 
 	FMatrix LocalToRecastWorld = LocalToWorld.ToMatrixWithScale() * Unreal2RecastMatrix();
 
@@ -2880,13 +2932,20 @@ void FRecastTileGenerator::RasterizeGeometryTransformCoordsAndFlipIndices(const
 			RasterizeGeometryFlippedIndices[i + 2] = Indices[i];
 		}
 	}
+
+	if (CoordsBounds.IsValid)
+	{
+		RasterizeGeometryWorldRecastCoordsBounds = Recast2UnrealBox(CoordsBounds).TransformBy(LocalToRecastWorld);
+	}
 }
 
+// Deprecated
 void FRecastTileGenerator::RasterizeGeometryTransformCoords(const TArray<FVector::FReal>& Coords, const FTransform& LocalToWorld)
 {
 	QUICK_SCOPE_CYCLE_COUNTER(STAT_Navigation_RasterizeGeometryTransformCoords);
 
 	RasterizeGeometryWorldRecastCoords.SetNumUninitialized(Coords.Num(), EAllowShrinking::No);
+	RasterizeGeometryWorldRecastCoordsBounds = FBox(ForceInit);
 
 	FMatrix LocalToRecastWorld = LocalToWorld.ToMatrixWithScale()*Unreal2RecastMatrix();
 
@@ -2902,7 +2961,13 @@ void FRecastTileGenerator::RasterizeGeometryTransformCoords(const TArray<FVector
 	}
 }
 
+// Deprecated
 ETimeSliceWorkResult FRecastTileGenerator::RasterizeGeometryTimeSliced(FNavMeshBuildContext& BuildContext, const TArray<FVector::FReal>& Coords, const TArray<int32>& Indices, const FTransform& LocalToWorld, const rcRasterizationFlags RasterizationFlags, FTileRasterizationContext& RasterContext)
+{
+	return RasterizeGeometryTimeSliced(BuildContext, Coords, Indices, FBox(ForceInit), LocalToWorld, RasterizationFlags, RasterContext);
+}
+
+ETimeSliceWorkResult FRecastTileGenerator::RasterizeGeometryTimeSliced(FNavMeshBuildContext& BuildContext, const TArray<FVector::FReal>& Coords, const TArray<int32>& Indices, const FBox& CoordsBounds, const FTransform& LocalToWorld, const rcRasterizationFlags RasterizationFlags, FTileRasterizationContext& RasterContext)
 {
 	QUICK_SCOPE_CYCLE_COUNTER(STAT_Navigation_RasterizeGeometry);
 
@@ -2914,7 +2979,7 @@ ETimeSliceWorkResult FRecastTileGenerator::RasterizeGeometryTimeSliced(FNavMeshB
 	{
 	case ERasterizeGeomTimeSlicedState::RasterizeGeometryTransformCoordsAndFlipIndices:
 	{
-		RasterizeGeometryTransformCoordsAndFlipIndices(Coords, Indices, LocalToWorld);
+		RasterizeGeometryTransformCoordsAndFlipIndices(Coords, Indices, CoordsBounds, LocalToWorld);
 
 		RasterizeGeomState = ERasterizeGeomTimeSlicedState::RasterizeGeometryRecast;
 
@@ -2928,7 +2993,7 @@ ETimeSliceWorkResult FRecastTileGenerator::RasterizeGeometryTimeSliced(FNavMeshB
 	case ERasterizeGeomTimeSlicedState::RasterizeGeometryRecast:
 	{
 		const TArray<int32>& RasterizeGeometryIndices = bRasterizeGeometryUseFlippedIndices ? RasterizeGeometryFlippedIndices : Indices;
-		WorkResult = RasterizeGeometryRecastTimeSliced(BuildContext, RasterizeGeometryWorldRecastCoords, RasterizeGeometryIndices, RasterizationFlags, RasterContext);
+		WorkResult = RasterizeGeometryRecastTimeSliced(BuildContext, RasterizeGeometryWorldRecastCoords, RasterizeGeometryIndices, RasterizeGeometryWorldRecastCoordsBounds, RasterizationFlags, RasterContext);
 
 		if (WorkResult != ETimeSliceWorkResult::CallAgainNextTimeSlice)
 		{
@@ -2947,12 +3012,17 @@ ETimeSliceWorkResult FRecastTileGenerator::RasterizeGeometryTimeSliced(FNavMeshB
 }
 
 void FRecastTileGenerator::RasterizeGeometry(FNavMeshBuildContext& BuildContext, const TArray<FVector::FReal>& Coords, const TArray<int32>& Indices, const FTransform& LocalToWorld, const rcRasterizationFlags RasterizationFlags, FTileRasterizationContext& RasterContext)
+{
+	return RasterizeGeometry(BuildContext, Coords, Indices, FBox(ForceInit), LocalToWorld, RasterizationFlags, RasterContext);
+}
+
+void FRecastTileGenerator::RasterizeGeometry(FNavMeshBuildContext& BuildContext, const TArray<FVector::FReal>& Coords, const TArray<int32>& Indices, const FBox& CoordsBounds, const FTransform& LocalToWorld, const rcRasterizationFlags RasterizationFlags, FTileRasterizationContext& RasterContext)
 {
 	QUICK_SCOPE_CYCLE_COUNTER(STAT_Navigation_RasterizeGeometry);
 
-	RasterizeGeometryTransformCoordsAndFlipIndices(Coords, Indices, LocalToWorld);
+	RasterizeGeometryTransformCoordsAndFlipIndices(Coords, Indices, CoordsBounds, LocalToWorld);
 	const TArray<int32>& RasterizeGeometryIndices = bRasterizeGeometryUseFlippedIndices ? RasterizeGeometryFlippedIndices : Indices;
-	RasterizeGeometryRecast(BuildContext, RasterizeGeometryWorldRecastCoords, RasterizeGeometryIndices, RasterizationFlags, RasterContext);
+	RasterizeGeometryRecast(BuildContext, RasterizeGeometryWorldRecastCoords, RasterizeGeometryIndices, RasterizeGeometryWorldRecastCoordsBounds, RasterizationFlags, RasterContext);
 }
 
 ETimeSliceWorkResult FRecastTileGenerator::RasterizeTrianglesTimeSliced(FNavMeshBuildContext& BuildContext, FTileRasterizationContext& RasterContext)
@@ -2970,7 +3040,7 @@ ETimeSliceWorkResult FRecastTileGenerator::RasterizeTrianglesTimeSliced(FNavMesh
 			while (RasterizeTrianglesTimeSlicedInstTransformIdx < Element.PerInstanceTransform.Num())
 			{
 				const FTransform& InstanceTransform = Element.PerInstanceTransform[RasterizeTrianglesTimeSlicedInstTransformIdx];
-				const ETimeSliceWorkResult WorkResult = RasterizeGeometryTimeSliced(BuildContext, Element.GeomCoords, Element.GeomIndices, InstanceTransform, Element.RasterizationFlags, RasterContext);
+				const ETimeSliceWorkResult WorkResult = RasterizeGeometryTimeSliced(BuildContext, Element.GeomCoords, Element.GeomIndices, Element.GeomCoordsBounds, InstanceTransform, Element.RasterizationFlags, RasterContext);
 			
 				//the original code just kept calling the RasterizeGeometry() functions and had no return type, 
 				//so we will process the next layer (if we are not needing to process this layer again next time slice) 
@@ -2991,7 +3061,7 @@ ETimeSliceWorkResult FRecastTileGenerator::RasterizeTrianglesTimeSliced(FNavMesh
 		}
 		else
 		{
-			const ETimeSliceWorkResult WorkResult = RasterizeGeometryRecastTimeSliced(BuildContext, Element.GeomCoords, Element.GeomIndices, Element.RasterizationFlags, RasterContext);
+			const ETimeSliceWorkResult WorkResult = RasterizeGeometryRecastTimeSliced(BuildContext, Element.GeomCoords, Element.GeomIndices, Element.GeomCoordsBounds, Element.RasterizationFlags, RasterContext);
 	
 			if (TimeSliceManager->GetTimeSlicer().IsTimeSliceFinishedCached())
 			{
@@ -3021,12 +3091,12 @@ void FRecastTileGenerator::RasterizeTriangles(FNavMeshBuildContext& BuildContext
 		{
 			for (const FTransform& InstanceTransform : Element.PerInstanceTransform)
 			{
-				RasterizeGeometry(BuildContext, Element.GeomCoords, Element.GeomIndices, InstanceTransform, Element.RasterizationFlags, RasterContext);
+				RasterizeGeometry(BuildContext, Element.GeomCoords, Element.GeomIndices, Element.GeomCoordsBounds, InstanceTransform, Element.RasterizationFlags, RasterContext);
 			}
 		}
 		else
 		{
-			RasterizeGeometryRecast(BuildContext, Element.GeomCoords, Element.GeomIndices, Element.RasterizationFlags, RasterContext);
+			RasterizeGeometryRecast(BuildContext, Element.GeomCoords, Element.GeomIndices, Element.GeomCoordsBounds, Element.RasterizationFlags, RasterContext);
 		}
 	}
 }
diff --git a/Engine/Source/Runtime/NavigationSystem/Public/NavMesh/RecastNavMeshGenerator.h b/Engine/Source/Runtime/NavigationSystem/Public/NavMesh/RecastNavMeshGenerator.h
--- a/Engine/Source/Runtime/NavigationSystem/Public/NavMesh/RecastNavMeshGenerator.h
+++ b/Engine/Source/Runtime/NavigationSystem/Public/NavMesh/RecastNavMeshGenerator.h
@@ -188,12 +188,17 @@
 	// Instance geometry
 	TArray<FVector::FReal>		GeomCoords;
 	TArray<int32>				GeomIndices;
+
+	// Bounds of the geometry
+	FBox GeomCoordsBounds;
 	
 	// Per instance transformations in unreal coords
 	// When empty, geometry is in world space
 	TArray<FTransform>	PerInstanceTransform;
 
 	rcRasterizationFlags RasterizationFlags;
+
+	NAVIGATIONSYSTEM_API static rcRasterizationFlags GetRasterizationFlags(const FCompositeNavModifier& InModifier);
 };
 
 struct FRecastAreaNavModifierElement
@@ -444,28 +449,15 @@
 	/** Gather geometry from a specified Navigation Data */
 	NAVIGATIONSYSTEM_API void GatherNavigationDataGeometry(const TSharedRef<FNavigationRelevantData, ESPMode::ThreadSafe>& ElementData, UNavigationSystemV1& NavSys, const FNavDataConfig& OwnerNavDataConfig, bool bGeometryChanged);
 
-	UE_DEPRECATED(5.5, "Use the new version without RasterContext instead.")
-	NAVIGATIONSYSTEM_API bool CreateHeightField(FNavMeshBuildContext& BuildContext, FTileRasterizationContext& RasterContext);
-	UE_DEPRECATED(5.5, "Use the new version without RasterContext instead.")
-	NAVIGATIONSYSTEM_API void GenerateRecastFilter(FNavMeshBuildContext& BuildContext, FTileRasterizationContext& RasterContext);
-	UE_DEPRECATED(5.5, "Use the new version without RasterContext instead.")
-	NAVIGATIONSYSTEM_API ETimeSliceWorkResult GenerateRecastFilterTimeSliced(FNavMeshBuildContext& BuildContext, FTileRasterizationContext& RasterContext);
-	UE_DEPRECATED(5.5, "Use the new version without RasterContext instead.")
-	NAVIGATIONSYSTEM_API bool BuildCompactHeightField(FNavMeshBuildContext& BuildContext, FTileRasterizationContext& RasterContext);
-	UE_DEPRECATED(5.5, "Use the new version without RasterContext instead.")
-	NAVIGATIONSYSTEM_API bool RecastErodeWalkable(FNavMeshBuildContext& BuildContext, FTileRasterizationContext& RasterContext);
-	
 	/** Start functions used by GenerateCompressedLayersTimeSliced / GenerateCompressedLayers */
 	NAVIGATIONSYSTEM_API bool CreateHeightField(FNavMeshBuildContext& BuildContext);
 	NAVIGATIONSYSTEM_API ETimeSliceWorkResult RasterizeTrianglesTimeSliced(FNavMeshBuildContext& BuildContext, FTileRasterizationContext& RasterContext);
 	NAVIGATIONSYSTEM_API void RasterizeTriangles(FNavMeshBuildContext& BuildContext, FTileRasterizationContext& RasterContext);
-	NAVIGATIONSYSTEM_API ETimeSliceWorkResult RasterizeGeometryRecastTimeSliced(FNavMeshBuildContext& BuildContext, const TArray<FVector::FReal>& Coords, const TArray<int32>& Indices, const rcRasterizationFlags RasterizationFlags, FTileRasterizationContext& RasterContext);
-	NAVIGATIONSYSTEM_API void RasterizeGeometryRecast(FNavMeshBuildContext& BuildContext, const TArray<FVector::FReal>& Coords, const TArray<int32>& Indices, const rcRasterizationFlags RasterizationFlags, FTileRasterizationContext& RasterContext);
-	NAVIGATIONSYSTEM_API void RasterizeGeometryTransformCoordsAndFlipIndices(const TArray<FVector::FReal>& Coords, const TArray<int32>& Indices, const FTransform& LocalToWorld);
-	UE_DEPRECATED(5.5, "This function was not handling the indices order correctly and is replaced by RasterizeGeometryTransformCoordsAndFlipIndices.")
-	NAVIGATIONSYSTEM_API void RasterizeGeometryTransformCoords(const TArray<FVector::FReal>& Coords, const FTransform& LocalToWorld);
-	NAVIGATIONSYSTEM_API ETimeSliceWorkResult RasterizeGeometryTimeSliced(FNavMeshBuildContext& BuildContext, const TArray<FVector::FReal>& Coords, const TArray<int32>& Indices, const FTransform& LocalToWorld, const rcRasterizationFlags RasterizationFlags, FTileRasterizationContext& RasterContext);
-	NAVIGATIONSYSTEM_API void RasterizeGeometry(FNavMeshBuildContext& BuildContext, const TArray<FVector::FReal>& Coords, const TArray<int32>& Indices, const FTransform& LocalToWorld, const rcRasterizationFlags RasterizationFlags, FTileRasterizationContext& RasterContext);
+	NAVIGATIONSYSTEM_API ETimeSliceWorkResult RasterizeGeometryRecastTimeSliced(FNavMeshBuildContext& BuildContext, const TArray<FVector::FReal>& Coords, const TArray<int32>& Indices, const FBox& CoordsBounds, const rcRasterizationFlags RasterizationFlags, FTileRasterizationContext& RasterContext);
+	NAVIGATIONSYSTEM_API void RasterizeGeometryRecast(FNavMeshBuildContext& BuildContext, const TArray<FVector::FReal>& Coords, const TArray<int32>& Indices, const FBox& CoordsBounds, const rcRasterizationFlags RasterizationFlags, FTileRasterizationContext& RasterContext);
+	NAVIGATIONSYSTEM_API void RasterizeGeometryTransformCoordsAndFlipIndices(const TArray<FVector::FReal>& Coords, const TArray<int32>& Indices, const FBox& CoordsBounds, const FTransform& LocalToWorld);
+	NAVIGATIONSYSTEM_API ETimeSliceWorkResult RasterizeGeometryTimeSliced(FNavMeshBuildContext& BuildContext, const TArray<FVector::FReal>& Coords, const TArray<int32>& Indices, const FBox& CoordsBounds, const FTransform& LocalToWorld, const rcRasterizationFlags RasterizationFlags, FTileRasterizationContext& RasterContext);
+	NAVIGATIONSYSTEM_API void RasterizeGeometry(FNavMeshBuildContext& BuildContext, const TArray<FVector::FReal>& Coords, const TArray<int32>& Indices, const FBox& CoordsBounds, const FTransform& LocalToWorld, const rcRasterizationFlags RasterizationFlags, FTileRasterizationContext& RasterContext);
 	NAVIGATIONSYSTEM_API void GenerateRecastFilter(FNavMeshBuildContext& BuildContext);
 	NAVIGATIONSYSTEM_API ETimeSliceWorkResult GenerateRecastFilterTimeSliced(FNavMeshBuildContext& BuildContext);
 	NAVIGATIONSYSTEM_API bool BuildCompactHeightField(FNavMeshBuildContext& BuildContext);
@@ -477,24 +469,6 @@
 	/** Builds CompressedLayers array (geometry + modifiers) time sliced*/
 	NAVIGATIONSYSTEM_API virtual ETimeSliceWorkResult GenerateCompressedLayersTimeSliced(FNavMeshBuildContext& BuildContext);
 
-	
-	/** builds CompressedLayers array (geometry + modifiers) */
-	UE_DEPRECATED(5.5, "Use the overload with dtLinkBuilderData instead.")
-	NAVIGATIONSYSTEM_API virtual bool GenerateCompressedLayers(FNavMeshBuildContext& BuildContext);
-
-	/** Builds a navigation data layer */
-	UE_DEPRECATED(5.5, "Use the overload with dtLinkBuilderData instead.")
-	NAVIGATIONSYSTEM_API bool GenerateNavigationDataLayer(FNavMeshBuildContext& BuildContext, FTileCacheCompressor& TileCompressor, FTileCacheAllocator& GenNavAllocator, FTileGenerationContext& GenerationContext, int32 LayerIdx);
-
-	/** builds NavigationData array (layers + obstacles) time sliced */
-	UE_DEPRECATED(5.5, "Use the overload with dtLinkBuilderData instead.")
-	NAVIGATIONSYSTEM_API ETimeSliceWorkResult GenerateNavigationDataTimeSliced(FNavMeshBuildContext& BuildContext);
-
-	/** builds NavigationData array (layers + obstacles) */
-	UE_DEPRECATED(5.5, "Use the overload with dtLinkBuilderData instead.")
-	NAVIGATIONSYSTEM_API bool GenerateNavigationData(FNavMeshBuildContext& BuildContext);
-	
-	
 	/** Builds CompressedLayers array (geometry + modifiers) */
 	NAVIGATIONSYSTEM_API virtual bool GenerateCompressedLayers(FNavMeshBuildContext& BuildContext, const dtLinkBuilderData& InLinkBuilderData);
 
@@ -527,13 +501,11 @@
 	NAVIGATIONSYSTEM_API void AppendModifier(const FCompositeNavModifier& Modifier, const FNavDataPerInstanceTransformDelegate& InTransformsDelegate);
 	/** Appends specified geometry to tile's geometry */
 	NAVIGATIONSYSTEM_API void ValidateAndAppendGeometry(const FNavigationRelevantData& ElementData, const FCompositeNavModifier& InModifier);
-	UE_DEPRECATED(5.5, "Use the version taking a const reference on FNavigationRelevantData.")
-	NAVIGATIONSYSTEM_API void ValidateAndAppendGeometry(const TSharedRef<FNavigationRelevantData, ESPMode::ThreadSafe>& ElementData, const FCompositeNavModifier& InModifier);
 
 	NAVIGATIONSYSTEM_API void AppendGeometry(const FNavigationRelevantData& ElementData, const FCompositeNavModifier& InModifier, const FNavDataPerInstanceTransformDelegate& InTransformsDelegate);
 	
 	/** prepare voxel cache from collision data */
-	NAVIGATIONSYSTEM_API void PrepareVoxelCache(const TNavStatArray<uint8>& RawCollisionCache, const FCompositeNavModifier& InModifier, TNavStatArray<rcSpanCache>& SpanData);
+	NAVIGATIONSYSTEM_API void PrepareVoxelCache(const TNavStatArray<uint8>& RawCollisionCache, const FBox& RecastBounds, const FCompositeNavModifier& InModifier, TNavStatArray<rcSpanCache>& SpanData);
 	NAVIGATIONSYSTEM_API bool HasVoxelCache(const TNavStatArray<uint8>& RawVoxelCache, rcSpanCache*& CachedVoxels, int32& NumCachedVoxels) const;
 	NAVIGATIONSYSTEM_API void AddVoxelCache(TNavStatArray<uint8>& RawVoxelCache, const rcSpanCache* CachedVoxels, const int32 NumCachedVoxels) const;
 
@@ -545,11 +517,48 @@
 	NAVIGATIONSYSTEM_API bool IsTileDebugAllowingGeneration() const;
 #endif
 
+	/** Deprecated functions */
+	UE_DEPRECATED(5.5, "Use the new version without RasterContext instead.")
+	NAVIGATIONSYSTEM_API bool CreateHeightField(FNavMeshBuildContext& BuildContext, FTileRasterizationContext& RasterContext);
+	UE_DEPRECATED(5.5, "Use the new version without RasterContext instead.")
+	NAVIGATIONSYSTEM_API void GenerateRecastFilter(FNavMeshBuildContext& BuildContext, FTileRasterizationContext& RasterContext);
+	UE_DEPRECATED(5.5, "Use the new version without RasterContext instead.")
+	NAVIGATIONSYSTEM_API ETimeSliceWorkResult GenerateRecastFilterTimeSliced(FNavMeshBuildContext& BuildContext, FTileRasterizationContext& RasterContext);
+	UE_DEPRECATED(5.5, "Use the new version without RasterContext instead.")
+	NAVIGATIONSYSTEM_API bool BuildCompactHeightField(FNavMeshBuildContext& BuildContext, FTileRasterizationContext& RasterContext);
+	UE_DEPRECATED(5.5, "Use the new version without RasterContext instead.")
+	NAVIGATIONSYSTEM_API bool RecastErodeWalkable(FNavMeshBuildContext& BuildContext, FTileRasterizationContext& RasterContext);
+	UE_DEPRECATED(5.5, "This function was not handling the indices order correctly and is replaced by RasterizeGeometryTransformCoordsAndFlipIndices.")
+	NAVIGATIONSYSTEM_API void RasterizeGeometryTransformCoords(const TArray<FVector::FReal>& Coords, const FTransform& LocalToWorld);
+	UE_DEPRECATED(5.5, "Use the overload with dtLinkBuilderData instead.")
+	NAVIGATIONSYSTEM_API virtual bool GenerateCompressedLayers(FNavMeshBuildContext& BuildContext);
+	UE_DEPRECATED(5.5, "Use the overload with dtLinkBuilderData instead.")
+	NAVIGATIONSYSTEM_API bool GenerateNavigationDataLayer(FNavMeshBuildContext& BuildContext, FTileCacheCompressor& TileCompressor, FTileCacheAllocator& GenNavAllocator, FTileGenerationContext& GenerationContext, int32 LayerIdx);
+	UE_DEPRECATED(5.5, "Use the overload with dtLinkBuilderData instead.")
+	NAVIGATIONSYSTEM_API ETimeSliceWorkResult GenerateNavigationDataTimeSliced(FNavMeshBuildContext& BuildContext);
+	UE_DEPRECATED(5.5, "Use the overload with dtLinkBuilderData instead.")
+	NAVIGATIONSYSTEM_API bool GenerateNavigationData(FNavMeshBuildContext& BuildContext);
+	UE_DEPRECATED(5.5, "Use the version taking a const reference on FNavigationRelevantData.")
+	NAVIGATIONSYSTEM_API void ValidateAndAppendGeometry(const TSharedRef<FNavigationRelevantData, ESPMode::ThreadSafe>& ElementData, const FCompositeNavModifier& InModifier);
+	UE_DEPRECATED(5.6, "Please use the overload with the added CoordBounds.")
+	NAVIGATIONSYSTEM_API ETimeSliceWorkResult RasterizeGeometryRecastTimeSliced(FNavMeshBuildContext& BuildContext, const TArray<FVector::FReal>& Coords, const TArray<int32>& Indices, const rcRasterizationFlags RasterizationFlags, FTileRasterizationContext& RasterContext);
+	UE_DEPRECATED(5.6, "Please use the overload with the added CoordBounds.")
+	NAVIGATIONSYSTEM_API void RasterizeGeometryRecast(FNavMeshBuildContext& BuildContext, const TArray<FVector::FReal>& Coords, const TArray<int32>& Indices, const rcRasterizationFlags RasterizationFlags, FTileRasterizationContext& RasterContext);
+	UE_DEPRECATED(5.6, "Please use the overload with the added CoordBounds.")
+	NAVIGATIONSYSTEM_API void RasterizeGeometryTransformCoordsAndFlipIndices(const TArray<FVector::FReal>& Coords, const TArray<int32>& Indices, const FTransform& LocalToWorld);
+	UE_DEPRECATED(5.6, "Please use the overload with the added CoordBounds.")
+	NAVIGATIONSYSTEM_API ETimeSliceWorkResult RasterizeGeometryTimeSliced(FNavMeshBuildContext& BuildContext, const TArray<FVector::FReal>& Coords, const TArray<int32>& Indices, const FTransform& LocalToWorld, const rcRasterizationFlags RasterizationFlags, FTileRasterizationContext& RasterContext);
+	UE_DEPRECATED(5.6, "Please use the overload with the added CoordBounds.")
+	NAVIGATIONSYSTEM_API void RasterizeGeometry(FNavMeshBuildContext& BuildContext, const TArray<FVector::FReal>& Coords, const TArray<int32>& Indices, const FTransform& LocalToWorld, const rcRasterizationFlags RasterizationFlags, FTileRasterizationContext& RasterContext);
+	UE_DEPRECATED(5.6, "Please use the overload with the added RecastBounds.")
+	NAVIGATIONSYSTEM_API void PrepareVoxelCache(const TNavStatArray<uint8>& RawCollisionCache, const FCompositeNavModifier& InModifier, TNavStatArray<rcSpanCache>& SpanData);
+
 protected:
 	uint32 bRegenerateCompressedLayers : 1;
 	uint32 bFullyEncapsulatedByInclusionBounds : 1;
 	uint32 bUpdateGeometry : 1;
 	uint32 bHasLowAreaModifiers : 1;
+	uint32 bHasGeometryToRasterizeAsFilledConvexVolume : 1; 
 
 	/** Start time slicing variables */
 	ERasterizeGeomRecastTimeSlicedState RasterizeGeomRecastState;
@@ -606,6 +615,7 @@
 	/** Result of calling RasterizeGeometry() */
 	TArray<FVector::FReal> RasterizeGeometryWorldRecastCoords;
 	TArray<int32> RasterizeGeometryFlippedIndices;
+	FBox RasterizeGeometryWorldRecastCoordsBounds;
 	uint8 bRasterizeGeometryUseFlippedIndices : 1;
 
 	// tile's geometry: without voxel cache
diff --git a/Engine/Source/Runtime/Navmesh/Private/Recast/Recast.cpp b/Engine/Source/Runtime/Navmesh/Private/Recast/Recast.cpp
--- a/Engine/Source/Runtime/Navmesh/Private/Recast/Recast.cpp
+++ b/Engine/Source/Runtime/Navmesh/Private/Recast/Recast.cpp
@@ -99,7 +99,8 @@ void rcFreeHeightField(rcHeightfield* hf)
 	rcFree(hf->EdgeHits);
 	rcFree(hf->RowExt);
 	rcFree(hf->tempspans);
-#endif
+	rcFree(hf->tempSpanColumns);
+#endif // EPIC_ADDITION_USE_NEW_RECAST_RASTERIZER
 	rcFree(hf);
 }
 
@@ -241,7 +242,7 @@ void rcCalcGridSize(const rcReal* bmin, const rcReal* bmax, rcReal cs, int* w, i
 /// @see rcAllocHeightfield, rcHeightfield 
 bool rcCreateHeightfield(rcContext* /*ctx*/, rcHeightfield& hf, int width, int height,
 						 const rcReal* bmin, const rcReal* bmax,
-						 rcReal cs, rcReal ch)
+						 rcReal cs, rcReal ch, bool bAllocateTempSpanColumns)
 {
 	// TODO: VC complains about unref formal variable, figure out a way to handle this better.
 //	rcAssert(ctx);
@@ -285,7 +286,14 @@ bool rcCreateHeightfield(rcContext* /*ctx*/, rcHeightfield& hf, int width, int h
 		}
 	}
 
-#endif
+	if (bAllocateTempSpanColumns)
+	{
+		hf.tempSpanColumns = (rcSpanData*)rcAlloc(sizeof(rcSpanData)*hf.width*hf.height, RC_ALLOC_PERM); 
+		if (!hf.tempSpanColumns)
+			return false;
+		memset(hf.tempSpanColumns, 0, sizeof(rcSpanData)*hf.width*hf.height);
+	}
+#endif // EPIC_ADDITION_USE_NEW_RECAST_RASTERIZER
 
 	return true;
 }
diff --git a/Engine/Source/Runtime/Navmesh/Private/Recast/RecastRasterization.cpp b/Engine/Source/Runtime/Navmesh/Private/Recast/RecastRasterization.cpp
--- a/Engine/Source/Runtime/Navmesh/Private/Recast/RecastRasterization.cpp
+++ b/Engine/Source/Runtime/Navmesh/Private/Recast/RecastRasterization.cpp
@@ -93,6 +93,32 @@ static void freeSpan(rcHeightfield& hf, rcSpan* ptr)
 	hf.freelist = ptr;
 }
 
+// @UE BEGIN
+// Merge overlapping spans
+static void mergeSpanData(rcSpanData& span, const rcSpanUInt smin, const rcSpanUInt smax, 
+					const unsigned char area, const int flagMergeThr)
+{
+	// For spans whose tops are really close to each other, prefer walkable areas.
+	// This is done in order to remove aliasing (similar to z-fighting) on surfaces close to each other.
+	if (rcAbs((int)span.smax - (int)smax) <= flagMergeThr)
+	{
+		span.area = rcMax(span.area, area);
+	}
+	else
+	{
+		// Use the new spans area if it will become the top.
+		if (smax > span.smax)
+			span.area = area;
+	}
+
+	// Merge height intervals.
+	if (smin < span.smin)
+		span.smin = smin;
+	if (smax > span.smax)
+		span.smax = smax;
+}
+// @UE END
+
 static void addSpan(rcHeightfield& hf, const int x, const int y,
 					const rcSpanUInt smin, const rcSpanUInt smax,
 					const unsigned char area, const int flagMergeThr)
@@ -132,26 +158,7 @@ static void addSpan(rcHeightfield& hf, const int x, const int y,
 		else
 		{
 			// @UE BEGIN
-			// Merge overlapping spans.
-
-			// For spans whose tops are really close to each other, prefer walkable areas.
-			// This is done in order to remove aliasing (similar to z-fighting) on surfaces close to each other.
-			if (rcAbs((int)s->data.smax - (int)cur->data.smax) <= flagMergeThr)
-			{
-				s->data.area = rcMax(s->data.area, cur->data.area);
-			}
-			else
-			{
-				// Use the new spans area if it will become the top.
-				if (cur->data.smax > s->data.smax)
-					s->data.area = cur->data.area;
-			}
-
-			// Merge height intervals.
-			if (cur->data.smin < s->data.smin)
-				s->data.smin = cur->data.smin;
-			if (cur->data.smax > s->data.smax)
-				s->data.smax = cur->data.smax;
+			mergeSpanData(s->data, cur->data.smin, cur->data.smax, cur->data.area, flagMergeThr);
 			// @UE END
 			
 			// Remove current span.
@@ -444,13 +451,15 @@ static void rasterizeTriTest(const rcReal* v0, const rcReal* v1, const rcReal* v
 }
 #endif //TEST_NEW_RASTERIZER
 
+template <typename AddSpanFunc>
 static void rasterizeTri(const rcReal* v0, const rcReal* v1, const rcReal* v2,
 						 const unsigned char area, rcHeightfield& hf,
 						 const rcReal* bmin, const rcReal* bmax,
 						 const rcReal cs, const rcReal ics, const rcReal ich, 
 						 const int flagMergeThr,
 						 const int rasterizationFlags, /*UE*/
-	                     const int* rasterizationMasks /*UE*/)
+	                     const int* rasterizationMasks, /*UE*/
+	                     AddSpanFunc& Func /*UE*/)
 {
 	rcEdgeHit* const hfEdgeHits = hf.EdgeHits; //this prevents a static analysis warning
 
@@ -502,7 +511,7 @@ static void rasterizeTri(const rcReal* v0, const rcReal* v1, const rcReal* v2,
 			triangle_ismin = 0; //UE
 		}
 
-		addSpan(hf, x0, y0, triangle_ismin, triangle_ismax, area, flagMergeThr);
+		Func.addSpan(hf, x0, y0, triangle_ismin, triangle_ismax, area, flagMergeThr);
 		return;
 	}
 
@@ -676,7 +685,7 @@ static void rasterizeTri(const rcReal* v0, const rcReal* v1, const rcReal* v2,
 				int xloop1 = intMin(hf.RowExt[y + 1].MaxCol, x1);
 				for (int x = xloop0; x <= xloop1; x++)
 				{
-					addSpan(hf, x, y, triangle_ismin_clamp, triangle_ismax_clamp, area, flagMergeThr);
+					Func.addSpan(hf, x, y, triangle_ismin_clamp, triangle_ismax_clamp, area, flagMergeThr);
 				}
 
 				// reset for next triangle
@@ -701,7 +710,7 @@ static void rasterizeTri(const rcReal* v0, const rcReal* v1, const rcReal* v2,
 					{
 						triangle_ismin_clamp = 0; //UE
 					}
-					addSpan(hf, x, y, triangle_ismin_clamp, triangle_ismax_clamp, area, flagMergeThr);
+					Func.addSpan(hf, x, y, triangle_ismin_clamp, triangle_ismax_clamp, area, flagMergeThr);
 				}
 
 				// reset for next triangle
@@ -914,7 +923,7 @@ static void rasterizeTri(const rcReal* v0, const rcReal* v1, const rcReal* v2,
 
 				}
 	#endif
-				addSpan(hf, x, y, smin, smax, area, flagMergeThr);
+				Func.addSpan(hf, x, y, smin, smax, area, flagMergeThr);
 			}
 
 			// reset for next triangle
@@ -1029,12 +1038,121 @@ static void rasterizeTri(const rcReal* v0, const rcReal* v1, const rcReal* v2,
 				ismin = 0; //UE
 			}
 
-			addSpan(hf, x, y, ismin, ismax, area, flagMergeThr);
+			Func.addSpan(hf, x, y, ismin, ismax, area, flagMergeThr);
 		}
 	}
 }
 #endif //EPIC_ADDITION_USE_NEW_RECAST_RASTERIZER
 
+struct AddSpanInHeightfield
+{
+	static void addSpan(rcHeightfield& hf, const int x, const int y,
+					const rcSpanUInt smin, const rcSpanUInt smax,
+					const unsigned char area, const int flagMergeThr)
+	{
+		::addSpan(hf, x, y, smin, smax, area, flagMergeThr);
+	}
+};
+
+struct AddSpanInTempColumns
+{
+	AddSpanInTempColumns(rcContext* ctx, rcHeightfield& hf, const float ics, const rcReal* vertsBMin, const rcReal* vertsBMax)
+	{
+#if EPIC_ADDITION_USE_NEW_RECAST_RASTERIZER
+		if (vertsBMin == nullptr || vertsBMax == nullptr)
+		{
+			if (ctx)
+			{
+				ctx->log(RC_LOG_WARNING, "rcRasterizeTriangles called for triangles to rasterize as filled convex but the verts bounds is not valid.");
+			}
+			return;
+		}
+
+		if (!hf.tempSpanColumns)
+		{
+			if (ctx)
+			{
+				ctx->log(RC_LOG_WARNING, "rcRasterizeTriangles called for triangles to rasterize as filled convex but tempSpanColumns was not allocated.");
+			}
+			return;
+		}
+
+		// init data
+		const rcReal xMin = rcClamp(vertsBMin[0], hf.bmin[0], hf.bmax[0]);
+		const rcReal xMax = rcClamp(vertsBMax[0], hf.bmin[0], hf.bmax[0]);
+		const rcReal zMin = rcClamp(vertsBMin[2], hf.bmin[2], hf.bmax[2]);
+		const rcReal zMax = rcClamp(vertsBMax[2], hf.bmin[2], hf.bmax[2]);
+		if (xMax == xMin || zMax == zMin)
+		{
+			// this can be valid if we receive a geometry that is outside of the height field. No need to log
+			return;
+		}
+
+		// compute the part of the heightfield that will be covered by verts
+		ibminx = ics * (xMin - hf.bmin[0]);
+		ibminz = ics * (zMin - hf.bmin[2]);
+		ibwidth = rcClamp(ics * (xMax - xMin), 0, hf.width);
+		ibheight = rcClamp(ics * (zMax - zMin), 0, hf.height);
+		memset(hf.tempSpanColumns, 0, sizeof(rcSpanData)*ibwidth*ibheight);	// we could keep an array of call idx to avoid doing that too often
+#endif // EPIC_ADDITION_USE_NEW_RECAST_RASTERIZER
+	}
+
+	void addSpan(rcHeightfield& hf, const int x, const int z,
+					const rcSpanUInt smin, const rcSpanUInt smax,
+					const unsigned char area, const int flagMergeThr)
+	{
+#if EPIC_ADDITION_USE_NEW_RECAST_RASTERIZER
+		// convert the heightfield coord to the tempSpanColumns coord
+		const int tempColumnx = x - ibminx;
+		const int tempColumnz = z - ibminz;
+		rcAssert(tempColumnx >= 0 && tempColumnx < ibwidth);
+		rcAssert(tempColumnz >= 0 && tempColumnz < ibheight);
+		const int tempColumnIdx = tempColumnx + (tempColumnz) * ibwidth;
+		rcSpanData& tempColumn = hf.tempSpanColumns[tempColumnIdx];
+		if (tempColumn.smin == 0 && tempColumn.smax == 0)
+		{
+			// first span added, store it as is
+			tempColumn.smin = smin;
+			tempColumn.smax = smax;
+			tempColumn.area = area;
+		}
+		else
+		{
+			// need to merge with a previously added span
+			mergeSpanData(tempColumn, smin, smax, area, flagMergeThr);
+		}
+#endif // EPIC_ADDITION_USE_NEW_RECAST_RASTERIZER
+	}
+
+	void TransferColumnsToHeightfield(rcHeightfield& hf, const float flagMergeThr)
+	{
+#if EPIC_ADDITION_USE_NEW_RECAST_RASTERIZER
+		// push all the created spans into the heightfield
+		for (int iz = 0; iz < ibheight; iz++)
+		{
+			for (int ix = 0; ix < ibwidth; ix++)
+			{
+				const int tempSpanColumnsIdx = ix + (iz * ibwidth);
+				const rcSpanData& tempSpanColumn = hf.tempSpanColumns[tempSpanColumnsIdx];
+				if (tempSpanColumn.smin == 0 && tempSpanColumn.smax == 0)
+				{
+					continue;
+				}
+
+				::addSpan(hf, ibminx + ix, ibminz + iz, tempSpanColumn.smin, tempSpanColumn.smax, tempSpanColumn.area, flagMergeThr);
+			}
+		}
+#endif // EPIC_ADDITION_USE_NEW_RECAST_RASTERIZER
+	}
+
+	bool IsValid() const { return ibwidth > 0 && ibheight > 0; }
+
+	int ibminx = 0;
+	int ibminz = 0;
+	int ibwidth = 0;
+	int ibheight = 0;
+};
+
 /// @par
 ///
 /// No spans will be added if the triangle does not overlap the heightfield grid.
@@ -1048,9 +1166,11 @@ void rcRasterizeTriangle(rcContext* ctx, const rcReal* v0, const rcReal* v1, con
 
 	ctx->startTimer(RC_TIMER_RASTERIZE_TRIANGLES);
 
+	AddSpanInHeightfield Func;
+
 	const rcReal ics = 1.0f/solid.cs;
 	const rcReal ich = 1.0f/solid.ch;
-	rasterizeTri(v0, v1, v2, area, solid, solid.bmin, solid.bmax, solid.cs, ics, ich, flagMergeThr, rasterizationFlags, rasterizationMasks); //UE
+	rasterizeTri(v0, v1, v2, area, solid, solid.bmin, solid.bmax, solid.cs, ics, ich, flagMergeThr, rasterizationFlags, rasterizationMasks, Func); //UE
 
 	ctx->stopTimer(RC_TIMER_RASTERIZE_TRIANGLES);
 }
@@ -1062,23 +1182,46 @@ void rcRasterizeTriangle(rcContext* ctx, const rcReal* v0, const rcReal* v1, con
 /// @see rcHeightfield
 void rcRasterizeTriangles(rcContext* ctx, const rcReal* verts, const int /*nv*/,
 						  const int* tris, const unsigned char* areas, const int nt,
-						  rcHeightfield& solid, const int flagMergeThr, const int rasterizationFlags, const int* rasterizationMasks) //UE
+						  rcHeightfield& solid, const int flagMergeThr, const int rasterizationFlags, const int* rasterizationMasks, //UE
+						  const rcReal* vertsbmin, const rcReal* vertsbmax) //UE
 {
 	if (ctx)
 		ctx->startTimer(RC_TIMER_RASTERIZE_TRIANGLES);
 	
 	const rcReal ics = 1.0f/solid.cs;
 	const rcReal ich = 1.0f/solid.ch;
-	// Rasterize triangles.
-	for (int i = 0; i < nt; ++i)
+
+	const bool bAsFilledVerticalConvexVolume = rasterizationFlags & RC_RASTERIZE_AS_FILLED_CONVEX;
+	if (!bAsFilledVerticalConvexVolume)
 	{
-		const rcReal* v0 = &verts[tris[i*3+0]*3];
-		const rcReal* v1 = &verts[tris[i*3+1]*3];
-		const rcReal* v2 = &verts[tris[i*3+2]*3];
-		// Rasterize.
-		rasterizeTri(v0, v1, v2, areas[i], solid, solid.bmin, solid.bmax, solid.cs, ics, ich, flagMergeThr, rasterizationFlags, rasterizationMasks); //UE
+		AddSpanInHeightfield Func;
+		// Rasterize triangles.
+		for (int i = 0; i < nt; ++i)
+		{
+			const rcReal* v0 = &verts[tris[i*3+0]*3];
+			const rcReal* v1 = &verts[tris[i*3+1]*3];
+			const rcReal* v2 = &verts[tris[i*3+2]*3];
+			// Rasterize.
+			rasterizeTri(v0, v1, v2, areas[i], solid, solid.bmin, solid.bmax, solid.cs, ics, ich, flagMergeThr, rasterizationFlags, rasterizationMasks, Func); //UE
+		}
 	}
-	
+	else
+	{
+		AddSpanInTempColumns Func(ctx, solid, ics, vertsbmin, vertsbmax);
+		if (Func.IsValid())
+		{
+			for (int i = 0; i < nt; ++i)
+			{
+				const rcReal* v0 = &verts[tris[i*3+0]*3];
+				const rcReal* v1 = &verts[tris[i*3+1]*3];
+				const rcReal* v2 = &verts[tris[i*3+2]*3];
+				rasterizeTri(v0, v1, v2, areas[i], solid, solid.bmin, solid.bmax, solid.cs, ics, ich, flagMergeThr, rasterizationFlags, rasterizationMasks, Func); //UE
+			}
+
+			Func.TransferColumnsToHeightfield(solid, flagMergeThr);
+		}
+	}
+
 	if (ctx)
 		ctx->stopTimer(RC_TIMER_RASTERIZE_TRIANGLES);
 }
@@ -1090,21 +1233,44 @@ void rcRasterizeTriangles(rcContext* ctx, const rcReal* verts, const int /*nv*/,
 /// @see rcHeightfield
 void rcRasterizeTriangles(rcContext* ctx, const rcReal* verts, const int /*nv*/,
 						  const unsigned short* tris, const unsigned char* areas, const int nt,
-						  rcHeightfield& solid, const int flagMergeThr, const int rasterizationFlags, const int* rasterizationMasks) //UE
+						  rcHeightfield& solid, const int flagMergeThr, const int rasterizationFlags, const int* rasterizationMasks, //UE
+						  const rcReal* vertsbmin, const rcReal* vertsbmax) //UE
 {
 	if (ctx)
 		ctx->startTimer(RC_TIMER_RASTERIZE_TRIANGLES);
 	
 	const rcReal ics = 1.0f/solid.cs;
 	const rcReal ich = 1.0f/solid.ch;
-	// Rasterize triangles.
-	for (int i = 0; i < nt; ++i)
+
+	const bool bAsFilledVerticalConvexVolume = rasterizationFlags & RC_RASTERIZE_AS_FILLED_CONVEX;
+	if (!bAsFilledVerticalConvexVolume)
 	{
-		const rcReal* v0 = &verts[tris[i*3+0]*3];
-		const rcReal* v1 = &verts[tris[i*3+1]*3];
-		const rcReal* v2 = &verts[tris[i*3+2]*3];
-		// Rasterize.
-		rasterizeTri(v0, v1, v2, areas[i], solid, solid.bmin, solid.bmax, solid.cs, ics, ich, flagMergeThr, rasterizationFlags, rasterizationMasks); //UE
+		AddSpanInHeightfield Func;
+		// Rasterize triangles.
+		for (int i = 0; i < nt; ++i)
+		{
+			const rcReal* v0 = &verts[tris[i*3+0]*3];
+			const rcReal* v1 = &verts[tris[i*3+1]*3];
+			const rcReal* v2 = &verts[tris[i*3+2]*3];
+			// Rasterize.
+			rasterizeTri(v0, v1, v2, areas[i], solid, solid.bmin, solid.bmax, solid.cs, ics, ich, flagMergeThr, rasterizationFlags, rasterizationMasks, Func); //UE
+		}
+	}
+	else
+	{
+		AddSpanInTempColumns Func(ctx, solid, ics, vertsbmin, vertsbmax);
+		if (Func.IsValid())
+		{
+			for (int i = 0; i < nt; ++i)
+			{
+				const rcReal* v0 = &verts[tris[i*3+0]*3];
+				const rcReal* v1 = &verts[tris[i*3+1]*3];
+				const rcReal* v2 = &verts[tris[i*3+2]*3];
+				rasterizeTri(v0, v1, v2, areas[i], solid, solid.bmin, solid.bmax, solid.cs, ics, ich, flagMergeThr, rasterizationFlags, rasterizationMasks, Func); //UE
+			}
+
+			Func.TransferColumnsToHeightfield(solid, flagMergeThr);
+		}
 	}
 	
 	if (ctx)
@@ -1117,21 +1283,44 @@ void rcRasterizeTriangles(rcContext* ctx, const rcReal* verts, const int /*nv*/,
 ///
 /// @see rcHeightfield
 void rcRasterizeTriangles(rcContext* ctx, const rcReal* verts, const unsigned char* areas, const int nt,
-						  rcHeightfield& solid, const int flagMergeThr, const int rasterizationFlags, const int* rasterizationMasks) //UE
+						  rcHeightfield& solid, const int flagMergeThr, const int rasterizationFlags, const int* rasterizationMasks, //UE
+						  const rcReal* vertsbmin, const rcReal* vertsbmax) //UE
 {
 	if (ctx)
 		ctx->startTimer(RC_TIMER_RASTERIZE_TRIANGLES);
 	
 	const rcReal ics = 1.0f/solid.cs;
 	const rcReal ich = 1.0f/solid.ch;
-	// Rasterize triangles.
-	for (int i = 0; i < nt; ++i)
+
+	const bool bAsFilledVerticalConvexVolume = rasterizationFlags & RC_RASTERIZE_AS_FILLED_CONVEX;
+	if (!bAsFilledVerticalConvexVolume)
+	{
+		AddSpanInHeightfield Func;
+		// Rasterize triangles.
+		for (int i = 0; i < nt; ++i)
+		{
+			const rcReal* v0 = &verts[(i*3+0)*3];
+			const rcReal* v1 = &verts[(i*3+1)*3];
+			const rcReal* v2 = &verts[(i*3+2)*3];
+			// Rasterize.
+			rasterizeTri(v0, v1, v2, areas[i], solid, solid.bmin, solid.bmax, solid.cs, ics, ich, flagMergeThr, rasterizationFlags, rasterizationMasks, Func); //UE
+		}
+	}
+	else
 	{
-		const rcReal* v0 = &verts[(i*3+0)*3];
-		const rcReal* v1 = &verts[(i*3+1)*3];
-		const rcReal* v2 = &verts[(i*3+2)*3];
-		// Rasterize.
-		rasterizeTri(v0, v1, v2, areas[i], solid, solid.bmin, solid.bmax, solid.cs, ics, ich, flagMergeThr, rasterizationFlags, rasterizationMasks); //UE
+		AddSpanInTempColumns Func(ctx, solid, ics, vertsbmin, vertsbmax);
+		if (Func.IsValid())
+		{
+			for (int i = 0; i < nt; ++i)
+			{
+				const rcReal* v0 = &verts[(i*3+0)*3];
+				const rcReal* v1 = &verts[(i*3+1)*3];
+				const rcReal* v2 = &verts[(i*3+2)*3];
+				rasterizeTri(v0, v1, v2, areas[i], solid, solid.bmin, solid.bmax, solid.cs, ics, ich, flagMergeThr, rasterizationFlags, rasterizationMasks, Func); //UE
+			}
+
+			Func.TransferColumnsToHeightfield(solid, flagMergeThr);
+		}
 	}
 	
 	if (ctx)
diff --git a/Engine/Source/Runtime/Navmesh/Public/Recast/Recast.h b/Engine/Source/Runtime/Navmesh/Public/Recast/Recast.h
--- a/Engine/Source/Runtime/Navmesh/Public/Recast/Recast.h
+++ b/Engine/Source/Runtime/Navmesh/Public/Recast/Recast.h
@@ -400,7 +400,7 @@ struct rcTempSpan
 {
 	int sminmax[2];			///< The lower and upper limit of the span. [Limit: < #smax]
 };
-#endif
+#endif // EPIC_ADDITION_USE_NEW_RECAST_RASTERIZER
 
 /// A dynamic heightfield representing obstructed space.
 /// @ingroup recast
@@ -420,7 +420,9 @@ struct rcHeightfield
 	rcEdgeHit* EdgeHits; ///< h + 1 bit flags that indicate what edges cross the z cell boundaries
 	rcRowExt* RowExt;		///< h structs that give the current x range for this z row
 	rcTempSpan* tempspans;		///< Heightfield of temp spans (width*height).
-#endif
+	rcSpanData* tempSpanColumns;	///< Heightfield of 1 span per cell for vertical column rasterization //UE
+#endif // EPIC_ADDITION_USE_NEW_RECAST_RASTERIZER
+
 };
 
 /// Provides information on the content of a cell column in a compact heightfield. 
@@ -701,6 +703,7 @@ enum rcFilterLowAreaFlags
 enum rcRasterizationFlags
 {
 	RC_PROJECT_TO_BOTTOM = 1 << 0,		///< Will create spans from the triangle surface to the bottom of the heightfield
+	RC_RASTERIZE_AS_FILLED_CONVEX = 1 << 1,		///< Will rasterize all the triangles of a list into a single span for each (x,z) and then add all those spans in the heightfield //UE
 };
 
 // UE
@@ -949,9 +952,11 @@ NAVMESH_API void rcCalcGridSize(const rcReal* bmin, const rcReal* bmax, rcReal c
 ///  @param[in]		bmax	The maximum bounds of the field's AABB. [(x, y, z)] [Units: wu]
 ///  @param[in]		cs		The xz-plane cell size to use for the field. [Limit: > 0] [Units: wu]
 ///  @param[in]		ch		The y-axis cell size to use for field. [Limit: > 0] [Units: wu]
+///  @param[in]		bAllocateTempSpanColumns	Indicate if we need to allocate the heightfield's tempSpanColumns.
 NAVMESH_API bool rcCreateHeightfield(rcContext* ctx, rcHeightfield& hf, int width, int height,
 						 const rcReal* bmin, const rcReal* bmax,
-						 rcReal cs, rcReal ch);
+						 rcReal cs, rcReal ch,
+						 bool bAllocateTempSpanColumns = false); //UE
 
 /// Resets all spans of heightfield.
 ///  @ingroup recast
@@ -1054,10 +1059,13 @@ NAVMESH_API void rcRasterizeTriangle(rcContext* ctx, const rcReal* v0, const rcR
 ///  								[Limit: >= 0] [Units: vx]
 ///  @param[in]     rtzFlags		Flags to change the rasterization behavior			//UE
 ///  @param[in]     rtzMasks		Mask for the rasterization flags [Size: hf.w*hf.h]	//UE
+///  @param[in]     vertsbmin		Min location of the bounding box of verts [(x, y, z)] //UE
+///  @param[in]     vertsbmax		Max location of the bounding box of verts [(x, y, z)] //UE
 NAVMESH_API void rcRasterizeTriangles(rcContext* ctx, const rcReal* verts, const int nv,
 						  const int* tris, const unsigned char* areas, const int nt,
 						  rcHeightfield& solid, const int flagMergeThr = 1, 
-						  const int rasterizationFlags = 0, const int* rasterizationMasks = nullptr); //UE
+						  const int rasterizationFlags = 0, const int* rasterizationMasks = nullptr, //UE
+						  const rcReal* vertsbmin = nullptr, const rcReal* vertsbmax = nullptr); //UE
 
 /// Rasterizes an indexed triangle mesh into the specified heightfield.
 ///  @ingroup recast
@@ -1072,10 +1080,13 @@ NAVMESH_API void rcRasterizeTriangles(rcContext* ctx, const rcReal* verts, const
 ///  								[Limit: >= 0] [Units: vx]
 ///  @param[in]     rtzFlags		Flags to change the rasterization behavior			//UE
 ///  @param[in]     rtzMasks		Mask for the rasterization flags [Size: hf.w*hf.h]	//UE
+///  @param[in]     vertsbmin		Min location of the bounding box of verts [(x, y, z)] //UE
+///  @param[in]     vertsbmax		Max location of the bounding box of verts [(x, y, z)] //UE
 NAVMESH_API void rcRasterizeTriangles(rcContext* ctx, const rcReal* verts, const int nv,
 						  const unsigned short* tris, const unsigned char* areas, const int nt,
 						  rcHeightfield& solid, const int flagMergeThr = 1, 
-						  const int rasterizationFlags = 0, const int* rasterizationMasks = nullptr); //UE
+						  const int rasterizationFlags = 0, const int* rasterizationMasks = nullptr, //UE
+						  const rcReal* vertsbmin = nullptr, const rcReal* vertsbmax = nullptr); //UE
 
 /// Rasterizes triangles into the specified heightfield.
 ///  @ingroup recast
@@ -1088,9 +1099,12 @@ NAVMESH_API void rcRasterizeTriangles(rcContext* ctx, const rcReal* verts, const
 ///  								[Limit: >= 0] [Units: vx]
 ///  @param[in]     rtzFlags		Flags to change the rasterization behavior			//UE
 ///  @param[in]     rtzMasks		Mask for the rasterization flags [Size: hf.w*hf.h]	//UE
+///  @param[in]     vertsbmin		Min location of the bounding box of verts [(x, y, z)] //UE
+///  @param[in]     vertsbmax		Max location of the bounding box of verts [(x, y, z)] //UE
 NAVMESH_API void rcRasterizeTriangles(rcContext* ctx, const rcReal* verts, const unsigned char* areas, const int nt,
 						  rcHeightfield& solid, const int flagMergeThr = 1,
-						  const int rasterizationFlags = 0, const int* rasterizationMasks = nullptr); //UE
+						  const int rasterizationFlags = 0, const int* rasterizationMasks = nullptr, //UE
+						  const rcReal* vertsbmin = nullptr, const rcReal* vertsbmax = nullptr); //UE
 
 /// Marks non-walkable spans as walkable if their maximum is within @p walkableClimp of a walkable neighbor. 
 ///  @ingroup recast
