- Adding the ability to specify different nav areas for upwards vs downwards traversal of generated jump nav links
- Adding in a vertical offset to try to alleviate debug rendering of 'stacked' navlinks going in opposite directions

diff --git a/Engine/Source/Runtime/NavigationSystem/Private/NavMesh/LinkGenerationConfig.cpp b/Engine/Source/Runtime/NavigationSystem/Private/NavMesh/LinkGenerationConfig.cpp
--- a/Engine/Source/Runtime/NavigationSystem/Private/NavMesh/LinkGenerationConfig.cpp
+++ b/Engine/Source/Runtime/NavigationSystem/Private/NavMesh/LinkGenerationConfig.cpp
@@ -8,9 +8,35 @@
 #include "Detour/DetourNavLinkBuilderConfig.h"
 #endif //WITH_RECAST
 
+// Refer to header for more information
+PRAGMA_DISABLE_DEPRECATION_WARNINGS
 FNavLinkGenerationJumpDownConfig::FNavLinkGenerationJumpDownConfig()
 {
-	AreaClass = UNavArea_Default::StaticClass();
+	DownDirectionAreaClass = UNavArea_Default::StaticClass();
+	UpDirectionAreaClass = UNavArea_Default::StaticClass();
+}
+PRAGMA_ENABLE_DEPRECATION_WARNINGS
+
+bool FNavLinkGenerationJumpDownConfig::Serialize(FArchive& Ar)
+{
+	UScriptStruct* const Struct = StaticStruct();
+	Struct->SerializeTaggedProperties(Ar, (uint8*)this, Struct, nullptr);
+
+#if WITH_EDITORONLY_DATA
+	if (Ar.IsLoading())
+	{
+		PRAGMA_DISABLE_DEPRECATION_WARNINGS
+		if (AreaClass_DEPRECATED)
+		{
+			DownDirectionAreaClass = AreaClass_DEPRECATED;
+			UpDirectionAreaClass = AreaClass_DEPRECATED;
+			AreaClass_DEPRECATED = nullptr;
+		}
+		PRAGMA_ENABLE_DEPRECATION_WARNINGS
+	}
+#endif //WITH_EDITORONLY_DATA
+
+	return true;
 }
 
 #if WITH_RECAST
diff --git a/Engine/Source/Runtime/NavigationSystem/Private/NavMesh/NavMeshRenderingComponent.cpp b/Engine/Source/Runtime/NavigationSystem/Private/NavMesh/NavMeshRenderingComponent.cpp
--- a/Engine/Source/Runtime/NavigationSystem/Private/NavMesh/NavMeshRenderingComponent.cpp
+++ b/Engine/Source/Runtime/NavigationSystem/Private/NavMesh/NavMeshRenderingComponent.cpp
@@ -46,6 +46,7 @@ static constexpr float NavMeshEdges_LineThickness = 4.f;
 static constexpr float LinkLines_LineThickness = 2.0f;
 static constexpr float GeneratedLinkLines_LineThickness = 4.0f;
 static constexpr float ClusterLinkLines_LineThickness = 2.0f;
+static constexpr float LinkLines_UnidirectionalArcVerticalSeparationDistance = 3.0f;
 
 namespace FNavMeshRenderingHelpers
 {
@@ -125,7 +126,7 @@ namespace FNavMeshRenderingHelpers
 		return Pt;
 	}
 
-	void CacheArc(TArray<FDebugRenderSceneProxy::FDebugLine>& DebugLines, const FVector& Start, const FVector& End, const FVector::FReal Height, const uint32 Segments, const FLinearColor& Color, float LineThickness = 0)
+	void CacheArc(TArray<FDebugRenderSceneProxy::FDebugLine>& DebugLines, const FVector& Start, const FVector& End, const FVector::FReal Height, const uint32 Segments, const FLinearColor& Color, float LineThickness = 0, float VerticalOffset = 0)
 	{
 		if (Segments == 0)
 		{
@@ -135,12 +136,14 @@ namespace FNavMeshRenderingHelpers
 		const FVector::FReal ArcPtsScale = 1. / (FVector::FReal)Segments;
 		const FVector Dir = End - Start;
 		const FVector::FReal Length = Dir.Size();
+		const FVector::FReal VerticalOffsetSign = FMath::Sign(Dir.Z);
 
 		FVector Prev = Start;
 		for (uint32 i = 1; i <= Segments; ++i)
 		{
 			const FVector::FReal u = (FVector::FReal)i * ArcPtsScale;
-			const FVector Pt = EvalArc(Start, Dir, Length*Height, u);
+			FVector Pt = EvalArc(Start, Dir, Length*Height, u);
+			Pt.Z += (i < Segments) ? VerticalOffsetSign * VerticalOffset : 0.0;
 
 			DebugLines.Add(FDebugRenderSceneProxy::FDebugLine(Prev, Pt, Color.ToFColor(true), LineThickness));
 			Prev = Pt;
@@ -160,8 +163,9 @@ namespace FNavMeshRenderingHelpers
 	void CacheLink(TArray<FDebugRenderSceneProxy::FDebugLine>& DebugLines, const FVector V0, const FVector V1, const FColor LinkColor, const uint8 LinkDirection, const bool bIsGenerated)
 	{
 		const float LineThickness = bIsGenerated ? GeneratedLinkLines_LineThickness : LinkLines_LineThickness;
+		const float ArcVerticalOffset = LinkDirection ? 0.0f : 0.5f * (LinkLines_UnidirectionalArcVerticalSeparationDistance + LineThickness);
 
-		FNavMeshRenderingHelpers::CacheArc(DebugLines, V0, V1, 0.4f, 4, LinkColor, LineThickness);
+		FNavMeshRenderingHelpers::CacheArc(DebugLines, V0, V1, 0.4f, 4, LinkColor, LineThickness, ArcVerticalOffset);
 
 		const FVector VOffset(0, 0, FVector::Dist(V0, V1) * 1.333f);
 		FNavMeshRenderingHelpers::CacheArrowHead(DebugLines, V1, V0 + VOffset, 30.f, LinkColor, LineThickness);
diff --git a/Engine/Source/Runtime/NavigationSystem/Private/NavMesh/RecastNavMeshGenerator.cpp b/Engine/Source/Runtime/NavigationSystem/Private/NavMesh/RecastNavMeshGenerator.cpp
--- a/Engine/Source/Runtime/NavigationSystem/Private/NavMesh/RecastNavMeshGenerator.cpp
+++ b/Engine/Source/Runtime/NavigationSystem/Private/NavMesh/RecastNavMeshGenerator.cpp
@@ -3933,17 +3933,51 @@ dtStatus FRecastTileGenerator::BuildTileCacheLinks(FNavMeshBuildContext& BuildCo
 #endif // RECAST_INTERNAL_DEBUG_DATA		
 	}
 
-	auto AddLinkLambda = [&OutGeneratedLinks](const dtNavLinkBuilderJumpDownConfig& config, const dtReal* posA, const dtReal* posB)
+	auto AddLinkSharedLambda = [&OutGeneratedLinks](const dtNavLinkBuilderJumpDownConfig& config) -> FGeneratedNavigationLink&
 	{
 		FGeneratedNavigationLink& NewLink = OutGeneratedLinks.Emplace_GetRef();
 		NewLink.bIsGenerated = true;
 		NewLink.NavLinkId = FNavLinkId(config.linkUserId);
-		NewLink.generatedLinkArea = config.area;
-		NewLink.generatedLinkPolyFlag = config.polyFlag;
+		return NewLink;
+	};
+
+	// Used when we're provided with identical up and down areas. In this situation, a single bidirectional link is sufficient for our needs
+	auto AddBidirectionalLinkLambda = [&AddLinkSharedLambda](const dtNavLinkBuilderJumpDownConfig& config, const dtReal* posA, const dtReal* posB)
+	{
+		FGeneratedNavigationLink& NewLink = AddLinkSharedLambda(config);
+		NewLink.generatedLinkArea = config.downDirArea;
+		NewLink.generatedLinkPolyFlag = config.downDirPolyFlag;
 		NewLink.Left = Recast2UnrealPoint(posA);
-		NewLink.Right = Recast2UnrealPoint(posB);	
+		NewLink.Right = Recast2UnrealPoint(posB);
+	};
+
+	// Used when we're provided with different up and down areas. In this situation, we need to create two different one-way link representations in order to meet our needs
+	auto AddDirectionalLinksLambda = [&AddLinkSharedLambda](const dtNavLinkBuilderJumpDownConfig& config, const dtReal* posA, const dtReal* posB)
+	{	
+		const FVector UnrealPosA = Recast2UnrealPoint(posA);
+		const FVector UnrealPosB = Recast2UnrealPoint(posB);
+
+		const bool bABDown = UnrealPosA.Z > UnrealPosB.Z;
+
+		FGeneratedNavigationLink& DownLink = AddLinkSharedLambda(config);
+		DownLink.generatedLinkArea = config.downDirArea;
+		DownLink.generatedLinkPolyFlag = config.downDirPolyFlag;
+		DownLink.Left = bABDown ? UnrealPosA : UnrealPosB;
+		DownLink.Right = bABDown ? UnrealPosB : UnrealPosA;
+		DownLink.Direction = ENavLinkDirection::LeftToRight;
+
+		// Note that in order for this link to be recognized by the navigation system, it needs to have unique left/right points, which is why the left and right points are flipped and we're marking this link as left-to-right
+		FGeneratedNavigationLink& UpLink = AddLinkSharedLambda(config);
+		UpLink.generatedLinkArea = config.upDirArea;
+		UpLink.generatedLinkPolyFlag = config.upDirPolyFlag;
+		UpLink.Left = DownLink.Right;
+		UpLink.Right = DownLink.Left;
+		UpLink.Direction = ENavLinkDirection::LeftToRight;
 	};
 
+	const bool bBidirectionalJumpDownLinks = linkBuilderConfig.jumpDownConfig.downDirArea == linkBuilderConfig.jumpDownConfig.upDirArea;
+	TFunctionRef<void(const dtNavLinkBuilderJumpDownConfig&, const dtReal*, const dtReal*)> AddLinkLambda = bBidirectionalJumpDownLinks ? TFunctionRef<void(const dtNavLinkBuilderJumpDownConfig&, const dtReal*, const dtReal*)>(AddBidirectionalLinkLambda) : TFunctionRef<void(const dtNavLinkBuilderJumpDownConfig&, const dtReal*, const dtReal*)>(AddDirectionalLinksLambda);
+
 	// Make FGeneratedNavigationLinks
 	for (const dtNavLinkBuilder::JumpLink& link : linkBuilder.m_links)
 	{
@@ -5438,20 +5472,26 @@ void FRecastNavMeshGenerator::ResolveGeneratedLinkAreas(FRecastBuildConfig& OutC
 {
 	if (IsGeneratingLinks())
 	{
-		const FNavLinkGenerationJumpDownConfig& JumpDown = DestNavMesh->NavLinkJumpDownConfig;
-		if (JumpDown.AreaClass)
+		auto GetAreaIDAndPolyFlags = [this](const TSubclassOf<UNavAreaBase>& AreaClass, unsigned char& OutAreaID, ARecastNavMesh::FNavPolyFlags& OutPolyFlags)
 		{
-			const int32* AreaIDPtr = AdditionalCachedData.AreaClassToIdMap.Find(JumpDown.AreaClass);
-			if (AreaIDPtr != nullptr)
-			{
-				OutConfig.JumpDownConfig.area = IntCastChecked<unsigned char>(*AreaIDPtr);
-				OutConfig.JumpDownConfig.polyFlag = AdditionalCachedData.FlagsPerArea[*AreaIDPtr];
-			}
-			else
+			if (AreaClass)
 			{
-				UE_LOG(LogNavigation, Warning, TEXT("FRecastTileGenerator: Trying to use undefined area class while resolving generated links areas. (%s)"), *GetNameSafe(JumpDown.AreaClass));
+				const int32* AreaIDPtr = AdditionalCachedData.AreaClassToIdMap.Find(AreaClass);
+				if (AreaIDPtr != nullptr)
+				{
+					OutAreaID = IntCastChecked<unsigned char>(*AreaIDPtr);
+					OutPolyFlags = AdditionalCachedData.FlagsPerArea[*AreaIDPtr];
+				}
+				else
+				{
+					UE_LOG(LogNavigation, Warning, TEXT("FRecastTileGenerator: Trying to use undefined area class while resolving generated links areas. (%s)"), *GetNameSafe(AreaClass));
+				}
 			}
-		}
+		};
+
+		const FNavLinkGenerationJumpDownConfig& JumpDown = DestNavMesh->NavLinkJumpDownConfig;
+		GetAreaIDAndPolyFlags(JumpDown.DownDirectionAreaClass, OutConfig.JumpDownConfig.downDirArea, OutConfig.JumpDownConfig.downDirPolyFlag);
+		GetAreaIDAndPolyFlags(JumpDown.UpDirectionAreaClass, OutConfig.JumpDownConfig.upDirArea, OutConfig.JumpDownConfig.upDirPolyFlag);
 	}
 }
 
diff --git a/Engine/Source/Runtime/NavigationSystem/Public/NavMesh/LinkGenerationConfig.h b/Engine/Source/Runtime/NavigationSystem/Public/NavMesh/LinkGenerationConfig.h
--- a/Engine/Source/Runtime/NavigationSystem/Public/NavMesh/LinkGenerationConfig.h
+++ b/Engine/Source/Runtime/NavigationSystem/Public/NavMesh/LinkGenerationConfig.h
@@ -1,4 +1,4 @@
-﻿// Copyright Epic Games, Inc. All Rights Reserved.
+// Copyright Epic Games, Inc. All Rights Reserved.
 #pragma once
 
 #include "Templates/SubclassOf.h"
@@ -30,6 +30,15 @@ struct FNavLinkGenerationJumpDownConfig
 
 	FNavLinkGenerationJumpDownConfig();
 
+	// Note: We need to explicitly disable warnings on these constructors/operators for clang to be happy with deprecated variables
+	PRAGMA_DISABLE_DEPRECATION_WARNINGS
+	~FNavLinkGenerationJumpDownConfig() = default;
+	FNavLinkGenerationJumpDownConfig(const FNavLinkGenerationJumpDownConfig&) = default;
+	FNavLinkGenerationJumpDownConfig(FNavLinkGenerationJumpDownConfig&&) = default;
+	FNavLinkGenerationJumpDownConfig& operator=(const FNavLinkGenerationJumpDownConfig&) = default;
+	FNavLinkGenerationJumpDownConfig& operator=(FNavLinkGenerationJumpDownConfig&&) = default;
+	PRAGMA_ENABLE_DEPRECATION_WARNINGS
+
 	/** Should this config be used to generate links. */
 	UPROPERTY(Config)		// @todo Hidden for now. Since there is currently only one config, it has no real usage.
 	bool bEnabled = true;
@@ -69,9 +78,26 @@ struct FNavLinkGenerationJumpDownConfig
 	UPROPERTY(EditAnywhere, Config, Category = Settings, meta = (Bitmask, BitmaskEnum = "/Script/NavigationSystem.ENavLinkBuilderFlags"))
 	uint16 LinkBuilderFlags = (uint16)ENavLinkBuilderFlags::CreateCenterPointLink;
 
+#if WITH_EDITORONLY_DATA
 	/** Area class for links generated by this configuration. */
+	UE_DEPRECATED(5.6, "AreaClass is deprecated")
+	UPROPERTY(Config, meta = (DeprecatedProperty, DeprecationMessage = "Use DownDirectionAreaClass and UpDirectionAreaClass instead."))
+	TSubclassOf<UNavAreaBase> AreaClass_DEPRECATED;
+#endif //WITH_EDITORONLY_DATA
+
+	/**
+	 * Area class for downward traversal of links generated by this configuration.
+	 * @note If the value matches UpDirectionAreaClass, a single bidirectional link will be used to represent links generated by this configuration.
+	 */
+	UPROPERTY(EditAnywhere, Config, Category = Settings)
+	TSubclassOf<UNavAreaBase> DownDirectionAreaClass;
+
+	/**
+	 * Area class for upward traversal of links generated by this configuration.
+	 * @note If the value matches DownDirectionAreaClass, a single bidirectional link will be used to represent links generated by this configuration.
+	 */
 	UPROPERTY(EditAnywhere, Config, Category = Settings)
-	TSubclassOf<UNavAreaBase> AreaClass;
+	TSubclassOf<UNavAreaBase> UpDirectionAreaClass;
 
 	/** Class used to handle links made with this configuration.
 	 * Using this allows to implement custom behaviors when using navlinks, for example during the pathfollow.
@@ -99,9 +125,20 @@ struct FNavLinkGenerationJumpDownConfig
 	UPROPERTY(Transient)
 	bool bLinkProxyRegistered = false;
 
+	/** Implemented for deprecated property cleanup purposes. */
+	bool Serialize(FArchive& Ar);
+
 #if WITH_RECAST	
 	/** Copy configuration to dtNavLinkBuilderJumpDownConfig. */
 	void CopyToDetourConfig(dtNavLinkBuilderJumpDownConfig& OutDetourConfig) const;
 #endif //WITH_RECAST
 };
 
+template<> struct TStructOpsTypeTraits<FNavLinkGenerationJumpDownConfig> : public TStructOpsTypeTraitsBase2<FNavLinkGenerationJumpDownConfig>
+{
+	enum
+	{
+		WithSerializer = true
+	};
+};
+
diff --git a/Engine/Source/Runtime/Navmesh/Public/Detour/DetourNavLinkBuilderConfig.h b/Engine/Source/Runtime/Navmesh/Public/Detour/DetourNavLinkBuilderConfig.h
--- a/Engine/Source/Runtime/Navmesh/Public/Detour/DetourNavLinkBuilderConfig.h
+++ b/Engine/Source/Runtime/Navmesh/Public/Detour/DetourNavLinkBuilderConfig.h
@@ -1,4 +1,4 @@
-﻿// Copyright Epic Games, Inc. All Rights Reserved.
+// Copyright Epic Games, Inc. All Rights Reserved.
 #pragma once
 
 #include "CoreTypes.h"
@@ -50,11 +50,17 @@ struct dtNavLinkBuilderJumpDownConfig
 	/// User id used to handle links made from this configuration.
 	unsigned long long linkUserId = 0;
 
-	/// User defined flags assigned to the off-mesh connections
-	unsigned short polyFlag = 0;
+	/// User defined flags assigned to downwards traversal of the off-mesh connections
+	unsigned short downDirPolyFlag = 0;
 
-	/// User defined area id assigned to the off-mesh connections
-	unsigned char area = 0;
+	/// User defined flags assigned to upwards traversal of the off-mesh connections
+	unsigned short upDirPolyFlag = 0;
+
+	/// User defined area id assigned to downards traversal of the off-mesh connections
+	unsigned char downDirArea = 0;
+
+	/// User defined area id assigned to upwards traversal of the off-mesh connections
+	unsigned char upDirArea = 0;
 
 	/// Initialize the configuration by computing cached values.
 	NAVMESH_API void init();
