[NavMesh] Fix "Added an oversized dirty area" warnings occuring even if using the disabled threshold value when using invokers.

diff --git a/Engine/Source/Runtime/NavigationSystem/Private/NavMesh/RecastNavMeshGenerator.cpp b/Engine/Source/Runtime/NavigationSystem/Private/NavMesh/RecastNavMeshGenerator.cpp
--- a/Engine/Source/Runtime/NavigationSystem/Private/NavMesh/RecastNavMeshGenerator.cpp
+++ b/Engine/Source/Runtime/NavigationSystem/Private/NavMesh/RecastNavMeshGenerator.cpp
@@ -6621,7 +6621,7 @@ void FRecastNavMeshGenerator::MarkDirtyTiles(const TArray<FNavigationDirtyArea>&
 			
 						// If not using active tile generation, those are reported earlier in FNavigationDirtyAreasController::AddArea
 						if (NavSys && NavSys->GetOperationMode() == FNavigationSystemRunMode::GameMode && NavSys->IsActiveTilesGenerationEnabled() &&
-							AdjustedAreaBounds.GetSize().GetMax() > NavSys->GetDirtyAreaWarningSizeThreshold())
+							NavSys->GetDirtyAreaWarningSizeThreshold() > 0.f && AdjustedAreaBounds.GetSize().GetMax() > NavSys->GetDirtyAreaWarningSizeThreshold())
 						{
 							const FVector2D AdjustedAreaBoundsSize(AdjustedAreaBounds.GetSize());
 
