Add Trace scopes to functions bound to FWorldDelegates::OnWorldPostActorTick
#rnx

diff --git a/Engine/Source/Runtime/Engine/Private/LevelTick.cpp b/Engine/Source/Runtime/Engine/Private/LevelTick.cpp
--- a/Engine/Source/Runtime/Engine/Private/LevelTick.cpp
+++ b/Engine/Source/Runtime/Engine/Private/LevelTick.cpp
@@ -1634,7 +1634,10 @@ void UWorld::Tick( ELevelTick TickType, float DeltaSeconds )
 		TRACE_CPUPROFILER_EVENT_SCOPE(bDoingActorTicks);
 		SCOPE_CYCLE_COUNTER(STAT_TickTime);
 
-		FWorldDelegates::OnWorldPostActorTick.Broadcast(this, TickType, DeltaSeconds);
+		{
+			TRACE_CPUPROFILER_EVENT_SCOPE(FWorldDelegates::OnWorldPostActorTick)
+			FWorldDelegates::OnWorldPostActorTick.Broadcast(this, TickType, DeltaSeconds);
+		}
 
 		// All tick is done, execute async trace
 		{
diff --git a/Engine/Source/Runtime/Experimental/ChaosSolverEngine/Private/Chaos/ChaosDebugDrawSubsystem.cpp b/Engine/Source/Runtime/Experimental/ChaosSolverEngine/Private/Chaos/ChaosDebugDrawSubsystem.cpp
--- a/Engine/Source/Runtime/Experimental/ChaosSolverEngine/Private/Chaos/ChaosDebugDrawSubsystem.cpp
+++ b/Engine/Source/Runtime/Experimental/ChaosSolverEngine/Private/Chaos/ChaosDebugDrawSubsystem.cpp
@@ -472,6 +472,8 @@ void UChaosDebugDrawSubsystem::StaticOnWorldTickStart(UWorld* World, ELevelTick
 
 void UChaosDebugDrawSubsystem::StaticOnWorldTickEnd(UWorld* World, ELevelTick TickType, float Dt)
 {
+	TRACE_CPUPROFILER_EVENT_SCOPE(UChaosDebugDrawSubsystem::StaticOnWorldTickEnd)
+
 	if (UChaosDebugDrawSubsystem* CDDSystem = World->GetSubsystem<UChaosDebugDrawSubsystem>())
 	{
 		CDDSystem->OnWorldTickEnd(TickType, Dt);
diff --git a/Engine/Source/Runtime/Experimental/Dataflow/Simulation/Private/Dataflow/DataflowSimulationManager.cpp b/Engine/Source/Runtime/Experimental/Dataflow/Simulation/Private/Dataflow/DataflowSimulationManager.cpp
--- a/Engine/Source/Runtime/Experimental/Dataflow/Simulation/Private/Dataflow/DataflowSimulationManager.cpp
+++ b/Engine/Source/Runtime/Experimental/Dataflow/Simulation/Private/Dataflow/DataflowSimulationManager.cpp
@@ -224,6 +224,8 @@ void UDataflowSimulationManager::OnStartup()
 	OnWorldPostActorTick = FWorldDelegates::OnWorldPostActorTick.AddLambda(
 	[](const UWorld* SimulationWorld, ELevelTick LevelTick, const float DeltaSeconds)
 	{
+		TRACE_CPUPROFILER_EVENT_SCOPE(UDataflowSimulationManager::OnWorldPostActorTick)
+
 		if(SimulationWorld)
 		{
 			if (UDataflowSimulationManager* DataflowManager = SimulationWorld->GetSubsystem<UDataflowSimulationManager>())
diff --git a/Engine/Source/Runtime/NavigationSystem/Public/NavigationSystem.h b/Engine/Source/Runtime/NavigationSystem/Public/NavigationSystem.h
--- a/Engine/Source/Runtime/NavigationSystem/Public/NavigationSystem.h
+++ b/Engine/Source/Runtime/NavigationSystem/Public/NavigationSystem.h
@@ -1425,7 +1425,11 @@ protected:
 	NAVIGATIONSYSTEM_API void OnLevelRemovedFromWorld(ULevel* InLevel, UWorld* InWorld);
 
 	/** Handler for FWorldDelegates::OnWorldPostActorTick event */
-	void OnWorldPostActorTick(UWorld* World, ELevelTick TickType, float DeltaTime) { PostponeAsyncQueries(); }
+	void OnWorldPostActorTick(UWorld* World, ELevelTick TickType, float DeltaTime)
+	{
+		TRACE_CPUPROFILER_EVENT_SCOPE(UNavigationSystemV1::OnWorldPostActorTick)
+		PostponeAsyncQueries();
+	}
 
 	/** Adds given request to requests queue. Note it's to be called only on game thread only */
 	NAVIGATIONSYSTEM_API void AddAsyncQuery(const FAsyncPathFindingQuery& Query);
