Fixes nav mesh generation crash on application exit when bDoFullyAsyncNavDataGathering is set to true in the project and FNavigationDataHandler::DemandLazyDataGathering is executing when the application is exiting.

Destroying the NavOctree while other threads are running nav building tasks is unsafe, so now the GameThread will wait for nav mesh build to cancel which waits at least until running async tasks are finished.

diff --git a/Engine/Source/Runtime/NavigationSystem/Private/NavigationSystem.cpp b/Engine/Source/Runtime/NavigationSystem/Private/NavigationSystem.cpp
--- a/Engine/Source/Runtime/NavigationSystem/Private/NavigationSystem.cpp
+++ b/Engine/Source/Runtime/NavigationSystem/Private/NavigationSystem.cpp
@@ -5057,18 +5057,30 @@ void UNavigationSystemV1::CleanUp(FNavigationSystem::ECleanupMode Mode)
 
 	FCoreUObjectDelegates::ReloadCompleteDelegate.Remove(ReloadCompleteDelegateHandle);
 
-	DestroyNavOctree();
-	
-	SetCrowdManager(nullptr);
-
+	// Unregister and cleanup navigation data before destroying their dependencies.
+	// The order of operations here mirrors ANavigationData::UnregisterAndCleanUp(), 
+	// minus it having to resolve this NavigationSystem.
 	for (int32 Idx = NavDataSet.Num() - 1; Idx >= 0; --Idx)
 	{
 		ANavigationData* NavData = NavDataSet[Idx];
 		if (NavData)
 		{
-			UnregisterNavData(NavData);
+			// Unregister the nav data
+			if (NavData->IsRegistered())
+			{
+				UnregisterNavData(NavData);
+			}
+
+			// Clean up nav data before the cleaning up the rest of the system. This may block while the NavData waits on async 
+			// tasks that it started, but this is safer than cleaning up navigation systems while those tasks are running, since
+			// those tasks may access state we're about to destroy such as the NavOctree.
+			NavData->CleanUp();
 		}
-	}	
+	}
+
+	DestroyNavOctree();
+
+	SetCrowdManager(nullptr);
 
 	if (NavDataSet.Num())
 	{
