[AINavigation] Moved unsubscribing from FCoreDelegates::OnGetOnScreenMessages in NavigationSystemV1's destructor over to NavigationSystemV1::CleanUp where the other unsubscribing takes place - by the time the navigation system's destructor gets called the WeakObjectPtr stored in FCoreDelegates::OnGetOnScreenMessages is no longer valid and unsubscribing fails.

diff --git a/Engine/Source/Runtime/NavigationSystem/Private/NavigationSystem.cpp b/Engine/Source/Runtime/NavigationSystem/Private/NavigationSystem.cpp
--- a/Engine/Source/Runtime/NavigationSystem/Private/NavigationSystem.cpp
+++ b/Engine/Source/Runtime/NavigationSystem/Private/NavigationSystem.cpp
@@ -672,7 +672,7 @@ void FNavRegenTimeSliceManager::LogTileStatistics(const TArray<TObjectPtr<ANavig
 #endif // !UE_BUILD_SHIPPING
 
 //----------------------------------------------------------------------//
-// UNavigationSystemV1                                                                
+// UNavigationSystemV1
 //----------------------------------------------------------------------//
 bool UNavigationSystemV1::bNavigationAutoUpdateEnabled = true;
 
@@ -692,7 +692,7 @@ bool UNavigationSystemV1::bUpdateNavOctreeOnComponentChange = true;
 bool UNavigationSystemV1::bStaticRuntimeNavigation = false;
 bool UNavigationSystemV1::bIsPIEActive = false;
 //----------------------------------------------------------------------//
-// life cycle stuff                                                                
+// life cycle stuff
 //----------------------------------------------------------------------//
 
 UNavigationSystemV1::UNavigationSystemV1(const FObjectInitializer& ObjectInitializer)
@@ -919,14 +919,10 @@ UNavigationSystemV1::UNavigationSystemV1(const FObjectInitializer& ObjectInitial
 	}
 }
 
-PRAGMA_DISABLE_DEPRECATION_WARNINGS
-UNavigationSystemV1::~UNavigationSystemV1()
+void UNavigationSystemV1::FinishDestroy()
 {
-	CleanUp(FNavigationSystem::ECleanupMode::CleanupUnsafe);
-
-#if !UE_BUILD_SHIPPING
-	FCoreDelegates::OnGetOnScreenMessages.RemoveAll(this);
-#endif // !UE_BUILD_SHIPPING
+	CleanUp();
+	Super::FinishDestroy();
 }
 
 void UNavigationSystemV1::GatherDebugLabels(TArray<FString>& InOutDebugLabels) const
@@ -978,8 +974,6 @@ void UNavigationSystemV1::GatherDebugLabels(TArray<FString>& InOutDebugLabels) c
 	InOutDebugLabels.Add(TEXT("")); // empty line
 }
 
-PRAGMA_ENABLE_DEPRECATION_WARNINGS
-
 void UNavigationSystemV1::ConfigureAsStatic(bool bEnableStatic)
 {
 	bStaticRuntimeNavigation = bEnableStatic;
@@ -1780,7 +1774,7 @@ void UNavigationSystemV1::SetNavigationAutoUpdateEnabled(bool bNewEnable, UNavig
 #endif // WITH_EDITOR
 
 //----------------------------------------------------------------------//
-// Public querying interface                                                                
+// Public querying interface
 //----------------------------------------------------------------------//
 FPathFindingResult UNavigationSystemV1::FindPathSync(const FNavAgentProperties& AgentProperties, FPathFindingQuery Query, EPathFindingMode::Type Mode)
 {
@@ -2093,8 +2087,8 @@ bool UNavigationSystemV1::ProjectPointToNavigation(const FVector& Point, FNavLoc
 	}
 
 	return NavData != nullptr && NavData->ProjectPoint(Point, OutLocation
-	                                            , FNavigationSystem::IsValidExtent(Extent) ? Extent : NavData->GetConfig().DefaultQueryExtent
-	                                            , QueryFilter);
+											, FNavigationSystem::IsValidExtent(Extent) ? Extent : NavData->GetConfig().DefaultQueryExtent
+											, QueryFilter);
 }
 
 UNavigationPath* UNavigationSystemV1::FindPathToActorSynchronously(UObject* WorldContextObject, const FVector& PathStart, AActor* GoalActor, float TetherDistance, AActor* PathfindingContext, TSubclassOf<UNavigationQueryFilter> FilterClass)
@@ -2849,7 +2843,7 @@ void UNavigationSystemV1::UnregisterNavData(ANavigationData* NavData)
 		return;
 	}
 
-    AgentToNavDataMap.Remove(NavData->GetNavAgentProperties());
+	AgentToNavDataMap.Remove(NavData->GetNavAgentProperties());
 
 	FScopeLock Lock(&NavDataRegistration);
 	NavDataRegistrationQueue.Remove(NavData);
@@ -3241,7 +3235,7 @@ void UNavigationSystemV1::ResetCachedFilter(TSubclassOf<UNavigationQueryFilter>
 bool UNavigationSystemV1::ShouldCreateNavigationSystemInstance(const UWorld* World) const
 {
 	ensureMsgf(IsTemplate(), TEXT("This method is expected to only be called on Template objects"
-							   " to determine if an instance of this type should be created."));
+								" to determine if an instance of this type should be created."));
 
 	return (World != nullptr)
 		&& (ShouldAllowClientSideNavigation() || (World->GetNetMode() != NM_Client));
@@ -4939,10 +4933,10 @@ void UNavigationSystemV1::OnLevelRemovedFromWorld(ULevel* InLevel, UWorld* InWor
 					{
 						// removing manually first so that UnregisterNavData won't mess with NavDataSet
 						NavDataSet.RemoveAt(DataIndex, EAllowShrinking::No);
-                            UnregisterNavData(NavData);
-                    }
-                }
-            }
+						UnregisterNavData(NavData);
+					}
+				}
+			}
 		}
 	}
 }
@@ -5054,6 +5048,9 @@ void UNavigationSystemV1::CleanUp(FNavigationSystem::ECleanupMode Mode)
 	FWorldDelegates::LevelAddedToWorld.RemoveAll(this);
 	FWorldDelegates::LevelRemovedFromWorld.RemoveAll(this);
 	FWorldDelegates::OnWorldBeginTearDown.RemoveAll(this);
+#if !UE_BUILD_SHIPPING
+	FCoreDelegates::OnGetOnScreenMessages.RemoveAll(this);
+#endif // !UE_BUILD_SHIPPING
 
 	FCoreUObjectDelegates::ReloadCompleteDelegate.Remove(ReloadCompleteDelegateHandle);
 
@@ -5085,8 +5082,8 @@ void UNavigationSystemV1::CleanUp(FNavigationSystem::ECleanupMode Mode)
 	if (NavDataSet.Num())
 	{
 		UE_LOG(LogNavigation, Error, TEXT("UNavigationSystemV1::CleanUp still has data in NavDataSet after unregister them all"));
-        NavDataSet.Reset();
-    }
+		NavDataSet.Reset();
+	}
 
 	if (AgentToNavDataMap.Num())
 	{
diff --git a/Engine/Source/Runtime/NavigationSystem/Public/NavigationSystem.h b/Engine/Source/Runtime/NavigationSystem/Public/NavigationSystem.h
--- a/Engine/Source/Runtime/NavigationSystem/Public/NavigationSystem.h
+++ b/Engine/Source/Runtime/NavigationSystem/Public/NavigationSystem.h
@@ -300,7 +300,7 @@ class UNavigationSystemV1 : public UNavigationSystemBase
 
 public:
 	NAVIGATIONSYSTEM_API UNavigationSystemV1(const FObjectInitializer& ObjectInitializer = FObjectInitializer::Get());
-	NAVIGATIONSYSTEM_API virtual ~UNavigationSystemV1();
+	NAVIGATIONSYSTEM_API virtual ~UNavigationSystemV1() = default;
 
 	NAVIGATIONSYSTEM_API void GatherDebugLabels(TArray<FString>& InOutDebugLabels) const;
 
@@ -597,6 +597,7 @@ public:
 
 	//~ Begin UObject Interface
 	NAVIGATIONSYSTEM_API virtual void PostInitProperties() override;
+	NAVIGATIONSYSTEM_API virtual void FinishDestroy() override;
 	static NAVIGATIONSYSTEM_API void AddReferencedObjects(UObject* InThis, FReferenceCollector& Collector);
 #if WITH_EDITOR
 	NAVIGATIONSYSTEM_API virtual void PostEditChangeChainProperty(FPropertyChangedChainEvent& PropertyChangedEvent) override;
@@ -1280,10 +1281,6 @@ protected:
 	
 	FNavigationOctreeController DefaultOctreeController;
 
-	/** Map of all custom navigation links, that are relevant for path following */
-	UE_DEPRECATED(5.3, "LinkIds are now based on FNavLinkId. Use CustomNavLinksMap instead, CustomLinksMap is no longer populated or used in the engine.")
-	TMap<uint32, FNavigationSystem::FCustomLinkOwnerInfo> CustomLinksMap;
-
 	TMap<FNavLinkId, FNavigationSystem::FCustomLinkOwnerInfo> CustomNavLinksMap;
 
 	FNavigationDirtyAreasController DefaultDirtyAreasController;
