[AINavigation] CIS Issue 668962 - preventing UNavigationSystemV1::CleanUp being called on CDOs in UNavigationSystemV1::FinishDestroy

#rnx

diff --git a/Engine/Source/Runtime/NavigationSystem/Private/NavigationSystem.cpp b/Engine/Source/Runtime/NavigationSystem/Private/NavigationSystem.cpp
--- a/Engine/Source/Runtime/NavigationSystem/Private/NavigationSystem.cpp
+++ b/Engine/Source/Runtime/NavigationSystem/Private/NavigationSystem.cpp
@@ -921,7 +921,10 @@ UNavigationSystemV1::UNavigationSystemV1(const FObjectInitializer& ObjectInitial
 
 void UNavigationSystemV1::FinishDestroy()
 {
-	CleanUp();
+	if (HasAnyFlags(RF_ClassDefaultObject) == false)
+	{
+		CleanUp(FNavigationSystem::ECleanupMode::CleanupUnsafe);
+	}
 	Super::FinishDestroy();
 }
 
