Adding a public accessor for ARecastNavMesh::NavLinkJumpDownConfig

diff --git a/Engine/Source/Runtime/NavigationSystem/Public/NavMesh/RecastNavMesh.h b/Engine/Source/Runtime/NavigationSystem/Public/NavMesh/RecastNavMesh.h
--- a/Engine/Source/Runtime/NavigationSystem/Public/NavMesh/RecastNavMesh.h
+++ b/Engine/Source/Runtime/NavigationSystem/Public/NavMesh/RecastNavMesh.h
@@ -1450,6 +1450,8 @@ public:
 	NAVIGATIONSYSTEM_API const TMap<FIntPoint, struct FRecastInternalDebugData>* GetDebugDataMap() const;
 #endif
 
+	const FNavLinkGenerationJumpDownConfig& GetNavLinkJumpDownConfig() const { return NavLinkJumpDownConfig; }
+
 protected:
 
 	NAVIGATIONSYSTEM_API void UpdatePolyRefBitsPreview();
