[NavigationSystem] used TStrongObjectPtr to fix a todo for async pathfindings

diff --git a/Engine/Source/Runtime/NavigationSystem/Private/NavigationSystem.cpp b/Engine/Source/Runtime/NavigationSystem/Private/NavigationSystem.cpp
--- a/Engine/Source/Runtime/NavigationSystem/Private/NavigationSystem.cpp
+++ b/Engine/Source/Runtime/NavigationSystem/Private/NavigationSystem.cpp
@@ -1958,12 +1958,8 @@ void UNavigationSystemV1::PerformAsyncQueries(TArray<FAsyncPathFindingQuery> Pat
 	int32 NumProcessed = 0;
 	for (FAsyncPathFindingQuery& Query : PathFindingQueries)
 	{
-		// @todo this is not necessarily the safest way to use UObjects outside of main thread. 
-		//	think about something else.
-		const ANavigationData* NavData = Query.NavData.IsValid() ? Query.NavData.Get() : GetDefaultNavDataInstance(FNavigationSystem::DontCreate);
-
 		// perform query
-		if (NavData)
+		if (const TStrongObjectPtr<const ANavigationData> NavData = Query.NavData.Pin())
 		{
 			if (Query.Mode == EPathFindingMode::Hierarchical)
 			{
