[NavigationSystem] recreate missing FPImplRecastNavMesh after UNavigationSystemV1::CleanUp (changing NavigationSystemConfig, using ANavSystemConfigOverride, etc.)

diff --git a/Engine/Source/Runtime/NavigationSystem/Private/NavMesh/RecastNavMesh.cpp b/Engine/Source/Runtime/NavigationSystem/Private/NavMesh/RecastNavMesh.cpp
--- a/Engine/Source/Runtime/NavigationSystem/Private/NavMesh/RecastNavMesh.cpp
+++ b/Engine/Source/Runtime/NavigationSystem/Private/NavMesh/RecastNavMesh.cpp
@@ -1091,12 +1091,17 @@ void ARecastNavMesh::RestrictBuildingToActiveTiles(bool InRestrictBuildingToActi
 
 void ARecastNavMesh::OnRegistered()
 {
-	Super::OnRegistered();
-
-	check(RecastNavMeshImpl);
+	// (Re)allocate if necessary (might have been destroyed by a call to Cleanup)
+	// Note that it needs to be created before calling the base class
+	if (RecastNavMeshImpl == nullptr)
+	{
+		RecastNavMeshImpl = new FPImplRecastNavMesh(this);
+	}
 
 	// This check can fail when the NavMeshVersion indicates the map needs the nav mesh rebuilt
 	ensure(RecastNavMeshImpl->GetRecastMesh() == nullptr || RecastNavMeshImpl->GetRecastMesh()->getBVQuantFactor((uint8)ENavigationDataResolution::Default) != 0);
+
+	Super::OnRegistered();
 }
 
 void ARecastNavMesh::SerializeRecastNavMesh(FArchive& Ar, FPImplRecastNavMesh*& NavMesh, int32 InNavMeshVersion)
