[NavigationSystem] the call to OnActorRegistered from AActor::PostRegisterAllComponents will now handle registration for the actor and all its components.
Component registration to the navigation system will be ignored until all components are registered to the scene to avoid extra work since actor registration will take care of it.
The component specific operations are now only used after the initial registration (e.g., component added/removed/updated).
This change fixes an issue exposed by the new registration flow when the NavigationElement is created and pushed with all the relevant data instead of relying on the callbacks when processing the registration queue. That delay was hiding the depencies between some components.

diff --git a/Engine/Source/Runtime/NavigationSystem/Private/NavRelevantComponent.cpp b/Engine/Source/Runtime/NavigationSystem/Private/NavRelevantComponent.cpp
--- a/Engine/Source/Runtime/NavigationSystem/Private/NavRelevantComponent.cpp
+++ b/Engine/Source/Runtime/NavigationSystem/Private/NavRelevantComponent.cpp
@@ -71,7 +71,14 @@ FBox UNavRelevantComponent::GetNavigationBounds() const
 	if (!bBoundsInitialized)
 	{
 		CalcAndCacheBounds();
-		bBoundsInitialized = true;
+
+		// Mark bounds as initialized after the actor has registered all its components
+		// since some of them have dependencies.
+		// (e.g., NavModifierComponent relies on registered primitive components)
+		if (GetOwner() && GetOwner()->HasActorRegisteredAllComponents())
+		{
+			bBoundsInitialized = true;
+		}
 	}
 
 	return Bounds;
diff --git a/Engine/Source/Runtime/NavigationSystem/Private/NavigationSystem.cpp b/Engine/Source/Runtime/NavigationSystem/Private/NavigationSystem.cpp
--- a/Engine/Source/Runtime/NavigationSystem/Private/NavigationSystem.cpp
+++ b/Engine/Source/Runtime/NavigationSystem/Private/NavigationSystem.cpp
@@ -174,10 +174,7 @@ CSV_DEFINE_CATEGORY(NavTasksDelays, true);
 CSV_DEFINE_CATEGORY(NavTasks, true);
 CSV_DEFINE_CATEGORY(NavInvokers, true);
 
-//----------------------------------------------------------------------//
-// consts
-//----------------------------------------------------------------------//
-namespace FNavigationSystem
+namespace UE::Navigation::Private
 {
 
 static FAutoConsoleCommandWithWorldArgsAndOutputDevice CmdNavDirtyAreaAroundPlayer(
@@ -222,18 +219,22 @@ static FAutoConsoleCommandWithWorldArgsAndOutputDevice CmdDumpOctreeElements(
 			{
 				if (const FNavigationOctree* Octree = NavSys->GetNavOctree())
 				{
+					int32 NumElements = 0;
 					Octree->FindNodesWithPredicate(
-						[](FNavigationOctree::FNodeIndex /*ParentNodeIndex*/, FNavigationOctree::FNodeIndex /*NodeIndex*/, const FBoxCenterAndExtent&) { return true; },
-						[&](FNavigationOctree::FNodeIndex /*ParentNodeIndex*/, const FNavigationOctree::FNodeIndex NodeIndex, const FBoxCenterAndExtent&)
+						[&NumElements](FNavigationOctree::FNodeIndex /*ParentNodeIndex*/, FNavigationOctree::FNodeIndex /*NodeIndex*/, const FBoxCenterAndExtent&) { return true; },
+						[&NumElements, &OutputDevice, &Octree](FNavigationOctree::FNodeIndex /*ParentNodeIndex*/, const FNavigationOctree::FNodeIndex NodeIndex, const FBoxCenterAndExtent&)
 						{
 							for (const FNavigationOctreeElement& OctreeElement : Octree->GetElementsForNode(NodeIndex))
 							{
+								NumElements++;
 								OutputDevice.Logf(ELogVerbosity::Type::Log, TEXT("%s bounds: [%s] parent:'%s'"),
 									*OctreeElement.GetSourceElement()->GetPathName(),
 									*OctreeElement.Bounds.ToString(),
 									*GetNameSafe(OctreeElement.GetSourceElement().Get().GetNavigationParent().Get()));
 							};
 						});
+
+					OutputDevice.Logf(ELogVerbosity::Type::Log, TEXT("Total: %d elements"), NumElements);
 				}
 				else
 				{
@@ -261,50 +262,86 @@ static FAutoConsoleCommandWithWorldArgsAndOutputDevice CmdLogInvokers(
 		}
 	));
 
-	const FNavDataConfig& GetFallbackNavDataConfig()
-	{
-		static FNavDataConfig FallbackNavDataConfig(FNavigationSystem::FallbackAgentRadius, FNavigationSystem::FallbackAgentHeight);
-		return FallbackNavDataConfig;
-	}
+const FNavDataConfig& GetFallbackNavDataConfig()
+{
+	static FNavDataConfig FallbackNavDataConfig(FNavigationSystem::FallbackAgentRadius, FNavigationSystem::FallbackAgentHeight);
+	return FallbackNavDataConfig;
+}
 
-	FORCEINLINE bool IsValidExtent(const FVector& Extent)
+FORCEINLINE bool IsValidExtent(const FVector& Extent)
+{
+	return Extent != INVALID_NAVEXTENT;
+}
+
+static bool bComponentShouldWaitForActorToRegister = true;
+static FAutoConsoleVariableRef CVarRollbackNavigationComponentShouldWaitForActorToRegister(
+	TEXT("UE.Rollback.Navigation.ComponentShouldWaitForActorToRegister"), bComponentShouldWaitForActorToRegister,
+	TEXT("Components registration to navigation octree will be postponed until owning actor is registered to the octree."
+		"\nCategory: [Navigation]"),
+	ECVF_Default);
+
+bool ShouldComponentWaitForActorToRegister(const UActorComponent* Comp)
+{
+	if (bComponentShouldWaitForActorToRegister)
 	{
-		return Extent != INVALID_NAVEXTENT;
+		// Ignore operations on components until the actor has registered all its components to the scene.
+		// Then, Actor registration to the navigation octree will also registers its components to the octree.
+		if (const AActor* Owner = Comp->GetOwner())
+		{
+			if (!Owner->HasActorRegisteredAllComponents())
+			{
+				return true;
+			}
+		}
 	}
 
-	FCustomLinkOwnerInfo::FCustomLinkOwnerInfo(INavLinkCustomInterface* Link)
+	return false;
+}
+
+} // UE::Navigation::Private
+
+//----------------------------------------------------------------------//
+// FNavigationSystem
+//----------------------------------------------------------------------//
+namespace FNavigationSystem
+{
+FCustomLinkOwnerInfo::FCustomLinkOwnerInfo(INavLinkCustomInterface* Link)
+{
+	LinkInterface = Link;
+	LinkOwner = Link->GetLinkOwner();
+}
+
+bool ShouldLoadNavigationOnClient(ANavigationData& NavData)
+{
+	const UWorld* World = NavData.GetWorld();
+
+	if (World && World->GetNavigationSystem())
 	{
-		LinkInterface = Link;
-		LinkOwner = Link->GetLinkOwner();
+		const UNavigationSystemV1* NavSys = Cast<UNavigationSystemV1>(World->GetNavigationSystem());
+		return NavSys && NavSys->ShouldLoadNavigationOnClient(&NavData);
 	}
-
-	bool ShouldLoadNavigationOnClient(ANavigationData& NavData)
+	else if (GEngine->NavigationSystemClass && GEngine->NavigationSystemClass->IsChildOf<UNavigationSystemV1>())
 	{
-		const UWorld* World = NavData.GetWorld();
-
-		if (World && World->GetNavigationSystem())
-		{
-			const UNavigationSystemV1* NavSys = Cast<UNavigationSystemV1>(World->GetNavigationSystem());
-			return NavSys && NavSys->ShouldLoadNavigationOnClient(&NavData);
-		}
-		else if (GEngine->NavigationSystemClass && GEngine->NavigationSystemClass->IsChildOf<UNavigationSystemV1>())
-		{
-			const UNavigationSystemV1* NavSysCDO = GEngine->NavigationSystemClass->GetDefaultObject<const UNavigationSystemV1>();
-			return NavSysCDO && NavSysCDO->ShouldLoadNavigationOnClient(&NavData);
-		}
-		return false;
+		const UNavigationSystemV1* NavSysCDO = GEngine->NavigationSystemClass->GetDefaultObject<const UNavigationSystemV1>();
+		return NavSysCDO && NavSysCDO->ShouldLoadNavigationOnClient(&NavData);
 	}
+	return false;
+}
 
-	void MakeAllComponentsNeverAffectNav(AActor& Actor)
+void MakeAllComponentsNeverAffectNav(AActor& Actor)
+{
+	const TSet<UActorComponent*> Components = Actor.GetComponents();
+	for (UActorComponent* ActorComp : Components)
 	{
-		const TSet<UActorComponent*> Components = Actor.GetComponents();
-		for (UActorComponent* ActorComp : Components)
-		{
-			ActorComp->SetCanEverAffectNavigation(false);
-		}
+		ActorComp->SetCanEverAffectNavigation(false);
 	}
-} // FNavigationSystem
+}
+
+} // namespace FNavigationSystem
 
+//----------------------------------------------------------------------//
+// NavigationDebugDrawing
+//----------------------------------------------------------------------//
 namespace NavigationDebugDrawing
 {
 	const float PathLineThickness = 3.f;
@@ -312,6 +349,9 @@ namespace NavigationDebugDrawing
 	const FVector PathNodeBoxExtent(16.f);
 }
 
+//----------------------------------------------------------------------//
+// FNavigationInvokerRaw
+//----------------------------------------------------------------------//
 FNavigationInvokerRaw::FNavigationInvokerRaw(const FVector& InLocation, float Min, float Max, const FNavAgentSelector& InSupportedAgents, ENavigationInvokerPriority InPriority)
 : Location(InLocation)
 , RadiusMin(Min)
@@ -2086,7 +2126,7 @@ bool UNavigationSystemV1::ProjectPointToNavigation(const FVector& Point, FNavLoc
 	}
 
 	return NavData != nullptr && NavData->ProjectPoint(Point, OutLocation
-											, FNavigationSystem::IsValidExtent(Extent) ? Extent : NavData->GetConfig().DefaultQueryExtent
+											, UE::Navigation::Private::IsValidExtent(Extent) ? Extent : NavData->GetConfig().DefaultQueryExtent
 											, QueryFilter);
 }
 
@@ -3322,6 +3362,11 @@ void UNavigationSystemV1::RegisterComponentToNavOctree(UActorComponent* Comp)
 		return;
 	}
 
+	if (UE::Navigation::Private::ShouldComponentWaitForActorToRegister(Comp))
+	{
+		return;
+	}
+
 	if (INavRelevantInterface* NavInterface = Cast<INavRelevantInterface>(Comp))
 	{
 		const AActor* OwnerActor = Comp->GetOwner();
@@ -3377,6 +3422,11 @@ void UNavigationSystemV1::UnregisterComponentToNavOctree(UActorComponent* Comp)
 	// skip IsComponentRelevantForNavigation check, it's only for adding new stuff
 	if (Comp)
 	{
+		if (UE::Navigation::Private::ShouldComponentWaitForActorToRegister(Comp))
+		{
+			return;
+		}
+
 		UnregisterNavRelevantObjectStatic(*Comp);
 	}
 }
@@ -3674,6 +3724,11 @@ void UNavigationSystemV1::UpdateComponentInNavOctree(UActorComponent& Comp)
 		return;
 	}
 
+	if (UE::Navigation::Private::ShouldComponentWaitForActorToRegister(&Comp))
+	{
+		return;
+	}
+
 	// special case for early out: use cached nav relevancy
 	if (Comp.bNavigationRelevant)
 	{
@@ -4001,9 +4056,27 @@ void UNavigationSystemV1::OnActorRegistered(AActor* Actor)
 		return;
 	}
 
-	if (const INavRelevantInterface* NavInterface = Cast<INavRelevantInterface>(Actor))
+	if (UE::Navigation::Private::bComponentShouldWaitForActorToRegister)
+	{
+		checkf(Actor && Actor->HasActorRegisteredAllComponents()
+			, TEXT("Actor is expected to be valid and all its components registered."));
+
+		// Tell all components they need to update their navigation bounds before getting registered to the navigation octree.
+		UpdateNavOctreeBounds(Actor);
+
+		// Register actor and all its components to the navigation octree.
+		// i.e., equivalent to:
+		// RegisterNavRelevantObjectStatic(Actor)
+		// For (Component C : Actor->GetComponents())
+		//    RegisterNavRelevantObjectStatic(C)
+		UpdateActorAndComponentsInNavOctree(*Actor);
+	}
+	else
 	{
-		RegisterNavRelevantObjectStatic(*NavInterface, *Actor);
+		if (const INavRelevantInterface* NavInterface = Cast<INavRelevantInterface>(Actor))
+		{
+			RegisterNavRelevantObjectStatic(*NavInterface, *Actor);
+		}
 	}
 }
 
@@ -6122,7 +6195,7 @@ void UNavigationSystemV1::ApplySupportedAgentsFilter()
 	// make sure there's at least one supported navigation agent size
 	if (SupportedAgents.Num() == 0)
 	{
-		SupportedAgents.Add(FNavigationSystem::GetFallbackNavDataConfig());
+		SupportedAgents.Add(UE::Navigation::Private::GetFallbackNavDataConfig());
 	}
 
 	// make all SupportedAgents filtered out by SupportedAgentsMask invalid by
