Mark 'NavOctree bookkeeping' trace scope as verbose as it's very spammy.

#rnx

diff --git a/Engine/Source/Runtime/NavigationSystem/Private/NavigationSystem.cpp b/Engine/Source/Runtime/NavigationSystem/Private/NavigationSystem.cpp
--- a/Engine/Source/Runtime/NavigationSystem/Private/NavigationSystem.cpp
+++ b/Engine/Source/Runtime/NavigationSystem/Private/NavigationSystem.cpp
@@ -69,7 +69,7 @@ DECLARE_CYCLE_STAT(TEXT("Nav Tick: mark dirty"), STAT_Navigation_TickMarkDirty,
 DECLARE_CYCLE_STAT(TEXT("Nav Tick: async build"), STAT_Navigation_TickAsyncBuild, STATGROUP_Navigation);
 DECLARE_CYCLE_STAT(TEXT("Nav Tick: dispatch async pathfinding results"), STAT_Navigation_DispatchAsyncPathfindingResults, STATGROUP_Navigation);
 DECLARE_CYCLE_STAT(TEXT("Nav Tick: async pathfinding"), STAT_Navigation_TickAsyncPathfinding, STATGROUP_Navigation);
-DECLARE_CYCLE_STAT(TEXT("NavOctree bookkeeping"), STAT_NavOctreeBookkeeping, STATGROUP_Navigation);
+DECLARE_CYCLE_STAT_WITH_FLAGS(TEXT("NavOctree bookkeeping"), STAT_NavOctreeBookkeeping, STATGROUP_Navigation, EStatFlags::Verbose);
 
 //----------------------------------------------------------------------//
 // Stats
