[Navigation] added local shared reference when unregistering a navigation element since the operation might remove the last reference and we need it to register back.

diff --git a/Engine/Source/Runtime/NavigationSystem/Private/NavigationDataHandler.cpp b/Engine/Source/Runtime/NavigationSystem/Private/NavigationDataHandler.cpp
--- a/Engine/Source/Runtime/NavigationSystem/Private/NavigationDataHandler.cpp
+++ b/Engine/Source/Runtime/NavigationSystem/Private/NavigationDataHandler.cpp
@@ -352,11 +352,15 @@ void FNavigationDataHandler::UpdateNavOctreeElement(FNavigationElementHandle Ele
 	// don't invalidate pending requests
 	UpdateFlags |= FNavigationOctreeController::OctreeUpdate_Refresh;
 
+	// Use local shared reference to make sure to keep element alive to register back
+	// since unregistering might remove the only reference.
+	TSharedRef<const FNavigationElement> LocalElementRef = UpdatedElement;
+
 	// Always try to unregister, even if element owner doesn't exist in octree (parent nodes).
 	// This is also why we need to provide the new element and not only the handle, so we can access the parent (expected to be always the same for an update).
-	UnregisterElementWithNavOctree(UpdatedElement, UpdateFlags);
+	UnregisterElementWithNavOctree(LocalElementRef, UpdateFlags);
 
-	const FSetElementId RequestId = RegisterElementWithNavOctree(UpdatedElement, UpdateFlags);
+	const FSetElementId RequestId = RegisterElementWithNavOctree(LocalElementRef, UpdateFlags);
 
 	// add original data to pending registration request
 	// so it could be dirtied properly when system receive unregister request while actor is still queued
@@ -402,11 +406,15 @@ void FNavigationDataHandler::UpdateNavOctreeParentChain(const TSharedRef<const F
 		bool bShouldRegisterChildren = true;
 		if (bSkipElementOwnerUpdate == false)
 		{
+			// Use local shared reference to make sure to keep element alive to register back
+			// since unregistering might remove the only reference.
+			TSharedRef<const FNavigationElement> LocalElementRef = Element;
+
 			// We don't want to register NavOctreeElement if owner was not already registered or queued
 			// so we use Unregister/Register combo instead of UpdateNavOctreeElement
-			if (UnregisterElementWithNavOctree(Element, UpdateFlags))
+			if (UnregisterElementWithNavOctree(LocalElementRef, UpdateFlags))
 			{
-				const FSetElementId NewId = RegisterElementWithNavOctree(Element, UpdateFlags);
+				const FSetElementId NewId = RegisterElementWithNavOctree(LocalElementRef, UpdateFlags);
 				bShouldRegisterChildren = NewId.IsValidId();
 			}
 			else
