Tentative fix for: "Warning: Referenced by: FRecastTileGenerator->AddRereferencedObjects"

diff --git a/Engine/Source/Runtime/NavigationSystem/Private/NavMesh/RecastNavMeshGenerator.cpp b/Engine/Source/Runtime/NavigationSystem/Private/NavMesh/RecastNavMeshGenerator.cpp
--- a/Engine/Source/Runtime/NavigationSystem/Private/NavMesh/RecastNavMeshGenerator.cpp
+++ b/Engine/Source/Runtime/NavigationSystem/Private/NavMesh/RecastNavMeshGenerator.cpp
@@ -4936,7 +4936,10 @@ void FRecastTileGenerator::AddReferencedObjects(FReferenceCollector& Collector)
 	{
 		// Local variable since 'AddReferencedObject' requires a non-const& parameter
 		TWeakObjectPtr<const UObject> WeakObject = RelevantData->SourceElement->GetWeakUObject();
-		Collector.AddReferencedObject(WeakObject);
+		if (WeakObject.IsValid())
+		{
+			Collector.AddReferencedObject(WeakObject);
+		}
 	}
 }
 
