[NavigationSystem] update OnActorRegistered to use OnComponentRegistered delegate since it is possible for derived systems to override them.

diff --git a/Engine/Source/Runtime/NavigationSystem/Private/NavigationSystem.cpp b/Engine/Source/Runtime/NavigationSystem/Private/NavigationSystem.cpp
--- a/Engine/Source/Runtime/NavigationSystem/Private/NavigationSystem.cpp
+++ b/Engine/Source/Runtime/NavigationSystem/Private/NavigationSystem.cpp
@@ -4056,6 +4056,11 @@ void UNavigationSystemV1::OnActorRegistered(AActor* Actor)
 		return;
 	}
 
+	if (const INavRelevantInterface* NavInterface = Cast<INavRelevantInterface>(Actor))
+	{
+		RegisterNavRelevantObjectStatic(*NavInterface, *Actor);
+	}
+
 	if (UE::Navigation::Private::bComponentShouldWaitForActorToRegister)
 	{
 		checkf(Actor && Actor->HasActorRegisteredAllComponents()
@@ -4064,18 +4069,14 @@ void UNavigationSystemV1::OnActorRegistered(AActor* Actor)
 		// Tell all components they need to update their navigation bounds before getting registered to the navigation octree.
 		UpdateNavOctreeBounds(Actor);
 
-		// Register actor and all its components to the navigation octree.
-		// i.e., equivalent to:
-		// RegisterNavRelevantObjectStatic(Actor)
-		// For (Component C : Actor->GetComponents())
-		//    RegisterNavRelevantObjectStatic(C)
-		UpdateActorAndComponentsInNavOctree(*Actor);
-	}
-	else
-	{
-		if (const INavRelevantInterface* NavInterface = Cast<INavRelevantInterface>(Actor))
+		// We can now process all the components registered to the scene.
+		// Note that we do so using the delegate since it is possible for derived systems to override them.
+		for (UActorComponent* Component : Actor->GetComponents())
 		{
-			RegisterNavRelevantObjectStatic(*NavInterface, *Actor);
+			if (Component->IsRegistered())
+			{
+				FNavigationSystem::OnComponentRegistered(*Component);
+			}
 		}
 	}
 }
