[Navigation] console command "ai.debug.nav.DumpRepositoryElements" to list all elements registered in the Navigation repository with their bounds and parent

diff --git a/Engine/Source/Runtime/NavigationSystem/Private/NavigationObjectRepository.cpp b/Engine/Source/Runtime/NavigationSystem/Private/NavigationObjectRepository.cpp
--- a/Engine/Source/Runtime/NavigationSystem/Private/NavigationObjectRepository.cpp
+++ b/Engine/Source/Runtime/NavigationSystem/Private/NavigationObjectRepository.cpp
@@ -1,6 +1,7 @@
 ﻿// Copyright Epic Games, Inc. All Rights Reserved.
 
 #include "NavigationObjectRepository.h"
+#include "Misc/OutputDevice.h"
 #include "NavigationSystem.h"
 #include "NavLinkCustomInterface.h"
 #include "AI/NavigationSystemBase.h"
@@ -8,6 +9,36 @@
 #include "AI/Navigation/NavRelevantInterface.h"
 #include "UObject/ObjectKey.h"
 
+namespace UE::Navigation::Private
+{
+static FAutoConsoleCommandWithWorldArgsAndOutputDevice CmdDumpRepositoryElements(
+	TEXT("ai.debug.nav.DumpRepositoryElements"),
+	TEXT("Logs details about each element stored in the navigation repository to the output device."),
+	FConsoleCommandWithWorldArgsAndOutputDeviceDelegate::CreateLambda([](const TArray<FString>& Args, const UWorld* World, FOutputDevice& OutputDevice)
+		{
+			if (const UNavigationObjectRepository* Repository = World->GetSubsystem<UNavigationObjectRepository>())
+			{
+				int32 NumElements = 0;
+
+				Repository->ForEachNavigationElement([&OutputDevice, &NumElements](const TSharedRef<const FNavigationElement>& Element)
+					{
+						NumElements++;
+						OutputDevice.Logf(ELogVerbosity::Log, TEXT("%s bounds: [%s] parent:'%s'"),
+							*Element->GetPathName(),
+							*Element->GetBounds().ToString(),
+							*GetNameSafe(Element->GetNavigationParent().Get()));
+					});
+
+				OutputDevice.Logf(ELogVerbosity::Log, TEXT("Total: %d elements"), NumElements);
+			}
+			else
+			{
+				OutputDevice.Log(ELogVerbosity::Error, TEXT("Command failed since it was unable to find the navigation repository"));
+			}
+		})
+	);
+} // UE::Navigation::Private
+
 TSharedPtr<const FNavigationElement> UNavigationObjectRepository::AddNavigationElement(FNavigationElement&& Element, const ENotifyOnSuccess NotifyOnSuccess /*= ENotifyOnSuccess::Yes*/)
 {
 #if DO_ENSURE // We don't want to execute the Find at all for targets where ensures are disabled
