[Navigation] fixed regression be adding back calls to AddLevelToOctree in ConditionalPopulateNavOctree so levels could register their collision (BSP).
The goal in CL 35531152 was only to not register all actors since those are now handled by the repository, not also removing level collision registration.

diff --git a/Engine/Source/Runtime/NavigationSystem/Private/NavigationSystem.cpp b/Engine/Source/Runtime/NavigationSystem/Private/NavigationSystem.cpp
--- a/Engine/Source/Runtime/NavigationSystem/Private/NavigationSystem.cpp
+++ b/Engine/Source/Runtime/NavigationSystem/Private/NavigationSystem.cpp
@@ -1204,14 +1204,29 @@ bool UNavigationSystemV1::ConditionalPopulateNavOctree()
 #endif // WITH_RECAST
 			}
 
-			if (!DefaultOctreeController.IsNavigationOctreeLocked()
-				&& Repository != nullptr)
+			if (!DefaultOctreeController.IsNavigationOctreeLocked())
 			{
-				// Register all elements registered in the repository world subsystem.
-				Repository->ForEachNavigationElement([this](const TSharedRef<const FNavigationElement>& Element)
+				UWorld* World = GetWorld();
+				check(World);
+
+				// Register level collisions
+				for (int32 LevelIndex = 0; LevelIndex < World->GetNumLevels(); ++LevelIndex)
 				{
-					RegisterNavigationElementWithNavOctree(Element, FNavigationOctreeController::OctreeUpdate_Default);
-				});
+					ULevel* Level = World->GetLevel(LevelIndex);
+					if (ensure(Level) && Level->bIsVisible)
+					{
+						AddLevelToOctree(*Level);
+					}
+				}
+
+				if (Repository != nullptr)
+				{
+					// Register all elements registered in the repository world subsystem.
+					Repository->ForEachNavigationElement([this](const TSharedRef<const FNavigationElement>& Element)
+						{
+							RegisterNavigationElementWithNavOctree(Element, FNavigationOctreeController::OctreeUpdate_Default);
+						});
+				}
 			}
 		}
 	}
@@ -5016,18 +5031,9 @@ void UNavigationSystemV1::OnLevelRemovedFromWorld(ULevel* InLevel, UWorld* InWor
 
 void UNavigationSystemV1::AddLevelToOctree(ULevel& Level)
 {
+	// We only need to add level collision (BSP)
+	// Actors and components are handled by the navigation element repository.
 	AddLevelCollisionToOctree(&Level);
-
-	for (int32 ActorIndex = 0; ActorIndex < Level.Actors.Num(); ActorIndex++)
-	{
-		AActor* Actor = Level.Actors[ActorIndex];
-
-		const bool bLegalActor = IsValid(Actor);
-		if (bLegalActor)
-		{
-			UpdateActorAndComponentsInNavOctree(*Actor);
-		}
-	}
 }
 
 void UNavigationSystemV1::AddLevelCollisionToOctree(ULevel* Level)
