NavigationData: Limit the number of DirtyAreas we can accumulate while rebuild is suspended.
Will prevent growing the array until memory runs out

diff --git a/Engine/Source/Runtime/NavigationSystem/Private/NavigationData.cpp b/Engine/Source/Runtime/NavigationSystem/Private/NavigationData.cpp
--- a/Engine/Source/Runtime/NavigationSystem/Private/NavigationData.cpp
+++ b/Engine/Source/Runtime/NavigationSystem/Private/NavigationData.cpp
@@ -22,11 +22,26 @@
 // for NavigationData
 #define NAVDATAVER_LATEST 13	
 
-static TAutoConsoleVariable<int32> CVarDestroyNavDataInCleanUpAndMarkPendingKill(
-	TEXT("ai.DestroyNavDataInCleanUpAndMarkPendingKill"),
-	1,
-	TEXT("If set to 1 NavData will be destroyed in CleanUpAndMarkPendingKill rather than being marked as garbage.\n"),
-	ECVF_Default);
+namespace UE::Navigation::Private
+{
+	static int32 bDestroyNavDataInCleanUpAndMarkPendingKill = 1;
+	static FAutoConsoleVariableRef CVarDestroyNavDataInCleanUpAndMarkPendingKill(
+		TEXT("ai.DestroyNavDataInCleanUpAndMarkPendingKill"),
+		bDestroyNavDataInCleanUpAndMarkPendingKill,
+		TEXT("If set to 1 NavData will be destroyed in CleanUpAndMarkPendingKill rather than being marked as garbage.\n"),
+		ECVF_Default);
+
+	// This prevents accumulating SuspendedDirtyAreas indefinitely and eventually running out of memory. That could happen if SetRebuildSuspended
+	// is enabled for long periods of time. The limit cannot be too low because moving complex geometry made of hundreds or thousands of Actors
+	// can quickly create several thousands of DirtyAreas.
+	// Careful: Reaching this limit causes a Rebuild All for this navigation data (if it ever resumes generation).
+	static int32 MaxDirtyAreasCountWhileSuspended = 32768;
+	static FAutoConsoleVariableRef CVarMaxDirtyAreasCountWhileSuspended(
+		TEXT("ai.navigation.MaxDirtyAreasCountWhileSuspended"),
+		MaxDirtyAreasCountWhileSuspended,
+		TEXT("If a navigation data accumulates too many dirty areas because it's been indefinitely suspended, we'll stop accumulating them and rebuild the entire navigation data if it ever resumes building.\n"),
+		ECVF_Default);
+}
 
 //----------------------------------------------------------------------//
 // FPathFindingQuery
@@ -515,7 +530,7 @@ void ANavigationData::CleanUpAndMarkPendingKill()
 
 	/* Need to check if the world is valid since, when this is called from Serialize, the World won't be set and Destroy will do nothing, 
 	 * in which case it will crash when it tries to register with the NavSystem in UNavigationSystemV1::ProcessRegistrationCandidates. */
-	if (CVarDestroyNavDataInCleanUpAndMarkPendingKill.GetValueOnAnyThread() && IsValid(GetWorld()))
+	if (UE::Navigation::Private::bDestroyNavDataInCleanUpAndMarkPendingKill && IsValid(GetWorld()))
 	{
 		Destroy();
 	}
@@ -614,7 +629,27 @@ void ANavigationData::RebuildDirtyAreas(const TArray<FNavigationDirtyArea>& Dirt
 {
 	if (bRebuildingSuspended)
 	{
-		SuspendedDirtyAreas.Append(DirtyAreas);
+		if (bReachedSuspendedAreasMaxCount)
+		{
+			return; // Already reached max count
+		}
+
+		const int32 TotalDirtyAreasCount = SuspendedDirtyAreas.Num() + DirtyAreas.Num();
+		if (TotalDirtyAreasCount > UE::Navigation::Private::MaxDirtyAreasCountWhileSuspended)
+		{
+			// Rebuilding has been suspended for too long and we're accumulating too many dirty areas.
+			// Let's just stop accumulating them and remember that the whole navigation data might now be invalid.
+			UE_VLOG_ALWAYS_UELOG(this, LogNavigation, Display, TEXT("%s reached ai.navigation.MaxDirtyAreasCountWhileSuspended (%d). Ignoring dirty areas until rebuilding is unsuspended."),
+				*GetName(),
+				UE::Navigation::Private::MaxDirtyAreasCountWhileSuspended);
+
+			bReachedSuspendedAreasMaxCount = true;
+			SuspendedDirtyAreas.Empty();
+		}
+		else
+		{
+			SuspendedDirtyAreas.Append(DirtyAreas);
+		}
 	}
 	else
 	{
@@ -633,13 +668,25 @@ void ANavigationData::SetRebuildingSuspended(const bool bNewSuspend)
 		UE_VLOG_UELOG(this, LogNavigation, Verbose, TEXT("%s nav generation %s")
 			, *GetName(), bNewSuspend ? TEXT("SUSPENDED") : TEXT("ACTIVE"));
 
-		if (bNewSuspend == false && SuspendedDirtyAreas.Num() > 0)
+		if (bNewSuspend == false)
 		{
-			UE_VLOG_UELOG(this, LogNavigation, Verbose, TEXT("%s resuming nav generation with %d dirty areas")
-				, *GetName(), SuspendedDirtyAreas.Num());
-			// resuming the generation so we need to utilize SuspendedDirtyAreas and clean it
-			RebuildDirtyAreas(SuspendedDirtyAreas);
-			SuspendedDirtyAreas.Empty();
+			if (bReachedSuspendedAreasMaxCount)
+			{
+				UE_VLOG_ALWAYS_UELOG(this, LogNavigation, Warning, TEXT("%s resuming nav generation after MaxDirtyAreasCountWhileSuspended was reached. This will rebuild the entire navigation data."), *GetName());
+
+				// We stopped tracking dirty areas so we'll need to rebuild all.
+				bReachedSuspendedAreasMaxCount = false;
+				RebuildAll();
+			}
+			else if (SuspendedDirtyAreas.Num() > 0)
+			{
+				UE_VLOG_UELOG(this, LogNavigation, Verbose, TEXT("%s resuming nav generation with %d dirty areas")
+					, *GetName(), SuspendedDirtyAreas.Num());
+
+				// resuming the generation so we need to utilize SuspendedDirtyAreas and clean it
+				RebuildDirtyAreas(SuspendedDirtyAreas);
+				SuspendedDirtyAreas.Empty();
+			}
 		}
 	}
 }
diff --git a/Engine/Source/Runtime/NavigationSystem/Public/NavigationData.h b/Engine/Source/Runtime/NavigationSystem/Public/NavigationData.h
--- a/Engine/Source/Runtime/NavigationSystem/Public/NavigationData.h
+++ b/Engine/Source/Runtime/NavigationSystem/Public/NavigationData.h
@@ -714,9 +714,10 @@ public:
 	 *	will be cleaned out. 
 	 *	Note that calling SetRebuildingSuspended(true) won't suspend the nav generation 
 	 *	already in progress.
-	 *	Note2: due to all areas dirtied during generation suspension ending up in 
-	 *	SuspendedDirtyAreas care needs to be taken to not use this feature for 
-	 *	extended periods of time - otherwise SuspendedDirtyAreas can get very large. */
+	 *	Warning: Leaving generation suspended for extended periods of time can trigger a RebuildAll().
+	 *		To prevent the queue from growing out of control, it is limited by ai.navigation.MaxDirtyAreasCountWhileSuspended
+	 *		If that limit is reached, upon resuming, the entire navigation data will be rebuilt.
+	 */
 	NAVIGATIONSYSTEM_API virtual void SetRebuildingSuspended(const bool bNewSuspend);
 
 	/** Retrieves if this NavData instance's navigation generation is suspended */
@@ -1108,6 +1109,12 @@ protected:
 	 *	to be applied at later date with SetRebuildingSuspended(false) call */
 	uint32 bRebuildingSuspended : 1;
 
+	/** Becomes true when too many dirty areas have been accumulated while rebuild
+	 *	was suspended (see SetRebuildingSuspended).
+	 *	If rebuilding is resumed after reaching the max count, this entire navigation
+	 *	data will be rebuilt. */
+	uint32 bReachedSuspendedAreasMaxCount : 1 = false;
+
 #if WITH_EDITORONLY_DATA
 	uint32 bIsBuildingOnLoad : 1;
 #endif
