Fix build: update VLOG macros

diff --git a/Engine/Source/Runtime/NavigationSystem/Private/NavigationData.cpp b/Engine/Source/Runtime/NavigationSystem/Private/NavigationData.cpp
--- a/Engine/Source/Runtime/NavigationSystem/Private/NavigationData.cpp
+++ b/Engine/Source/Runtime/NavigationSystem/Private/NavigationData.cpp
@@ -639,7 +639,7 @@ void ANavigationData::RebuildDirtyAreas(const TArray<FNavigationDirtyArea>& Dirt
 		{
 			// Rebuilding has been suspended for too long and we're accumulating too many dirty areas.
 			// Let's just stop accumulating them and remember that the whole navigation data might now be invalid.
-			UE_VLOG_ALWAYS_UELOG(this, LogNavigation, Display, TEXT("%s reached ai.navigation.MaxDirtyAreasCountWhileSuspended (%d). Ignoring dirty areas until rebuilding is unsuspended."),
+			UE_VLOG_UELOG(this, LogNavigation, Display, TEXT("%s reached ai.navigation.MaxDirtyAreasCountWhileSuspended (%d). Ignoring dirty areas until rebuilding is unsuspended."),
 				*GetName(),
 				UE::Navigation::Private::MaxDirtyAreasCountWhileSuspended);
 
@@ -672,7 +672,7 @@ void ANavigationData::SetRebuildingSuspended(const bool bNewSuspend)
 		{
 			if (bReachedSuspendedAreasMaxCount)
 			{
-				UE_VLOG_ALWAYS_UELOG(this, LogNavigation, Warning, TEXT("%s resuming nav generation after MaxDirtyAreasCountWhileSuspended was reached. This will rebuild the entire navigation data."), *GetName());
+				UE_VLOG_UELOG(this, LogNavigation, Warning, TEXT("%s resuming nav generation after MaxDirtyAreasCountWhileSuspended was reached. This will rebuild the entire navigation data."), *GetName());
 
 				// We stopped tracking dirty areas so we'll need to rebuild all.
 				bReachedSuspendedAreasMaxCount = false;
