Fix a use-after-free in recast caused by a regression from 37261517.

diff --git a/Engine/Source/Runtime/NavigationSystem/Private/NavMesh/RecastNavMeshGenerator.cpp b/Engine/Source/Runtime/NavigationSystem/Private/NavMesh/RecastNavMeshGenerator.cpp
--- a/Engine/Source/Runtime/NavigationSystem/Private/NavMesh/RecastNavMeshGenerator.cpp
+++ b/Engine/Source/Runtime/NavigationSystem/Private/NavMesh/RecastNavMeshGenerator.cpp
@@ -3959,19 +3959,25 @@ dtStatus FRecastTileGenerator::BuildTileCacheLinks(FNavMeshBuildContext& BuildCo
 
 		const bool bABDown = UnrealPosA.Z > UnrealPosB.Z;
 
-		FGeneratedNavigationLink& DownLink = AddLinkSharedLambda(config);
-		DownLink.generatedLinkArea = config.downDirArea;
-		DownLink.generatedLinkPolyFlag = config.downDirPolyFlag;
-		DownLink.Left = bABDown ? UnrealPosA : UnrealPosB;
-		DownLink.Right = bABDown ? UnrealPosB : UnrealPosA;
-		DownLink.Direction = ENavLinkDirection::LeftToRight;
+		const FVector DownLinkLeft = bABDown ? UnrealPosA : UnrealPosB;
+		const FVector DownLinkRight = bABDown ? UnrealPosB : UnrealPosA;
+
+		// Scope to ensure DownLink goes out of scope before UpLink is created, otherwise it could point to invalid memory!
+		{
+			FGeneratedNavigationLink& DownLink = AddLinkSharedLambda(config);
+			DownLink.generatedLinkArea = config.downDirArea;
+			DownLink.generatedLinkPolyFlag = config.downDirPolyFlag;
+			DownLink.Left = DownLinkLeft;
+			DownLink.Right = DownLinkRight;
+			DownLink.Direction = ENavLinkDirection::LeftToRight;
+		}
 
 		// Note that in order for this link to be recognized by the navigation system, it needs to have unique left/right points, which is why the left and right points are flipped and we're marking this link as left-to-right
 		FGeneratedNavigationLink& UpLink = AddLinkSharedLambda(config);
 		UpLink.generatedLinkArea = config.upDirArea;
 		UpLink.generatedLinkPolyFlag = config.upDirPolyFlag;
-		UpLink.Left = DownLink.Right;
-		UpLink.Right = DownLink.Left;
+		UpLink.Left = DownLinkRight;
+		UpLink.Right = DownLinkLeft;
 		UpLink.Direction = ENavLinkDirection::LeftToRight;
 	};
 
