[NavigationSystem] reduced the scope of NavElementAccessDetector to protect only the access to the container and not the creation of the navigation element.
That case was problematic since it was resulting in a reentrant read on that same detector that is not recursive.
Those issues were already fixed by 38443284 but only if `UE.Rollback.Navigation.ComponentShouldWaitForActorToRegister` is set to 'true'. This change will fix the case it is set to 'false'

diff --git a/Engine/Source/Runtime/NavigationSystem/Private/NavigationObjectRepository.cpp b/Engine/Source/Runtime/NavigationSystem/Private/NavigationObjectRepository.cpp
--- a/Engine/Source/Runtime/NavigationSystem/Private/NavigationObjectRepository.cpp
+++ b/Engine/Source/Runtime/NavigationSystem/Private/NavigationObjectRepository.cpp
@@ -112,9 +112,11 @@ TSharedPtr<const FNavigationElement> UNavigationObjectRepository::RegisterNavRel
 	// In such case we update with the latest.
 	if (const TSharedPtr<const FNavigationElement> ElementPtr = GetNavigationElementForUObject(&NavRelevantObject))
 	{
-		UE_MT_SCOPED_WRITE_ACCESS(NavElementAccessDetector);
 		const TSharedRef<const FNavigationElement> NewElement = FNavigationElement::CreateFromNavRelevantInterface(NavRelevantInterface);
-		NavRelevantElements[ElementPtr->GetHandle()] = NewElement;
+		{
+			UE_MT_SCOPED_WRITE_ACCESS(NavElementAccessDetector);
+			NavRelevantElements[ElementPtr->GetHandle()] = NewElement;
+		}
 
 		if (NotifyOnSuccess == ENotifyOnSuccess::Yes)
 		{
