PR #9515: Fix: only spawn the navmesh for the world that was loaded

diff --git a/Engine/Source/Runtime/NavigationSystem/Private/NavigationSystem.cpp b/Engine/Source/Runtime/NavigationSystem/Private/NavigationSystem.cpp
--- a/Engine/Source/Runtime/NavigationSystem/Private/NavigationSystem.cpp
+++ b/Engine/Source/Runtime/NavigationSystem/Private/NavigationSystem.cpp
@@ -5054,10 +5054,15 @@ void UNavigationSystemV1::RemoveLevelCollisionFromOctree(ULevel* Level)
 
 void UNavigationSystemV1::OnPostLoadMap(UWorld* LoadedWorld)
 {
+	if (LoadedWorld != GetWorld())
+	{
+		return;
+	}
+
 	UE_LOG(LogNavigation, Verbose, TEXT("%hs (Package: %s)"), __FUNCTION__, *GetNameSafe(LoadedWorld->GetOuter()));
 
 	// If map has been loaded and there are some navigation bounds volumes 
-	// then create appropriate navigation structured
+	// then create appropriate navigation structure.
 	ANavigationData* NavData = GetDefaultNavDataInstance(FNavigationSystem::DontCreate);
 
 	// Do this if there's currently no navigation
