Allow RecastNavMeshDataChunks to reclaim navmesh tiles when using Dynamic Modifiers Only runtime generation while detaching them if they have different layers but the same XY position.
Also updated the warning for reusing unloaded sublevels to use the updated console command 'LevelStreaming.ShouldReuseUnloadedButStillAroundLevels 0'.

diff --git a/Engine/Source/Runtime/NavigationSystem/Private/NavMesh/RecastNavMeshDataChunk.cpp b/Engine/Source/Runtime/NavigationSystem/Private/NavMesh/RecastNavMeshDataChunk.cpp
--- a/Engine/Source/Runtime/NavigationSystem/Private/NavMesh/RecastNavMeshDataChunk.cpp
+++ b/Engine/Source/Runtime/NavigationSystem/Private/NavMesh/RecastNavMeshDataChunk.cpp
@@ -236,7 +236,7 @@ TArray<FNavTileRef> URecastNavMeshDataChunk::AttachTiles(ARecastNavMesh& NavMesh
 			{
 				if (TileData.TileRawData->RawData == nullptr)
 				{
-					UE_LOG(LogNavigation, Warning, TEXT("Null rawdata. This can be caused by the reuse of unloaded sublevels. 's.ForceGCAfterLevelStreamedOut 1' can be used until this gets fixed."));
+					UE_LOG(LogNavigation, Warning, TEXT("Null rawdata. This can be caused by the reuse of unloaded sublevels. 'LevelStreaming.ShouldReuseUnloadedButStillAroundLevels 0' can be used until this gets fixed."));
 					continue;
 				}
 				
@@ -374,7 +374,10 @@ TArray<FNavTileRef> URecastNavMeshDataChunk::DetachTiles(ARecastNavMesh& NavMesh
 		}
 
 		TArray<const dtMeshTile*> ExtraMeshTiles;
-		const bool bIsDynamic = NavMesh.SupportsRuntimeGeneration();
+		check(NavMesh.GetWorld());
+		const bool bIsGameWorld = NavMesh.GetWorld()->IsGameWorld();
+		// Whether the navmesh is fully dynamic and supports rebuild from geometry. This allows for Dynamic Modifiers Only navmesh to take ownership of tiles on every layer when at the same XY location
+		const bool bIsDynamic = bIsGameWorld && NavMesh.GetRuntimeGenerationMode() == ERuntimeGenerationType::Dynamic;
 		
 		for (FRecastTileData& TileData : Tiles)
 		{
