[NavigationSystem]
* Fixed non unity compile errors

diff --git a/Engine/Source/Runtime/NavigationSystem/Private/NavMesh/NavMeshRenderingComponent.cpp b/Engine/Source/Runtime/NavigationSystem/Private/NavMesh/NavMeshRenderingComponent.cpp
--- a/Engine/Source/Runtime/NavigationSystem/Private/NavMesh/NavMeshRenderingComponent.cpp
+++ b/Engine/Source/Runtime/NavigationSystem/Private/NavMesh/NavMeshRenderingComponent.cpp
@@ -5,6 +5,7 @@
  =============================================================================*/
 
 #include "NavMesh/NavMeshRenderingComponent.h"
+#include "Materials/Material.h"
 #include "NavigationSystem.h"
 #include "Engine/GameViewportClient.h"
 #include "Engine/Canvas.h"
