PR #12867: Fix invoker registration on navigation system override initialisation

diff --git a/Engine/Source/Runtime/NavigationSystem/Private/NavigationInvokerComponent.cpp b/Engine/Source/Runtime/NavigationSystem/Private/NavigationInvokerComponent.cpp
--- a/Engine/Source/Runtime/NavigationSystem/Private/NavigationInvokerComponent.cpp
+++ b/Engine/Source/Runtime/NavigationSystem/Private/NavigationInvokerComponent.cpp
@@ -20,6 +20,8 @@ void UNavigationInvokerComponent::Activate(bool bReset)
 {
 	Super::Activate(bReset);
 
+	UNavigationSystemBase::OnNavigationInitStartStaticDelegate().AddUObject(this, &UNavigationInvokerComponent::OnNavigationInitStart);
+
 	AActor* Owner = GetOwner();
 	if (Owner)
 	{
@@ -36,6 +38,8 @@ void UNavigationInvokerComponent::Deactivate()
 	{
 		UNavigationSystemV1::UnregisterNavigationInvoker(*Owner);
 	}
+
+	UNavigationSystemBase::OnNavigationInitStartStaticDelegate().RemoveAll(this);
 }
 
 void UNavigationInvokerComponent::PostInitProperties()
@@ -60,4 +64,19 @@ void UNavigationInvokerComponent::SetGenerationRadii(const float GenerationRadiu
 {
 	TileGenerationRadius = GenerationRadius;
 	TileRemovalRadius = RemovalRadius;
-}
\ No newline at end of file
+}
+
+void UNavigationInvokerComponent::OnNavigationInitStart(const UNavigationSystemBase& NavSys)
+{
+	if (NavSys.GetWorld() != GetWorld())
+	{
+		return;
+	}
+
+	AActor* Owner = GetOwner();
+	if (Owner)
+	{
+		UNavigationSystemV1::UnregisterNavigationInvoker(*Owner);
+		UNavigationSystemV1::RegisterNavigationInvoker(*Owner, TileGenerationRadius, TileRemovalRadius, SupportedAgents, Priority);
+	}
+}
diff --git a/Engine/Source/Runtime/NavigationSystem/Public/NavigationInvokerComponent.h b/Engine/Source/Runtime/NavigationSystem/Public/NavigationInvokerComponent.h
--- a/Engine/Source/Runtime/NavigationSystem/Public/NavigationInvokerComponent.h
+++ b/Engine/Source/Runtime/NavigationSystem/Public/NavigationInvokerComponent.h
@@ -53,4 +53,7 @@ public:
 	NAVIGATIONSYSTEM_API virtual void Activate(bool bReset = false) override;
 	NAVIGATIONSYSTEM_API virtual void Deactivate() override;
 	NAVIGATIONSYSTEM_API virtual void PostInitProperties() override;
+
+private:
+	void OnNavigationInitStart(const class UNavigationSystemBase& NavSys);
 };
