[NavigationSystem] added code to handle an edge case caused by AActor::RerunConstructionScripts

diff --git a/Engine/Source/Runtime/NavigationSystem/Private/NavigationSystem.cpp b/Engine/Source/Runtime/NavigationSystem/Private/NavigationSystem.cpp
--- a/Engine/Source/Runtime/NavigationSystem/Private/NavigationSystem.cpp
+++ b/Engine/Source/Runtime/NavigationSystem/Private/NavigationSystem.cpp
@@ -3771,9 +3771,22 @@ void UNavigationSystemV1::UpdateComponentInNavOctree(UActorComponent& Comp)
 			}
 		}
 	}
-	else if (Comp.CanEverAffectNavigation())
+	else if (Comp.CanEverAffectNavigation()
+#if WITH_EDITOR
+		|| GIsReconstructingBlueprintInstances
+		// This condition handles a crappy edge case with component registration in Editor
+		// Problem occurs when a component in an instance has 'bCanEverAffectNavigation = false' and AActor::RerunConstructionScripts() is called
+		// 1. Current component values are serialized to FActorComponentInstanceData
+		// 2. Component gets unregistered then destroyed (nothing to do here since it is not affecting navigation)
+		// 3. New component gets created and registered using default values from the template (default is affecting navigation so we register to the octree)
+		// 4. FActorComponentInstanceData is applied to the component (changing `bCanEverAffectNavigation` from `true` to `false` directly in memory)
+		// 5. Component will re-register itself since it was registered at Step 3
+		//    Problem is that we normally don't need to do anything for components never affecting navigation
+		//    so we never unregister that component from the octree!
+#endif // WITH_EDITOR
+		)
 	{
-		// could have been relevant before and not it isn't. Need to check if there's an octree element ID for it
+		// could have been relevant before and now it isn't. Need to check if there's an octree element ID for it
 		if (UNavigationObjectRepository* Repository = UWorld::GetSubsystem<UNavigationObjectRepository>(Comp.GetWorld()))
 		{
 			Repository->UnregisterNavRelevantObject(&Comp);
