Add optional MaxTileCoordinate in the RecastNavMeshTileGenerationDebug to be able to debug a rectangle of navmesh tiles instead of a single tile

diff --git a/Engine/Source/Runtime/NavigationSystem/Private/NavMesh/RecastNavMesh.cpp b/Engine/Source/Runtime/NavigationSystem/Private/NavMesh/RecastNavMesh.cpp
--- a/Engine/Source/Runtime/NavigationSystem/Private/NavMesh/RecastNavMesh.cpp
+++ b/Engine/Source/Runtime/NavigationSystem/Private/NavMesh/RecastNavMesh.cpp
@@ -175,6 +175,17 @@ FNavMeshTileData::FNavData::~FNavData()
 #endif // WITH_RECAST
 }
 
+bool FRecastNavMeshTileGenerationDebug::IsWithinTileCoordinates(const int32 TileX, const int32 TileY) const
+{
+	if (bUseMaxTileCoordinate)
+	{
+		return FMath::IsWithinInclusive(TileX, TileCoordinate.X, MaxTileCoordinate.X)
+			&& FMath::IsWithinInclusive(TileY, TileCoordinate.Y, MaxTileCoordinate.Y);
+	}
+
+	return TileX == TileCoordinate.X && TileY == TileCoordinate.Y;
+}
+
 // Temporary test to narrow a crash.
 void FNavMeshTileData::FNavData::TestPtr() const
 {
@@ -428,6 +439,7 @@ FRecastNavMeshTileGenerationDebug::FRecastNavMeshTileGenerationDebug()
 {
 	bEnabled = false;
 	TileCoordinate = FIntVector::ZeroValue;
+	MaxTileCoordinate = FIntVector::ZeroValue;
 	bGenerateDebugTileOnly = false;
 	bCollisionGeometry = false;
 	HeightFieldRenderMode = EHeightFieldRenderMode::Solid;
@@ -446,6 +458,7 @@ FRecastNavMeshTileGenerationDebug::FRecastNavMeshTileGenerationDebug()
 	bTileCacheContours = false;
 	bTileCachePolyMesh = false;
 	bTileCacheDetailMesh = false;
+	bUseMaxTileCoordinate = false;
 	LinkGenerationDebugFlags = 0;
 	LinkGenerationSelectedEdge = -1;
 }
diff --git a/Engine/Source/Runtime/NavigationSystem/Private/NavMesh/RecastNavMeshGenerator.cpp b/Engine/Source/Runtime/NavigationSystem/Private/NavMesh/RecastNavMeshGenerator.cpp
--- a/Engine/Source/Runtime/NavigationSystem/Private/NavMesh/RecastNavMeshGenerator.cpp
+++ b/Engine/Source/Runtime/NavigationSystem/Private/NavMesh/RecastNavMeshGenerator.cpp
@@ -8113,16 +8113,15 @@ protected:
 #if RECAST_INTERNAL_DEBUG_DATA
 bool FRecastTileGenerator::IsTileDebugActive() const
 {
-	const FIntVector& Coord = TileDebugSettings.TileCoordinate;
-	return (TileDebugSettings.bEnabled && TileX == Coord.X && TileY == Coord.Y) || (TileX == GNavmeshDebugTileX && TileY == GNavmeshDebugTileY);
+	return (TileDebugSettings.bEnabled && TileDebugSettings.IsWithinTileCoordinates(TileX, TileY))
+		|| (TileX == GNavmeshDebugTileX && TileY == GNavmeshDebugTileY);
 }
 
 bool FRecastTileGenerator::IsTileDebugAllowingGeneration() const
 {
-	const FIntVector& Coord = TileDebugSettings.TileCoordinate;
 	if (TileDebugSettings.bEnabled && TileDebugSettings.bGenerateDebugTileOnly)
 	{
-		return TileX == Coord.X && TileY == Coord.Y;
+		return TileDebugSettings.IsWithinTileCoordinates(TileX, TileY);
 	}
 	else if (GNavmeshGenerateDebugTileOnly)
 	{
diff --git a/Engine/Source/Runtime/NavigationSystem/Public/NavMesh/RecastNavMesh.h b/Engine/Source/Runtime/NavigationSystem/Public/NavMesh/RecastNavMesh.h
--- a/Engine/Source/Runtime/NavigationSystem/Public/NavMesh/RecastNavMesh.h
+++ b/Engine/Source/Runtime/NavigationSystem/Public/NavMesh/RecastNavMesh.h
@@ -349,11 +349,17 @@ struct FRecastNavMeshTileGenerationDebug
 	UPROPERTY(Transient, EditAnywhere, Category = Debug)
 	uint32 bEnabled : 1;
 
-	/** Selected tile coordinate, only this tile will have it's internal data kept.
-	 *  Tip: displaying the navmesh using 'Draw Tile Labels' show tile coordinates. */ 
+	/** Selected tile coordinate, only this tile will have it's internal data kept. 
+	 *  When MaxTileCoordinate is enabled, this defines the lower bound of a rectangle.
+	 *  Tip: displaying the navmesh using 'Draw Tile Labels' show tile coordinates. */
 	UPROPERTY(EditAnywhere, Category = Debug)
 	FIntVector TileCoordinate = FIntVector::ZeroValue;
 
+	/** Optional: Highest bound included tile to define a rectangle with TileCoordinate for which the internal data is kept.
+	 *  Tip: displaying the navmesh using 'Draw Tile Labels' show tile coordinates. */
+	UPROPERTY(EditAnywhere, Category = Debug, meta=(EditCondition="bUseMaxTileCoordinate"))
+	FIntVector MaxTileCoordinate = FIntVector::ZeroValue;
+
 	/** If set, the generator will only generate the tile selected to debug (set in TileCoordinate).*/
 	UPROPERTY(EditAnywhere, Category = Debug)
 	uint32 bGenerateDebugTileOnly : 1;
@@ -412,12 +418,19 @@ struct FRecastNavMeshTileGenerationDebug
 	UPROPERTY(EditAnywhere, Category = Debug)
 	uint32 bTileCacheDetailMesh : 1;
 
+	UPROPERTY(EditAnywhere, Category = Debug, meta=(InlineEditConditionToggle))
+	uint32 bUseMaxTileCoordinate : 1;
+
 	UPROPERTY(EditAnywhere, Category = Debug, meta = (Bitmask, BitmaskEnum = "/Script/NavigationSystem.ELinkGenerationDebugFlags"))
 	uint16 LinkGenerationDebugFlags;
 
 	/** Using -1 as no selected edge. */
 	UPROPERTY(EditAnywhere, Category = Debug, meta=(UIMin=-1, ClampMin=-1))
 	int32 LinkGenerationSelectedEdge;
+
+	/** Indicates if the tile coordinate in parameter is within the tile coordinate selected to debug.
+	 *  Either equal to TileCoordinate or inside the rectangle formed by TileCoordinate and MaxTileCoordinate. */
+	NAVIGATIONSYSTEM_API bool IsWithinTileCoordinates(const int32 TileX, const int32 TileY) const;
 };
 
 /**
