Prevent adding auto-generated nav links with null NavArea

diff --git a/Engine/Source/Runtime/NavigationSystem/Private/NavMesh/RecastNavMeshGenerator.cpp b/Engine/Source/Runtime/NavigationSystem/Private/NavMesh/RecastNavMeshGenerator.cpp
--- a/Engine/Source/Runtime/NavigationSystem/Private/NavMesh/RecastNavMeshGenerator.cpp
+++ b/Engine/Source/Runtime/NavigationSystem/Private/NavMesh/RecastNavMeshGenerator.cpp
@@ -3867,7 +3867,10 @@ dtStatus FRecastTileGenerator::BuildTileCacheLinks(FNavMeshBuildContext& BuildCo
 
 	dtLinkBuilderConfig linkBuilderConfig;
 	linkBuilderConfig.jumpDownConfig = TileConfig.JumpDownConfig;
-	linkBuilderConfig.jumpDownConfig.init();
+	if (linkBuilderConfig.jumpDownConfig.enabled)
+	{
+		linkBuilderConfig.jumpDownConfig.init();
+	}
 	
 	linkBuilderConfig.jumpOverConfig = TileConfig.JumpOverConfig;
 	
@@ -3932,100 +3935,107 @@ dtStatus FRecastTileGenerator::BuildTileCacheLinks(FNavMeshBuildContext& BuildCo
 #endif // RECAST_INTERNAL_DEBUG_DATA		
 	}
 
-	auto AddLinkSharedLambda = [&OutGeneratedLinks](const dtNavLinkBuilderJumpDownConfig& config) -> FGeneratedNavigationLink&
+	if (linkBuilderConfig.jumpDownConfig.enabled)
 	{
-		FGeneratedNavigationLink& NewLink = OutGeneratedLinks.Emplace_GetRef();
-		NewLink.bIsGenerated = true;
-		NewLink.NavLinkId = FNavLinkId(config.linkUserId);
-		return NewLink;
-	};
+		auto AddLinkSharedLambda = [&OutGeneratedLinks](const dtNavLinkBuilderJumpDownConfig& config) -> FGeneratedNavigationLink&
+		{
+			FGeneratedNavigationLink& NewLink = OutGeneratedLinks.Emplace_GetRef();
+			NewLink.bIsGenerated = true;
+			NewLink.NavLinkId = FNavLinkId(config.linkUserId);
+			return NewLink;
+		};
 
-	// Used when we're provided with identical up and down areas. In this situation, a single bidirectional link is sufficient for our needs
-	auto AddBidirectionalLinkLambda = [&AddLinkSharedLambda](const dtNavLinkBuilderJumpDownConfig& config, const dtReal* posA, const dtReal* posB)
-	{
-		FGeneratedNavigationLink& NewLink = AddLinkSharedLambda(config);
-		NewLink.generatedLinkArea = config.downDirArea;
-		NewLink.generatedLinkPolyFlag = config.downDirPolyFlag;
-		NewLink.Left = Recast2UnrealPoint(posA);
-		NewLink.Right = Recast2UnrealPoint(posB);
-	};
+		// Used when we're provided with identical up and down areas. In this situation, a single bidirectional link is sufficient for our needs
+		auto AddBidirectionalLinkLambda = [&AddLinkSharedLambda](const dtNavLinkBuilderJumpDownConfig& config, const dtReal* posA, const dtReal* posB)
+		{
+			FGeneratedNavigationLink& NewLink = AddLinkSharedLambda(config);
+			NewLink.generatedLinkArea = config.downDirArea;
+			NewLink.generatedLinkPolyFlag = config.downDirPolyFlag;
+			NewLink.Left = Recast2UnrealPoint(posA);
+			NewLink.Right = Recast2UnrealPoint(posB);
+		};
 
-	// Used when we're provided with different up and down areas. In this situation, we need to create two different one-way link representations in order to meet our needs
-	auto AddDirectionalLinksLambda = [&AddLinkSharedLambda](const dtNavLinkBuilderJumpDownConfig& config, const dtReal* posA, const dtReal* posB)
-	{	
-		const FVector UnrealPosA = Recast2UnrealPoint(posA);
-		const FVector UnrealPosB = Recast2UnrealPoint(posB);
+		// Used when we're provided with different up and down areas. In this situation, we need to create two different one-way link representations in order to meet our needs
+		auto AddDirectionalLinksLambda = [&AddLinkSharedLambda](const dtNavLinkBuilderJumpDownConfig& config, const dtReal* posA, const dtReal* posB)
+		{
+			const FVector UnrealPosA = Recast2UnrealPoint(posA);
+			const FVector UnrealPosB = Recast2UnrealPoint(posB);
 
-		const bool bABDown = UnrealPosA.Z > UnrealPosB.Z;
+			const bool bABDown = UnrealPosA.Z > UnrealPosB.Z;
 
-		const FVector DownLinkLeft = bABDown ? UnrealPosA : UnrealPosB;
-		const FVector DownLinkRight = bABDown ? UnrealPosB : UnrealPosA;
+			const FVector DownLinkLeft = bABDown ? UnrealPosA : UnrealPosB;
+			const FVector DownLinkRight = bABDown ? UnrealPosB : UnrealPosA;
 
-		// Scope to ensure DownLink goes out of scope before UpLink is created, otherwise it could point to invalid memory!
-		{
-			FGeneratedNavigationLink& DownLink = AddLinkSharedLambda(config);
-			DownLink.generatedLinkArea = config.downDirArea;
-			DownLink.generatedLinkPolyFlag = config.downDirPolyFlag;
-			DownLink.Left = DownLinkLeft;
-			DownLink.Right = DownLinkRight;
-			DownLink.Direction = ENavLinkDirection::LeftToRight;
-		}
 
-		// Note that in order for this link to be recognized by the navigation system, it needs to have unique left/right points, which is why the left and right points are flipped and we're marking this link as left-to-right
-		FGeneratedNavigationLink& UpLink = AddLinkSharedLambda(config);
-		UpLink.generatedLinkArea = config.upDirArea;
-		UpLink.generatedLinkPolyFlag = config.upDirPolyFlag;
-		UpLink.Left = DownLinkRight;
-		UpLink.Right = DownLinkLeft;
-		UpLink.Direction = ENavLinkDirection::LeftToRight;
-	};
+			if (config.downDirArea != RECAST_NULL_AREA)
+			{
+				FGeneratedNavigationLink& DownLink = AddLinkSharedLambda(config);
+				DownLink.generatedLinkArea = config.downDirArea;
+				DownLink.generatedLinkPolyFlag = config.downDirPolyFlag;
+				DownLink.Left = DownLinkLeft;
+				DownLink.Right = DownLinkRight;
+				DownLink.Direction = ENavLinkDirection::LeftToRight;
+			}
 
-	const bool bBidirectionalJumpDownLinks = linkBuilderConfig.jumpDownConfig.downDirArea == linkBuilderConfig.jumpDownConfig.upDirArea;
-	TFunctionRef<void(const dtNavLinkBuilderJumpDownConfig&, const dtReal*, const dtReal*)> AddLinkLambda = bBidirectionalJumpDownLinks ? TFunctionRef<void(const dtNavLinkBuilderJumpDownConfig&, const dtReal*, const dtReal*)>(AddBidirectionalLinkLambda) : TFunctionRef<void(const dtNavLinkBuilderJumpDownConfig&, const dtReal*, const dtReal*)>(AddDirectionalLinksLambda);
+			if (config.upDirArea != RECAST_NULL_AREA)
+			{
+				// Note that in order for this link to be recognized by the navigation system, it needs to have unique left/right points, which is why the left and right points are flipped and we're marking this link as left-to-right
+				FGeneratedNavigationLink& UpLink = AddLinkSharedLambda(config);
+				UpLink.generatedLinkArea = config.upDirArea;
+				UpLink.generatedLinkPolyFlag = config.upDirPolyFlag;
+				UpLink.Left = DownLinkRight;
+				UpLink.Right = DownLinkLeft;
+				UpLink.Direction = ENavLinkDirection::LeftToRight;
+			}
+		};
 
-	// Make FGeneratedNavigationLinks
-	for (const dtNavLinkBuilder::JumpLink& link : linkBuilder.m_links)
-	{
-		if (link.flags == dtNavLinkBuilder::FILTERED)
-			continue;
+		const bool bBidirectionalJumpDownLinks = linkBuilderConfig.jumpDownConfig.downDirArea == linkBuilderConfig.jumpDownConfig.upDirArea;
+		TFunctionRef<void(const dtNavLinkBuilderJumpDownConfig&, const dtReal*, const dtReal*)> AddLinkLambda = bBidirectionalJumpDownLinks ? TFunctionRef<void(const dtNavLinkBuilderJumpDownConfig&, const dtReal*, const dtReal*)>(AddBidirectionalLinkLambda) : TFunctionRef<void(const dtNavLinkBuilderJumpDownConfig&, const dtReal*, const dtReal*)>(AddDirectionalLinksLambda);
 
-		// Only "JumpDownConfig are expected for now
-		if (link.action != DT_LINK_ACTION_JUMP_DOWN)
+		// Make FGeneratedNavigationLinks
+		for (const dtNavLinkBuilder::JumpLink& link : linkBuilder.m_links)
 		{
-			continue;
-		}
+			if (link.flags == dtNavLinkBuilder::FILTERED)
+				continue;
 
-		if (TileConfig.JumpDownConfig.linkBuilderFlags & DT_NAVLINK_CREATE_CENTER_POINT_LINK)
-		{
-			// Make a link using the center of the range.
-			dtReal midA[3];
-			dtVlerp(midA, &link.spine0[0], &link.spine1[0], 0.5);
-			dtReal midB[3];
-			dtVlerp(midB, &link.spine0[(link.nspine-1)*3], &link.spine1[(link.nspine-1)*3], 0.5);
-			// Since trajectory validation starts at agentClimb height to ignore small bumps, remove the offset for the actual link height.
-			midA[1] -= linkBuilderConfig.agentClimb;
-			midB[1] -= linkBuilderConfig.agentClimb;
-			AddLinkLambda(linkBuilderConfig.jumpDownConfig, midA, midB);
-		}
+			// Only "JumpDownConfig are expected for now
+			if (link.action != DT_LINK_ACTION_JUMP_DOWN)
+			{
+				continue;
+			}
 
-		if (TileConfig.JumpDownConfig.linkBuilderFlags & DT_NAVLINK_CREATE_EXTREMITY_LINKS)
-		{
-			dtReal posA[3];
-			dtReal posB[3];
-			dtVcopy(posA, &link.spine0[0]);
-			dtVcopy(posB, &link.spine0[(link.nspine-1)*3]);
-			posA[1] -= linkBuilderConfig.agentClimb;
-			posB[1] -= linkBuilderConfig.agentClimb;
-			AddLinkLambda(linkBuilderConfig.jumpDownConfig, posA, posB);
+			if (TileConfig.JumpDownConfig.linkBuilderFlags & DT_NAVLINK_CREATE_CENTER_POINT_LINK)
+			{
+				// Make a link using the center of the range.
+				dtReal midA[3];
+				dtVlerp(midA, &link.spine0[0], &link.spine1[0], 0.5);
+				dtReal midB[3];
+				dtVlerp(midB, &link.spine0[(link.nspine-1)*3], &link.spine1[(link.nspine-1)*3], 0.5);
+				// Since trajectory validation starts at agentClimb height to ignore small bumps, remove the offset for the actual link height.
+				midA[1] -= linkBuilderConfig.agentClimb;
+				midB[1] -= linkBuilderConfig.agentClimb;
+				AddLinkLambda(linkBuilderConfig.jumpDownConfig, midA, midB);
+			}
 
-			dtVcopy(posA, &link.spine1[0]);
-			dtVcopy(posB, &link.spine1[(link.nspine-1)*3]);
-			posA[1] -= linkBuilderConfig.agentClimb;
-			posB[1] -= linkBuilderConfig.agentClimb;
-			AddLinkLambda(linkBuilderConfig.jumpDownConfig, posA, posB);
+			if (TileConfig.JumpDownConfig.linkBuilderFlags & DT_NAVLINK_CREATE_EXTREMITY_LINKS)
+			{
+				dtReal posA[3];
+				dtReal posB[3];
+				dtVcopy(posA, &link.spine0[0]);
+				dtVcopy(posB, &link.spine0[(link.nspine-1)*3]);
+				posA[1] -= linkBuilderConfig.agentClimb;
+				posB[1] -= linkBuilderConfig.agentClimb;
+				AddLinkLambda(linkBuilderConfig.jumpDownConfig, posA, posB);
+
+				dtVcopy(posA, &link.spine1[0]);
+				dtVcopy(posB, &link.spine1[(link.nspine-1)*3]);
+				posA[1] -= linkBuilderConfig.agentClimb;
+				posB[1] -= linkBuilderConfig.agentClimb;
+				AddLinkLambda(linkBuilderConfig.jumpDownConfig, posA, posB);
+			}
 		}
 	}
-	
+
 	LogOnExit();
 	return DT_SUCCESS;
 }
@@ -5499,6 +5509,12 @@ void FRecastNavMeshGenerator::ResolveGeneratedLinkAreas(FRecastBuildConfig& OutC
 		const FNavLinkGenerationJumpDownConfig& JumpDown = DestNavMesh->NavLinkJumpDownConfig;
 		GetAreaIDAndPolyFlags(JumpDown.DownDirectionAreaClass, OutConfig.JumpDownConfig.downDirArea, OutConfig.JumpDownConfig.downDirPolyFlag);
 		GetAreaIDAndPolyFlags(JumpDown.UpDirectionAreaClass, OutConfig.JumpDownConfig.upDirArea, OutConfig.JumpDownConfig.upDirPolyFlag);
+
+		if (OutConfig.JumpDownConfig.downDirArea == RECAST_NULL_AREA && OutConfig.JumpDownConfig.upDirArea == RECAST_NULL_AREA)
+		{
+			UE_LOG(LogNavigation, Warning, TEXT("FRecastTileGenerator: no area class has been defined for the jump down links. Disabling them"));
+			OutConfig.JumpDownConfig.enabled = false;
+		}
 	}
 }
 
diff --git a/Engine/Source/Runtime/NavigationSystem/Public/NavMesh/LinkGenerationConfig.h b/Engine/Source/Runtime/NavigationSystem/Public/NavMesh/LinkGenerationConfig.h
--- a/Engine/Source/Runtime/NavigationSystem/Public/NavMesh/LinkGenerationConfig.h
+++ b/Engine/Source/Runtime/NavigationSystem/Public/NavMesh/LinkGenerationConfig.h
@@ -88,6 +88,7 @@ struct FNavLinkGenerationJumpDownConfig
 	/**
 	 * Area class for downward traversal of links generated by this configuration.
 	 * @note If the value matches UpDirectionAreaClass, a single bidirectional link will be used to represent links generated by this configuration.
+	 * If the value is null, no link will be generated in this direction
 	 */
 	UPROPERTY(EditAnywhere, Config, Category = Settings)
 	TSubclassOf<UNavAreaBase> DownDirectionAreaClass;
@@ -95,6 +96,7 @@ struct FNavLinkGenerationJumpDownConfig
 	/**
 	 * Area class for upward traversal of links generated by this configuration.
 	 * @note If the value matches DownDirectionAreaClass, a single bidirectional link will be used to represent links generated by this configuration.
+	 * If the value is null, no link will be generated in this direction
 	 */
 	UPROPERTY(EditAnywhere, Config, Category = Settings)
 	TSubclassOf<UNavAreaBase> UpDirectionAreaClass;
