RecastNavMesh: Fix comment to mention the correct invalid return value of INVALID_NAVNODEREF

diff --git a/Engine/Source/Runtime/NavigationSystem/Public/NavMesh/PImplRecastNavMesh.h b/Engine/Source/Runtime/NavigationSystem/Public/NavMesh/PImplRecastNavMesh.h
--- a/Engine/Source/Runtime/NavigationSystem/Public/NavMesh/PImplRecastNavMesh.h
+++ b/Engine/Source/Runtime/NavigationSystem/Public/NavMesh/PImplRecastNavMesh.h
@@ -139,7 +139,7 @@ public:
 	NAVIGATIONSYSTEM_API bool ProjectPointMulti(const FVector& Point, TArray<FNavLocation>& OutLocations, const FVector& Extent,
 		FVector::FReal MinZ, FVector::FReal MaxZ, const FNavigationQueryFilter& Filter, const UObject* Owner) const;
 
-	/** Returns nearest navmesh polygon to Loc, or INVALID_NAVMESHREF if Loc is not on the navmesh. */
+	/** Returns nearest navmesh polygon to Loc, or INVALID_NAVNODEREF if Loc is not on the navmesh. */
 	NAVIGATIONSYSTEM_API NavNodeRef FindNearestPoly(FVector const& Loc, FVector const& Extent, const FNavigationQueryFilter& Filter, const UObject* Owner) const;
 
 	/** Finds the polygons along the navigation graph that touch the specified circle. Return true if found any. */
diff --git a/Engine/Source/Runtime/NavigationSystem/Public/NavMesh/RecastNavMesh.h b/Engine/Source/Runtime/NavigationSystem/Public/NavMesh/RecastNavMesh.h
--- a/Engine/Source/Runtime/NavigationSystem/Public/NavMesh/RecastNavMesh.h
+++ b/Engine/Source/Runtime/NavigationSystem/Public/NavMesh/RecastNavMesh.h
@@ -1294,7 +1294,7 @@ public:
 	/** Finds the polygons along the navigation graph that touch the specified circle. */
 	NAVIGATIONSYSTEM_API bool FindPolysAroundCircle(const FVector& CenterPos, const NavNodeRef CenterNodeRef, const FVector::FReal Radius, const FSharedConstNavQueryFilter& Filter, const UObject* QueryOwner, TArray<NavNodeRef>* OutPolys = nullptr, TArray<NavNodeRef>* OutPolysParent = nullptr, TArray<float>* OutPolysCost = nullptr, int32* OutPolysCount = nullptr) const;
 	
-	/** Returns nearest navmesh polygon to Loc, or INVALID_NAVMESHREF if Loc is not on the navmesh. */
+	/** Returns nearest navmesh polygon to Loc, or INVALID_NAVNODEREF if Loc is not on the navmesh. */
 	NAVIGATIONSYSTEM_API NavNodeRef FindNearestPoly(FVector const& Loc, FVector const& Extent, FSharedConstNavQueryFilter Filter = NULL, const UObject* Querier = NULL) const;
 
 	/** Finds the distance to the closest wall, limited to MaxDistance
