Added recast functionality to return edges in a tile that form the border around a the region that is considered inside a nav filter and the outside of a nav filter

diff --git a/Engine/Source/Runtime/NavigationSystem/Private/NavMesh/PImplRecastNavMesh.cpp b/Engine/Source/Runtime/NavigationSystem/Private/NavMesh/PImplRecastNavMesh.cpp
--- a/Engine/Source/Runtime/NavigationSystem/Private/NavMesh/PImplRecastNavMesh.cpp
+++ b/Engine/Source/Runtime/NavigationSystem/Private/NavMesh/PImplRecastNavMesh.cpp
@@ -2691,6 +2691,101 @@ void FPImplRecastNavMesh::GetDebugPolyEdges(const dtMeshTile& Tile, bool bIntern
 	GetTilePolyEdges(Tile, bInternalEdges, bNavMeshEdges, InternalEdgeVerts, NavMeshEdgeVerts);
 }
 
+void FPImplRecastNavMesh::GetTilePolyEdgesForFilter(const dtMeshTile& Tile, const FNavigationQueryFilter& InQueryFilter, TArray<FNavigationWallEdge>& OutEdges) const
+{
+	const FRecastQueryFilter* FilterImplementation = (const FRecastQueryFilter*)(InQueryFilter.GetImplementation());
+	if (!FilterImplementation)
+	{
+		return;
+	}
+
+	const dtQueryFilter* QueryFilter = FilterImplementation->GetAsDetourQueryFilter();
+	for (int i = 0; i < Tile.header->polyCount; ++i)
+	{
+		const dtPoly* Poly = &Tile.polys[i];
+
+		if (Poly->getType() != DT_POLYTYPE_GROUND)
+		{
+			continue;
+		}
+
+		const dtPolyRef Base = DetourNavMesh->getPolyRefBase(&Tile);
+		auto IsPolyExterior = [QueryFilter](const dtPoly& Poly, const dtMeshTile& Tile, const dtPolyRef PolyIndex, const dtPolyRef Base) -> bool
+			{
+				if (!QueryFilter)
+				{
+					return Poly.getArea() == RECAST_NULL_AREA;
+				}
+
+				return !QueryFilter->passFilter(Base | (dtPolyRef)PolyIndex, &Tile, &Poly);
+			};
+
+		const bool bIsPolyiExterior = IsPolyExterior(*Poly, Tile, (dtPolyRef)i, Base);
+		for (int j = 0, nj = (int)Poly->vertCount; j < nj; ++j)
+		{
+			if (bIsPolyiExterior)
+			{
+				const bool bIsPolyjLink = (Poly->neis[j] & DT_EXT_LINK) != 0;
+				if (bIsPolyjLink)
+				{
+					// Recover poly j from the neighbor tile and test the filter
+					unsigned int k = Poly->firstLink;
+					while (k != DT_NULL_LINK)
+					{
+						const dtLink& Link = DetourNavMesh->getLink(&Tile, k);
+						if (Link.edge == j)
+						{
+							const dtMeshTile* PolyjTile;
+							const dtPoly* Polyj;
+							const dtStatus Status = DetourNavMesh->getTileAndPolyByRef(Link.ref, &PolyjTile, &Polyj);
+							if (dtStatusSucceed(Status) && Polyj && PolyjTile)
+							{
+								if (!IsPolyExterior(*Polyj, *PolyjTile, Link.ref, DetourNavMesh->getPolyRefBase(PolyjTile)))
+								{
+									// Poly i is interior and neighbor j is exterior
+									const FVector::FReal* V0 = &Tile.verts[Poly->verts[Link.edge] * 3];
+									const FVector::FReal* V1 = &Tile.verts[Poly->verts[(Link.edge + 1) % nj] * 3];
+									OutEdges.Emplace(Recast2UnrVector(V0), Recast2UnrVector(V1));
+								}
+							}
+
+							break;
+						}
+						k = Link.next;
+					}
+				}
+				else
+				{
+					// Check if j is interior to mark that edge
+					const bool bIsPolyjExterior = Poly->neis[j] == 0 || IsPolyExterior(Tile.polys[Poly->neis[j] - 1], Tile, Poly->neis[j] - 1, Base);
+					if (!bIsPolyjExterior)
+					{
+						// Poly i is interior and neighbor j is exterior
+						const FVector::FReal* V0 = &Tile.verts[Poly->verts[j] * 3];
+						const FVector::FReal* V1 = &Tile.verts[Poly->verts[(j + 1) % nj] * 3];
+						OutEdges.Emplace(Recast2UnrVector(V0), Recast2UnrVector(V1));
+					}
+				}
+			}
+			else
+			{
+				// Here we only need to check for no poly since the case where i is Interior and j is Exterior through filter is handled above
+				// with the polygons interchanged and we don't want to duplicate the edges. The inverted case doesn't happen if j is not a poly since 
+				// the for loop won't iterate on it
+				const bool bIsPolyjOutsideNavmesh = Poly->neis[j] == 0;
+				if (bIsPolyjOutsideNavmesh)
+				{
+					// Poly i is interior and neighbor j is exterior
+					// Swap the vertices to keep the same edges orientations
+					const FVector::FReal* V0 = &Tile.verts[Poly->verts[(j + 1) % nj] * 3];
+					const FVector::FReal* V1 = &Tile.verts[Poly->verts[j] * 3];
+					OutEdges.Emplace(Recast2UnrVector(V0), Recast2UnrVector(V1));
+				}
+			}
+		}
+	}
+}
+
 void FPImplRecastNavMesh::GetTilePolyEdges(const dtMeshTile& Tile, bool bGatherInteriorPolyEdges, bool bGatherExteriorNavMeshEdges, TArray<FVector>& OutInteriorPolyEdgeVerts, TArray<FVector>& OutExteriorNavMeshEdgeVerts) const
 {
 	static const FVector::FReal thr = FMath::Square(0.01f);
diff --git a/Engine/Source/Runtime/NavigationSystem/Private/NavMesh/RecastNavMesh.cpp b/Engine/Source/Runtime/NavigationSystem/Private/NavMesh/RecastNavMesh.cpp
--- a/Engine/Source/Runtime/NavigationSystem/Private/NavMesh/RecastNavMesh.cpp
+++ b/Engine/Source/Runtime/NavigationSystem/Private/NavMesh/RecastNavMesh.cpp
@@ -2219,6 +2219,18 @@ bool ARecastNavMesh::FindEdges(const NavNodeRef CenterNodeRef, const FVector Cen
 	return false;
 }
 
+void ARecastNavMesh::GetTilePolyEdgesForFilter(FNavTileRef TileRef, FSharedConstNavQueryFilter Filter, TArray<FNavigationWallEdge>& OutEdges) const
+{
+	const dtNavMesh* const DetourNavMesh = GetRecastMesh();
+	if (DetourNavMesh)
+	{
+		if (const dtMeshTile* const Tile = DetourNavMesh->getTileByRef(static_cast<dtTileRef>(TileRef)))
+		{
+			RecastNavMeshImpl->GetTilePolyEdgesForFilter(*Tile, GetRightFilterRef(Filter), OutEdges);
+		}
+	}
+}
+
 bool ARecastNavMesh::GetEdgesInTile(FNavTileRef TileRef, TArray<FNavigationWallEdge>& OutEdges) const
 {
 	if (HasValidNavmesh())
diff --git a/Engine/Source/Runtime/NavigationSystem/Public/NavMesh/PImplRecastNavMesh.h b/Engine/Source/Runtime/NavigationSystem/Public/NavMesh/PImplRecastNavMesh.h
--- a/Engine/Source/Runtime/NavigationSystem/Public/NavMesh/PImplRecastNavMesh.h
+++ b/Engine/Source/Runtime/NavigationSystem/Public/NavMesh/PImplRecastNavMesh.h
@@ -332,6 +332,14 @@ public:
 	/** workhorse function finding portal edges between corridor polys */
 	NAVIGATIONSYSTEM_API void GetEdgesForPathCorridorImpl(const TArray<NavNodeRef>* PathCorridor, TArray<FNavigationPortalEdge>* PathCorridorEdges, const dtNavMeshQuery& NavQuery) const;
 
+	/** 
+	 * Traverse a tile and returns a list of all the edges around the passed in filter.
+	 * @param Tile The tile whose edges to traverse.
+	 * @param InQueryFilter the filter used to determine if an edge should be considered a border
+	 * @param OutEdges Output edges array
+	*/
+	NAVIGATIONSYSTEM_API void GetTilePolyEdgesForFilter(const dtMeshTile& Tile, const FNavigationQueryFilter& InQueryFilter, TArray<FNavigationWallEdge>& OutEdges) const;
+
 protected:
 	/** 
 	 *	@param ForbiddenFlags polys with flags matching the filter will get added to 
diff --git a/Engine/Source/Runtime/NavigationSystem/Public/NavMesh/RecastNavMesh.h b/Engine/Source/Runtime/NavigationSystem/Public/NavMesh/RecastNavMesh.h
--- a/Engine/Source/Runtime/NavigationSystem/Public/NavMesh/RecastNavMesh.h
+++ b/Engine/Source/Runtime/NavigationSystem/Public/NavMesh/RecastNavMesh.h
@@ -1383,6 +1383,9 @@ public:
 	/** Find up to 64 navmesh eges in up to 64 polys around the center */
 	NAVIGATIONSYSTEM_API bool FindEdges(const NavNodeRef CenterNodeRef, const FVector Center, const FVector::FReal Radius, const FSharedConstNavQueryFilter Filter, TArray<FNavigationWallEdge>& OutEdges) const;
 
+	/** Get all  navmesh edges from tile that form the border of the passed in Filter */
+	NAVIGATIONSYSTEM_API void GetTilePolyEdgesForFilter(FNavTileRef TileRef, FSharedConstNavQueryFilter Filter, TArray<FNavigationWallEdge>& OutEdges) const;
+
 	/** Get all exterior nav mesh edges from tile */
 	NAVIGATIONSYSTEM_API bool GetEdgesInTile(FNavTileRef TileRef, TArray<FNavigationWallEdge>& OutEdges) const;
 
