Fix android compilation for FNavigationWallEdge not having a parameter constructor

diff --git a/Engine/Source/Runtime/NavigationSystem/Public/NavMesh/RecastNavMesh.h b/Engine/Source/Runtime/NavigationSystem/Public/NavMesh/RecastNavMesh.h
--- a/Engine/Source/Runtime/NavigationSystem/Public/NavMesh/RecastNavMesh.h
+++ b/Engine/Source/Runtime/NavigationSystem/Public/NavMesh/RecastNavMesh.h
@@ -326,6 +326,12 @@ namespace ERecastNamedFilter
 
 struct FNavigationWallEdge
 {
+	FNavigationWallEdge() = default;
+	FNavigationWallEdge(const FVector& InStart, const FVector& InEnd)
+		: Start(InStart)
+		, End(InEnd)
+	{
+	}
 	FVector Start = FVector::ZeroVector;
 	FVector End = FVector::ZeroVector;
 };
