[NavigationSystem] UpdateActorAndComponentsInNavOctree should only perform a full update (unregister/register) on registered components. Otherwise they only need to be unregistered. This is to prevent trashed component to stay registered in the NavOctree when executing RerunConstructionScript.

diff --git a/Engine/Source/Runtime/NavigationSystem/Private/NavigationSystem.cpp b/Engine/Source/Runtime/NavigationSystem/Private/NavigationSystem.cpp
--- a/Engine/Source/Runtime/NavigationSystem/Private/NavigationSystem.cpp
+++ b/Engine/Source/Runtime/NavigationSystem/Private/NavigationSystem.cpp
@@ -3842,13 +3842,14 @@ void UNavigationSystemV1::UpdateActorAndComponentsInNavOctree(AActor& Actor, con
 			{
 				if (Component != nullptr
 					&& Component->CanEverAffectNavigation()
+					&& Component->IsRegistered()
 					&& ActorToUpdate.IsComponentRelevantForNavigation(Component))
 				{
 					const INavRelevantInterface* ComponentNavRelevantInterface = Cast<INavRelevantInterface>(Component);
 					if (ComponentNavRelevantInterface && ComponentNavRelevantInterface->IsNavigationRelevant())
 					{
 						UpdateNavRelevantObjectInNavOctreeStatic(*ComponentNavRelevantInterface, *Component, NavSys,
-							[&ActorToUpdate, Component](UNavigationSystemV1& InNavSys, const TSharedRef<const FNavigationElement>& InElement)
+							[](UNavigationSystemV1& InNavSys, const TSharedRef<const FNavigationElement>& InElement)
 							{
 								InNavSys.UpdateNavOctreeElement(InElement->GetHandle(), InElement, FNavigationOctreeController::OctreeUpdate_Default);
 							});
