Add way to display area class in NavFilter Areas entries

diff --git a/Engine/Source/Runtime/NavigationSystem/Public/NavFilters/NavigationQueryFilter.h b/Engine/Source/Runtime/NavigationSystem/Public/NavFilters/NavigationQueryFilter.h
--- a/Engine/Source/Runtime/NavigationSystem/Public/NavFilters/NavigationQueryFilter.h
+++ b/Engine/Source/Runtime/NavigationSystem/Public/NavFilters/NavigationQueryFilter.h
@@ -108,7 +108,7 @@ class UNavigationQueryFilter : public UObject
 	GENERATED_UCLASS_BODY()
 	
 	/** list of overrides for navigation areas */
-	UPROPERTY(EditAnywhere, Category=Filter)
+	UPROPERTY(EditAnywhere, Category=Filter, meta = (TitleProperty = "AreaClass"))
 	TArray<FNavigationFilterArea> Areas;
 
 	/** required flags of navigation nodes */
