PR #13109: [NavigationSystem] replace legacy repository unregisternavrelevantobject method with UnregisterNavRelevantObjectInternal

diff --git a/Engine/Source/Runtime/NavigationSystem/Private/NavigationSystem.cpp b/Engine/Source/Runtime/NavigationSystem/Private/NavigationSystem.cpp
--- a/Engine/Source/Runtime/NavigationSystem/Private/NavigationSystem.cpp
+++ b/Engine/Source/Runtime/NavigationSystem/Private/NavigationSystem.cpp
@@ -3857,10 +3857,9 @@ void UNavigationSystemV1::UpdateActorAndComponentsInNavOctree(AActor& Actor, con
 					}
 				}
 
-				if (NavSys
-					&& ensureMsgf(NavSys->Repository, TEXT("UNavigationObjectRepository is mandatory for navigation system operations.")))
+				if (NavSys)
 				{
-					NavSys->Repository->UnregisterNavRelevantObject(Component);
+					NavSys->UnregisterNavRelevantObjectInternal(*Component);
 				}
 			}
 		};
