Add vislog to log to record invalid pathfind requests and change log verbosity for some error that should be a warning.

diff --git a/Engine/Plugins/AI/MassAI/Source/MassAIBehavior/Private/Tasks/MassNavMeshPathFollowTask.cpp b/Engine/Plugins/AI/MassAI/Source/MassAIBehavior/Private/Tasks/MassNavMeshPathFollowTask.cpp
--- a/Engine/Plugins/AI/MassAI/Source/MassAIBehavior/Private/Tasks/MassNavMeshPathFollowTask.cpp
+++ b/Engine/Plugins/AI/MassAI/Source/MassAIBehavior/Private/Tasks/MassNavMeshPathFollowTask.cpp
@@ -44,6 +44,7 @@ bool FMassNavMeshPathFollowTask::RequestPath(FStateTreeExecutionContext& Context
 	
 	if (!NavData || !TargetLocation.EndOfPathPosition.IsSet())
 	{
+		MASSBEHAVIOR_LOG(Warning, TEXT("%s"), !NavData ? TEXT("Invalid NavData") : TEXT("EndOfPathPosition not set") );
 		return false;
 	}
 		
@@ -124,6 +125,9 @@ bool FMassNavMeshPathFollowTask::RequestPath(FStateTreeExecutionContext& Context
 	
 		return true;
 	}
+
+	MASSBEHAVIOR_LOG(Warning, TEXT("Failed to find a path, result: %s (start: %s, end: %s"),
+		 *UEnum::GetValueAsString(Result.Result), *Query.StartLocation.ToCompactString(), *Query.EndLocation.ToCompactString());
 	
 	return false;
 }
@@ -166,7 +170,7 @@ EStateTreeRunStatus FMassNavMeshPathFollowTask::EnterState(FStateTreeExecutionCo
 		
 	if (!RequestPath(Context, InstanceData.TargetLocation))
 	{
-		MASSBEHAVIOR_LOG(Error, TEXT("Failed to request path."));
+		MASSBEHAVIOR_LOG(Warning, TEXT("Failed to request path."));
 		return EStateTreeRunStatus::Failed;
 	}
 
diff --git a/Engine/Source/Runtime/NavigationSystem/Private/NavMesh/PImplRecastNavMesh.cpp b/Engine/Source/Runtime/NavigationSystem/Private/NavMesh/PImplRecastNavMesh.cpp
--- a/Engine/Source/Runtime/NavigationSystem/Private/NavMesh/PImplRecastNavMesh.cpp
+++ b/Engine/Source/Runtime/NavigationSystem/Private/NavMesh/PImplRecastNavMesh.cpp
@@ -1350,7 +1350,7 @@ bool FPImplRecastNavMesh::InitPathfinding(const FVector& UnrealStart, const FVec
 	Query.findNearestPoly(&RecastStartToProject.X, Extent, Filter, &StartPoly, &RecastStart.X);
 	if (StartPoly == INVALID_NAVNODEREF)
 	{
-		UE_VLOG(NavMeshOwner, LogNavigation, Warning, TEXT("FPImplRecastNavMesh::InitPathfinding start point not on navmesh (%s)"), *UnrealStart.ToString());
+		UE_VLOG_UELOG(NavMeshOwner, LogNavigation, Warning, TEXT("FPImplRecastNavMesh::InitPathfinding start point not on navmesh (%s)"), *UnrealStart.ToString());
 		UE_VLOG_SEGMENT(NavMeshOwner, LogNavigation, Warning, UnrealStart, UnrealEnd, FColor::Red, TEXT("Failed path"));
 		UE_VLOG_LOCATION(NavMeshOwner, LogNavigation, Warning, UnrealStart, 15, FColor::Red, TEXT("Start failed"));
 		UE_VLOG_BOX(NavMeshOwner, LogNavigation, Warning, FBox(UnrealStart - NavExtent, UnrealStart + NavExtent), FColor::Red, TEXT_EMPTY);
@@ -1368,7 +1368,7 @@ bool FPImplRecastNavMesh::InitPathfinding(const FVector& UnrealStart, const FVec
 	{
 		if (Query.isRequiringNavigableEndLocation())
 		{
-			UE_VLOG(NavMeshOwner, LogNavigation, Warning, TEXT("FPImplRecastNavMesh::InitPathfinding end point not on navmesh (%s)"), *UnrealEnd.ToString());
+			UE_VLOG_UELOG(NavMeshOwner, LogNavigation, Warning, TEXT("FPImplRecastNavMesh::InitPathfinding end point not on navmesh (%s)"), *UnrealEnd.ToString());
 			UE_VLOG_SEGMENT(NavMeshOwner, LogNavigation, Warning, UnrealEnd, UnrealEnd, FColor::Red, TEXT("Failed path"));
 			UE_VLOG_LOCATION(NavMeshOwner, LogNavigation, Warning, UnrealEnd, 15, FColor::Red, TEXT("End failed"));
 			UE_VLOG_BOX(NavMeshOwner, LogNavigation, Warning, FBox(UnrealEnd - NavExtent, UnrealEnd + NavExtent), FColor::Red, TEXT_EMPTY);
