Fixup API macros for merged modules, PURE_VIRTUAL does not need API export

diff --git a/Engine/Source/Runtime/LiveLinkInterface/Public/LiveLinkRole.h b/Engine/Source/Runtime/LiveLinkInterface/Public/LiveLinkRole.h
--- a/Engine/Source/Runtime/LiveLinkInterface/Public/LiveLinkRole.h
+++ b/Engine/Source/Runtime/LiveLinkInterface/Public/LiveLinkRole.h
@@ -19,11 +19,11 @@ class ULiveLinkRole : public UObject
 	GENERATED_BODY()
 
 public:
-	LIVELINKINTERFACE_API virtual UScriptStruct* GetStaticDataStruct() const PURE_VIRTUAL(ULiveLinkRole::GetStaticDataStruct, return nullptr;);
-	LIVELINKINTERFACE_API virtual UScriptStruct* GetFrameDataStruct() const PURE_VIRTUAL(ULiveLinkRole::GetFrameDataStruct, return nullptr;);
-	LIVELINKINTERFACE_API virtual UScriptStruct* GetBlueprintDataStruct() const PURE_VIRTUAL(ULiveLinkRole::GetBlueprintDataStruct, return nullptr;);
+	virtual UScriptStruct* GetStaticDataStruct() const PURE_VIRTUAL(ULiveLinkRole::GetStaticDataStruct, return nullptr;);
+	virtual UScriptStruct* GetFrameDataStruct() const PURE_VIRTUAL(ULiveLinkRole::GetFrameDataStruct, return nullptr;);
+	virtual UScriptStruct* GetBlueprintDataStruct() const PURE_VIRTUAL(ULiveLinkRole::GetBlueprintDataStruct, return nullptr;);
 
-	LIVELINKINTERFACE_API virtual bool InitializeBlueprintData(const FLiveLinkSubjectFrameData& InSourceData, FLiveLinkBlueprintDataStruct& OutBlueprintData) const PURE_VIRTUAL(ULiveLinkRole::InitializeBlueprintData, return false;);
+	virtual bool InitializeBlueprintData(const FLiveLinkSubjectFrameData& InSourceData, FLiveLinkBlueprintDataStruct& OutBlueprintData) const PURE_VIRTUAL(ULiveLinkRole::InitializeBlueprintData, return false;);
 
 	LIVELINKINTERFACE_API virtual FText GetDisplayName() const;
 	virtual bool IsStaticDataValid(const FLiveLinkStaticDataStruct& InStaticData, bool& bOutShouldLogWarning) const { return true; }
diff --git a/Engine/Source/Runtime/LiveLinkInterface/Public/LiveLinkSourceFactory.h b/Engine/Source/Runtime/LiveLinkInterface/Public/LiveLinkSourceFactory.h
--- a/Engine/Source/Runtime/LiveLinkInterface/Public/LiveLinkSourceFactory.h
+++ b/Engine/Source/Runtime/LiveLinkInterface/Public/LiveLinkSourceFactory.h
@@ -29,10 +29,10 @@ class ULiveLinkSourceFactory : public UObject
 
 public:
 	/** The name of the menu item (of any EMenuType) */
-	LIVELINKINTERFACE_API virtual FText GetSourceDisplayName() const PURE_VIRTUAL( ULiveLinkSourceFactory::GetSourceDisplayName, return FText(); );
+	virtual FText GetSourceDisplayName() const PURE_VIRTUAL( ULiveLinkSourceFactory::GetSourceDisplayName, return FText(); );
 
 	/** The tooltip of the menu item (of any EMenyType) */
-	LIVELINKINTERFACE_API virtual FText GetSourceTooltip() const PURE_VIRTUAL(ULiveLinkSourceFactory::GetSourceTooltip, return FText(); );
+	virtual FText GetSourceTooltip() const PURE_VIRTUAL(ULiveLinkSourceFactory::GetSourceTooltip, return FText(); );
 
 	UE_DEPRECATED(4.23, "CreateSourceCreationPanel is deprecated. LiveLinkSourceFactory can now be used at runtime. This factory won't work until it's been updated.")
 	LIVELINKINTERFACE_API virtual TSharedPtr<SWidget> CreateSourceCreationPanel();
@@ -63,7 +63,7 @@ public:
 	LIVELINKINTERFACE_API virtual TSharedPtr<SWidget> BuildCreationPanel(FOnLiveLinkSourceCreated OnLiveLinkSourceCreated) const;
 
 	/** Create a new source from a ConnectionString */
-	LIVELINKINTERFACE_API virtual TSharedPtr<ILiveLinkSource> CreateSource(const FString& ConnectionString) const PURE_VIRTUAL(ULiveLinkSourceFactory::CreateSource, return TSharedPtr<ILiveLinkSource>(); );
+	virtual TSharedPtr<ILiveLinkSource> CreateSource(const FString& ConnectionString) const PURE_VIRTUAL(ULiveLinkSourceFactory::CreateSource, return TSharedPtr<ILiveLinkSource>(); );
 
 	/** Whether this factory is enabled. Can be overriden to disable it in certain configurations. */
 	virtual bool IsEnabled() const { return true; }
diff --git a/Engine/Source/Runtime/LiveLinkInterface/Public/LiveLinkSubjectRemapper.h b/Engine/Source/Runtime/LiveLinkInterface/Public/LiveLinkSubjectRemapper.h
--- a/Engine/Source/Runtime/LiveLinkInterface/Public/LiveLinkSubjectRemapper.h
+++ b/Engine/Source/Runtime/LiveLinkInterface/Public/LiveLinkSubjectRemapper.h
@@ -41,13 +41,13 @@ public:
 	/**
 	 * Create an instance of a ILiveLinkSubjectRemapperWorker that can be used outside of the game thread.
 	 */
-	LIVELINKINTERFACE_API virtual FWorkerSharedPtr CreateWorker() PURE_VIRTUAL(ULiveLinkSubjectRemapper::CreateWorker, return FWorkerSharedPtr(););
+	virtual FWorkerSharedPtr CreateWorker() PURE_VIRTUAL(ULiveLinkSubjectRemapper::CreateWorker, return FWorkerSharedPtr(););
 
 	/**
 	 * Get the instance of ILiveLinkSubjectRemapperWorker that was created by the CreateWorker method. 
 	 * Returns an invalid pointer if no instance was created.
 	 */
-	LIVELINKINTERFACE_API virtual FWorkerSharedPtr GetWorker() const PURE_VIRTUAL(ULiveLinkSubjectRemapper::GetWorker, return FWorkerSharedPtr(););
+	virtual FWorkerSharedPtr GetWorker() const PURE_VIRTUAL(ULiveLinkSubjectRemapper::GetWorker, return FWorkerSharedPtr(););
 
 	/**
 	 * Called to initialize the remapper with information from the subject that owns the remapper.
@@ -55,10 +55,10 @@ public:
 	virtual void Initialize(const FLiveLinkSubjectKey& SubjectKey) {}
 
 	/** Get what role is supported by this remapper. */
-	LIVELINKINTERFACE_API virtual TSubclassOf<ULiveLinkRole> GetSupportedRole() const PURE_VIRTUAL(ULiveLinkSubjectRemapper::GetSupportedRole, return TSubclassOf<ULiveLinkRole>(););
+	virtual TSubclassOf<ULiveLinkRole> GetSupportedRole() const PURE_VIRTUAL(ULiveLinkSubjectRemapper::GetSupportedRole, return TSubclassOf<ULiveLinkRole>(););
 	
 	/** Retuns whether the remapper can currently be used for remapping. */
-	LIVELINKINTERFACE_API virtual bool IsValidRemapper() const PURE_VIRTUAL(ULiveLinkSubjectRemapper::IsValid, return false;);
+	virtual bool IsValidRemapper() const PURE_VIRTUAL(ULiveLinkSubjectRemapper::IsValid, return false;);
 
 public:
 	/** Name mapping between source bone name and transformed bone name */
diff --git a/Engine/Source/Runtime/MediaAssets/Public/MediaSource.h b/Engine/Source/Runtime/MediaAssets/Public/MediaSource.h
--- a/Engine/Source/Runtime/MediaAssets/Public/MediaSource.h
+++ b/Engine/Source/Runtime/MediaAssets/Public/MediaSource.h
@@ -54,7 +54,7 @@ public:
 	 * @see GetProxies
 	 */
 	UFUNCTION(BlueprintCallable, Category="Media|MediaSource")
-	MEDIAASSETS_API virtual FString GetUrl() const PURE_VIRTUAL(UMediaSource::GetUrl, return FString(););
+	virtual FString GetUrl() const PURE_VIRTUAL(UMediaSource::GetUrl, return FString(););
 
 	/**
 	 * Validate the media source settings (must be implemented in child classes).
@@ -62,7 +62,7 @@ public:
 	 * @return true if validation passed, false otherwise.
 	 */
 	UFUNCTION(BlueprintCallable, Category="Media|MediaSource")
-	MEDIAASSETS_API virtual bool Validate() const PURE_VIRTUAL(UMediaSource::Validate, return false;);
+	virtual bool Validate() const PURE_VIRTUAL(UMediaSource::Validate, return false;);
 
 #if WITH_EDITOR
 
diff --git a/Engine/Source/Runtime/MovieScene/Public/Bindings/MovieSceneReplaceableBinding.h b/Engine/Source/Runtime/MovieScene/Public/Bindings/MovieSceneReplaceableBinding.h
--- a/Engine/Source/Runtime/MovieScene/Public/Bindings/MovieSceneReplaceableBinding.h
+++ b/Engine/Source/Runtime/MovieScene/Public/Bindings/MovieSceneReplaceableBinding.h
@@ -71,7 +71,7 @@ protected:
 	* Must be implemented. Called during non-editor/runtime to resolve the binding dynamically. In editor worlds/Sequencer will instead use the PreviewSpawnable binding to spawn a preview object.
 	* If no object is returned, Sequencer's BindingOverrides can still be used to dynamically bind the object. See UMovieSceneReplaceableActorBinding for an example.
 	*/
-	UE_API virtual FMovieSceneBindingResolveResult ResolveRuntimeBindingInternal(const FMovieSceneBindingResolveParams& ResolveParams, int32 BindingIndex, TSharedRef<const UE::MovieScene::FSharedPlaybackState> SharedPlaybackState) const PURE_VIRTUAL(UMovieSceneReplaceableBindingBase::ResolveRuntimeBindingInternal, return FMovieSceneBindingResolveResult(););
+	virtual FMovieSceneBindingResolveResult ResolveRuntimeBindingInternal(const FMovieSceneBindingResolveParams& ResolveParams, int32 BindingIndex, TSharedRef<const UE::MovieScene::FSharedPlaybackState> SharedPlaybackState) const PURE_VIRTUAL(UMovieSceneReplaceableBindingBase::ResolveRuntimeBindingInternal, return FMovieSceneBindingResolveResult(););
 
 	/* Called from CreateNewCustomBinding to create the inner spawnable used for sequencer preview. By default will just call GetInnerSpawnableClass and call CreateNewCustomBinding on that. */
 	UE_API virtual UMovieSceneSpawnableBindingBase* CreateInnerSpawnable(UObject* SourceObject, UMovieScene& OwnerMovieScene);
@@ -79,10 +79,10 @@ protected:
 	/*
 	* Must be implemented and return a non abstract spawnable binding class inheriting from UMovieSceneSpawnableBindingBase to use for the preview for this replaceable binding.
 	*/
-	UE_API virtual TSubclassOf<UMovieSceneSpawnableBindingBase> GetInnerSpawnableClass() const PURE_VIRTUAL(UMovieSceneReplaceableBindingBase::GetInnerSpawnableClass, return nullptr;);
+	virtual TSubclassOf<UMovieSceneSpawnableBindingBase> GetInnerSpawnableClass() const PURE_VIRTUAL(UMovieSceneReplaceableBindingBase::GetInnerSpawnableClass, return nullptr;);
 
 	/* Must be implemented. Called from CreateNewCustomBinding to allow the replaceable to initialize any data members from the source object. */
-	UE_API virtual void InitReplaceableBinding(UObject* SourceObject, UMovieScene & OwnerMovieScene) PURE_VIRTUAL(UMovieSceneReplaceableBindingBase::InitReplaceableBinding, return;);
+	virtual void InitReplaceableBinding(UObject* SourceObject, UMovieScene & OwnerMovieScene) PURE_VIRTUAL(UMovieSceneReplaceableBindingBase::InitReplaceableBinding, return;);
 
 protected:
 
diff --git a/Engine/Source/Runtime/MovieScene/Public/Bindings/MovieSceneSpawnableBinding.h b/Engine/Source/Runtime/MovieScene/Public/Bindings/MovieSceneSpawnableBinding.h
--- a/Engine/Source/Runtime/MovieScene/Public/Bindings/MovieSceneSpawnableBinding.h
+++ b/Engine/Source/Runtime/MovieScene/Public/Bindings/MovieSceneSpawnableBinding.h
@@ -88,10 +88,10 @@ protected:
 
 	
 	/* Must be overridden. Handles the actual spawning of the object. Is overridden by UMovieSceneSpawnableActorBindingBase for example to handle actor-specific spawning. */
-	UE_API virtual UObject* SpawnObjectInternal(UWorld* WorldContext, FName SpawnName, const FGuid& BindingId, int32 BindingIndex, UMovieScene& MovieScene, FMovieSceneSequenceIDRef TemplateID, TSharedRef<const UE::MovieScene::FSharedPlaybackState> SharedPlaybackState)  PURE_VIRTUAL(UMovieSceneSpawnableBindingBase::SpawnObjectInternal, return nullptr;);
+	virtual UObject* SpawnObjectInternal(UWorld* WorldContext, FName SpawnName, const FGuid& BindingId, int32 BindingIndex, UMovieScene& MovieScene, FMovieSceneSequenceIDRef TemplateID, TSharedRef<const UE::MovieScene::FSharedPlaybackState> SharedPlaybackState)  PURE_VIRTUAL(UMovieSceneSpawnableBindingBase::SpawnObjectInternal, return nullptr;);
 
 	/* Must be overridden. Handles the actual destruction of the object. Is overridden by UMovieSceneSpawnableActorBindingBase for example to handle actor-specific destruction.*/
-	UE_API virtual void DestroySpawnedObjectInternal(UObject* Object) PURE_VIRTUAL(UMovieSceneSpawnableBindingBase::DestroySpawnedObjectInternal, return;);
+	virtual void DestroySpawnedObjectInternal(UObject* Object) PURE_VIRTUAL(UMovieSceneSpawnableBindingBase::DestroySpawnedObjectInternal, return;);
 
 	// Helper functions used by various base implementations- can be overridden to customize spawn behavior.
 
diff --git a/Engine/Source/Runtime/MovieSceneCapture/Public/Protocols/FrameGrabberProtocol.h b/Engine/Source/Runtime/MovieSceneCapture/Public/Protocols/FrameGrabberProtocol.h
--- a/Engine/Source/Runtime/MovieSceneCapture/Public/Protocols/FrameGrabberProtocol.h
+++ b/Engine/Source/Runtime/MovieSceneCapture/Public/Protocols/FrameGrabberProtocol.h
@@ -51,14 +51,14 @@ protected:
 
 	 * @return Shared pointer to a payload to associate with the frame, or nullptr
 	 */
-	MOVIESCENECAPTURE_API virtual FFramePayloadPtr GetFramePayload(const FFrameMetrics& FrameMetrics) PURE_VIRTUAL(UFrameGrabberProtocol::GetFramePayload, return nullptr;)
+	virtual FFramePayloadPtr GetFramePayload(const FFrameMetrics& FrameMetrics) PURE_VIRTUAL(UFrameGrabberProtocol::GetFramePayload, return nullptr;)
 
 	/**
 	 * Process a captured frame. This may be called on any thread.
 	 *
 	 * @param Frame			The captured frame data, including any payload retrieved from GetFramePayload
 	 */
-	MOVIESCENECAPTURE_API virtual void ProcessFrame(FCapturedFrameData Frame) PURE_VIRTUAL(UFrameGrabberProtocol::ProcessFrame,)
+	virtual void ProcessFrame(FCapturedFrameData Frame) PURE_VIRTUAL(UFrameGrabberProtocol::ProcessFrame,)
 
 private:
 
diff --git a/Engine/Source/Runtime/MovieSceneTracks/Public/Bindings/MovieSceneSpawnableActorBinding.h b/Engine/Source/Runtime/MovieSceneTracks/Public/Bindings/MovieSceneSpawnableActorBinding.h
--- a/Engine/Source/Runtime/MovieSceneTracks/Public/Bindings/MovieSceneSpawnableActorBinding.h
+++ b/Engine/Source/Runtime/MovieSceneTracks/Public/Bindings/MovieSceneSpawnableActorBinding.h
@@ -75,7 +75,7 @@ protected:
 
 	/* Override to provide Actor class to be spawned*/
 	// TODO: Make UFUNCTION
-	MOVIESCENETRACKS_API virtual TSubclassOf<AActor> GetActorClass() const PURE_VIRTUAL(UMovieSceneSpawnableActorBindingBase::GetActorClass, return AActor::StaticClass(););
+	virtual TSubclassOf<AActor> GetActorClass() const PURE_VIRTUAL(UMovieSceneSpawnableActorBindingBase::GetActorClass, return AActor::StaticClass(););
 
 	/* Optionally override to provide an Actor template to use during Spawn */
 	virtual AActor* GetActorTemplate() const { return nullptr; }
diff --git a/Engine/Source/Runtime/NavigationSystem/Public/CrowdManagerBase.h b/Engine/Source/Runtime/NavigationSystem/Public/CrowdManagerBase.h
--- a/Engine/Source/Runtime/NavigationSystem/Public/CrowdManagerBase.h
+++ b/Engine/Source/Runtime/NavigationSystem/Public/CrowdManagerBase.h
@@ -14,14 +14,14 @@ class UCrowdManagerBase : public UObject
 {
 	GENERATED_BODY()
 public:
-	NAVIGATIONSYSTEM_API virtual void Tick(float DeltaTime) PURE_VIRTUAL(UCrowdManagerBase::Tick, );
+	virtual void Tick(float DeltaTime) PURE_VIRTUAL(UCrowdManagerBase::Tick, );
 
 	/** Called by the nav system when a new navigation data instance is registered. */
-	NAVIGATIONSYSTEM_API virtual void OnNavDataRegistered(ANavigationData& NavDataInstance) PURE_VIRTUAL(UCrowdManagerBase::OnNavDataRegistered, );
+	virtual void OnNavDataRegistered(ANavigationData& NavDataInstance) PURE_VIRTUAL(UCrowdManagerBase::OnNavDataRegistered, );
 
 	/** Called by the nav system when a navigation data instance is removed. */
-	NAVIGATIONSYSTEM_API virtual void OnNavDataUnregistered(ANavigationData& NavDataInstance) PURE_VIRTUAL(UCrowdManagerBase::OnNavDataUnregistered, );
+	virtual void OnNavDataUnregistered(ANavigationData& NavDataInstance) PURE_VIRTUAL(UCrowdManagerBase::OnNavDataUnregistered, );
 
-	NAVIGATIONSYSTEM_API virtual void CleanUp(float DeltaTime) PURE_VIRTUAL(UCrowdManagerBase::CleanUp, );
+	virtual void CleanUp(float DeltaTime) PURE_VIRTUAL(UCrowdManagerBase::CleanUp, );
 };
 
diff --git a/Engine/Source/Runtime/NavigationSystem/Public/NavLinkHostInterface.h b/Engine/Source/Runtime/NavigationSystem/Public/NavLinkHostInterface.h
--- a/Engine/Source/Runtime/NavigationSystem/Public/NavLinkHostInterface.h
+++ b/Engine/Source/Runtime/NavigationSystem/Public/NavLinkHostInterface.h
@@ -24,7 +24,7 @@ class INavLinkHostInterface
 	/**
 	 *	Retrieves UNavLinkDefinition derived UClasses hosted by this interface implementer
 	 */
-	NAVIGATIONSYSTEM_API virtual bool GetNavigationLinksClasses(TArray<TSubclassOf<class UNavLinkDefinition> >& OutClasses) const PURE_VIRTUAL(INavLinkHostInterface::GetNavigationLinksClasses,return false;);
+	virtual bool GetNavigationLinksClasses(TArray<TSubclassOf<class UNavLinkDefinition> >& OutClasses) const PURE_VIRTUAL(INavLinkHostInterface::GetNavigationLinksClasses,return false;);
 
 	/** 
 	 *	_Optional_ way of retrieving navigation link data - if INavLinkHostInterface 
diff --git a/Engine/Source/Runtime/NavigationSystem/Public/NavigationData.h b/Engine/Source/Runtime/NavigationSystem/Public/NavigationData.h
--- a/Engine/Source/Runtime/NavigationSystem/Public/NavigationData.h
+++ b/Engine/Source/Runtime/NavigationSystem/Public/NavigationData.h
@@ -772,7 +772,7 @@ protected:
 
 public:
 	/** Returns bounding box for the navmesh. */
-	NAVIGATIONSYSTEM_API virtual FBox GetBounds() const PURE_VIRTUAL(ANavigationData::GetBounds,return FBox(););
+	virtual FBox GetBounds() const PURE_VIRTUAL(ANavigationData::GetBounds,return FBox(););
 	
 	/** Returns list of navigable bounds. */
 	NAVIGATIONSYSTEM_API TArray<FBox> GetNavigableBounds() const;
@@ -893,13 +893,13 @@ public:
 	}
 
 	/** Raycasts batched for efficiency */
-	NAVIGATIONSYSTEM_API virtual void BatchRaycast(TArray<FNavigationRaycastWork>& Workload, FSharedConstNavQueryFilter QueryFilter, const UObject* Querier = NULL) const PURE_VIRTUAL(ANavigationData::BatchRaycast, );
+	virtual void BatchRaycast(TArray<FNavigationRaycastWork>& Workload, FSharedConstNavQueryFilter QueryFilter, const UObject* Querier = NULL) const PURE_VIRTUAL(ANavigationData::BatchRaycast, );
 
 	/**	Tries to move current nav location towards target constrained to navigable area. Faster than ProjectPointToNavmesh.
 	 *	@param OutLocation if successful this variable will be filed with result
 	 *	@return true if successful, false otherwise
 	 */
-	NAVIGATIONSYSTEM_API virtual bool FindMoveAlongSurface(const FNavLocation& StartLocation, const FVector& TargetPosition, FNavLocation& OutLocation, FSharedConstNavQueryFilter Filter = NULL, const UObject* Querier = NULL) const PURE_VIRTUAL(ANavigationData::FindMoveAlongSurface, return false;);
+	virtual bool FindMoveAlongSurface(const FNavLocation& StartLocation, const FVector& TargetPosition, FNavLocation& OutLocation, FSharedConstNavQueryFilter Filter = NULL, const UObject* Querier = NULL) const PURE_VIRTUAL(ANavigationData::FindMoveAlongSurface, return false;);
 
 	/**	Returns the navmesh edges that touch the convex polygon. The edges are not clipped by the polygon. 
 	 *	@param StartLocation a location on the navmesh where to start searching.
@@ -908,7 +908,7 @@ public:
 	 *	@param Filter Nav filter to use, or if nullptr, default filter is used. 
 	 *	@return true if successful, false otherwise
 	 */
-	NAVIGATIONSYSTEM_API virtual bool FindOverlappingEdges(const FNavLocation& StartLocation, TConstArrayView<FVector> ConvexPolygon, TArray<FVector>& OutEdges, FSharedConstNavQueryFilter Filter = NULL, const UObject* Querier = NULL) const PURE_VIRTUAL(ANavigationData::FindOverlappingEdges, return false;);
+	virtual bool FindOverlappingEdges(const FNavLocation& StartLocation, TConstArrayView<FVector> ConvexPolygon, TArray<FVector>& OutEdges, FSharedConstNavQueryFilter Filter = NULL, const UObject* Querier = NULL) const PURE_VIRTUAL(ANavigationData::FindOverlappingEdges, return false;);
 	
 	/**	Searches navmesh edges between the two path points, search up to the convex polygon described in SearchArea. The returned edges are not clipped to the search area polygon.
 	 *  @param Path Path where From and To belong to.
@@ -920,48 +920,48 @@ public:
 	 *	@param Filter Nav filter to use, or if nullptr, default filter is used. 
 	 *	@return true if successful, false otherwise
 	 */
-	NAVIGATIONSYSTEM_API virtual bool GetPathSegmentBoundaryEdges(const FNavigationPath& Path, const FNavPathPoint& StartPoint, const FNavPathPoint& EndPoint, const TConstArrayView<FVector> SearchArea, TArray<FVector>& OutEdges, const float MaxAreaEnterCost, FSharedConstNavQueryFilter Filter = NULL, const UObject* Querier = NULL) const PURE_VIRTUAL(ANavigationData::GetPathSegmentBoundaryEdges, return false;);
+	virtual bool GetPathSegmentBoundaryEdges(const FNavigationPath& Path, const FNavPathPoint& StartPoint, const FNavPathPoint& EndPoint, const TConstArrayView<FVector> SearchArea, TArray<FVector>& OutEdges, const float MaxAreaEnterCost, FSharedConstNavQueryFilter Filter = NULL, const UObject* Querier = NULL) const PURE_VIRTUAL(ANavigationData::GetPathSegmentBoundaryEdges, return false;);
 
-	NAVIGATIONSYSTEM_API virtual FNavLocation GetRandomPoint(FSharedConstNavQueryFilter Filter = NULL, const UObject* Querier = NULL) const PURE_VIRTUAL(ANavigationData::GetRandomPoint, return FNavLocation(););
+	virtual FNavLocation GetRandomPoint(FSharedConstNavQueryFilter Filter = NULL, const UObject* Querier = NULL) const PURE_VIRTUAL(ANavigationData::GetRandomPoint, return FNavLocation(););
 
 	/** finds a random location in Radius, reachable from Origin */
-	NAVIGATIONSYSTEM_API virtual bool GetRandomReachablePointInRadius(const FVector& Origin, float Radius, FNavLocation& OutResult, FSharedConstNavQueryFilter Filter = NULL, const UObject* Querier = NULL) const PURE_VIRTUAL(ANavigationData::GetRandomReachablePointInRadius, return false;);
+	virtual bool GetRandomReachablePointInRadius(const FVector& Origin, float Radius, FNavLocation& OutResult, FSharedConstNavQueryFilter Filter = NULL, const UObject* Querier = NULL) const PURE_VIRTUAL(ANavigationData::GetRandomReachablePointInRadius, return false;);
 
 	/** finds a random location in navigable space, in given Radius */
-	NAVIGATIONSYSTEM_API virtual bool GetRandomPointInNavigableRadius(const FVector& Origin, float Radius, FNavLocation& OutResult, FSharedConstNavQueryFilter Filter = NULL, const UObject* Querier = NULL) const PURE_VIRTUAL(ANavigationData::GetRandomPointInNavigableRadius, return false;);
+	virtual bool GetRandomPointInNavigableRadius(const FVector& Origin, float Radius, FNavLocation& OutResult, FSharedConstNavQueryFilter Filter = NULL, const UObject* Querier = NULL) const PURE_VIRTUAL(ANavigationData::GetRandomPointInNavigableRadius, return false;);
 	
 	/**	Tries to project given Point to this navigation type, within given Extent.
 	 *	@param OutLocation if successful this variable will be filed with result
 	 *	@return true if successful, false otherwise
 	 */
-	NAVIGATIONSYSTEM_API virtual bool ProjectPoint(const FVector& Point, FNavLocation& OutLocation, const FVector& Extent, FSharedConstNavQueryFilter Filter = NULL, const UObject* Querier = NULL) const PURE_VIRTUAL(ANavigationData::ProjectPoint, return false;);
+	virtual bool ProjectPoint(const FVector& Point, FNavLocation& OutLocation, const FVector& Extent, FSharedConstNavQueryFilter Filter = NULL, const UObject* Querier = NULL) const PURE_VIRTUAL(ANavigationData::ProjectPoint, return false;);
 
 	/**	batches ProjectPoint's work for efficiency */
-	NAVIGATIONSYSTEM_API virtual void BatchProjectPoints(TArray<FNavigationProjectionWork>& Workload, const FVector& Extent, FSharedConstNavQueryFilter Filter = NULL, const UObject* Querier = NULL) const PURE_VIRTUAL(ANavigationData::BatchProjectPoints, );
+	virtual void BatchProjectPoints(TArray<FNavigationProjectionWork>& Workload, const FVector& Extent, FSharedConstNavQueryFilter Filter = NULL, const UObject* Querier = NULL) const PURE_VIRTUAL(ANavigationData::BatchProjectPoints, );
 
 	/** Project batch of points using shared search filter. This version is not requiring user to pass in Extent, 
 	 *	and is instead relying on FNavigationProjectionWork.ProjectionLimit.
 	 *	@note function should assert if item's FNavigationProjectionWork.ProjectionLimit is invalid */
-	NAVIGATIONSYSTEM_API virtual void BatchProjectPoints(TArray<FNavigationProjectionWork>& Workload, FSharedConstNavQueryFilter Filter = NULL, const UObject* Querier = NULL) const PURE_VIRTUAL(ANavigationData::BatchProjectPoints, );
+	virtual void BatchProjectPoints(TArray<FNavigationProjectionWork>& Workload, FSharedConstNavQueryFilter Filter = NULL, const UObject* Querier = NULL) const PURE_VIRTUAL(ANavigationData::BatchProjectPoints, );
 
 	/** Calculates path from PathStart to PathEnd and retrieves its cost.
  *	@NOTE this function does not generate string pulled path so the result is an (over-estimated) approximation
  *	@NOTE potentially expensive, so use it with caution */
-	NAVIGATIONSYSTEM_API virtual ENavigationQueryResult::Type CalcPathCost(const FVector& PathStart, const FVector& PathEnd, FVector::FReal& OutPathCost, FSharedConstNavQueryFilter QueryFilter = NULL, const UObject* Querier = NULL) const PURE_VIRTUAL(ANavigationData::CalcPathCost, return ENavigationQueryResult::Invalid;);
+	virtual ENavigationQueryResult::Type CalcPathCost(const FVector& PathStart, const FVector& PathEnd, FVector::FReal& OutPathCost, FSharedConstNavQueryFilter QueryFilter = NULL, const UObject* Querier = NULL) const PURE_VIRTUAL(ANavigationData::CalcPathCost, return ENavigationQueryResult::Invalid;);
 
 	/** Calculates path from PathStart to PathEnd and retrieves its length.
 	 *	@NOTE this function does not generate string pulled path so the result is an (over-estimated) approximation
 	 *	@NOTE potentially expensive, so use it with caution */
-	NAVIGATIONSYSTEM_API virtual ENavigationQueryResult::Type CalcPathLength(const FVector& PathStart, const FVector& PathEnd, FVector::FReal& OutPathLength, FSharedConstNavQueryFilter QueryFilter = NULL, const UObject* Querier = NULL) const PURE_VIRTUAL(ANavigationData::CalcPathLength, return ENavigationQueryResult::Invalid;);
+	virtual ENavigationQueryResult::Type CalcPathLength(const FVector& PathStart, const FVector& PathEnd, FVector::FReal& OutPathLength, FSharedConstNavQueryFilter QueryFilter = NULL, const UObject* Querier = NULL) const PURE_VIRTUAL(ANavigationData::CalcPathLength, return ENavigationQueryResult::Invalid;);
 
 	/** Calculates path from PathStart to PathEnd and retrieves its length.
 	 *	@NOTE this function does not generate string pulled path so the result is an (over-estimated) approximation
 	 *	@NOTE potentially expensive, so use it with caution */
-	NAVIGATIONSYSTEM_API virtual ENavigationQueryResult::Type CalcPathLengthAndCost(const FVector& PathStart, const FVector& PathEnd, FVector::FReal& OutPathLength, FVector::FReal& OutPathCost, FSharedConstNavQueryFilter QueryFilter = NULL, const UObject* Querier = NULL) const PURE_VIRTUAL(ANavigationData::CalcPathLengthAndCost, return ENavigationQueryResult::Invalid;);
+	virtual ENavigationQueryResult::Type CalcPathLengthAndCost(const FVector& PathStart, const FVector& PathEnd, FVector::FReal& OutPathLength, FVector::FReal& OutPathCost, FSharedConstNavQueryFilter QueryFilter = NULL, const UObject* Querier = NULL) const PURE_VIRTUAL(ANavigationData::CalcPathLengthAndCost, return ENavigationQueryResult::Invalid;);
 
 	/** Checks if specified navigation node contains given location 
 	 *	@param Location is expressed in WorldSpace, navigation data is responsible for tansforming if need be */
-	NAVIGATIONSYSTEM_API virtual bool DoesNodeContainLocation(NavNodeRef NodeRef, const FVector& WorldSpaceLocation) const PURE_VIRTUAL(ANavigationData::DoesNodeContainLocation, return false;);
+	virtual bool DoesNodeContainLocation(NavNodeRef NodeRef, const FVector& WorldSpaceLocation) const PURE_VIRTUAL(ANavigationData::DoesNodeContainLocation, return false;);
 
 	NAVIGATIONSYSTEM_API double GetWorldTimeStamp() const;
 
diff --git a/Engine/Source/Runtime/Overlay/Public/Overlays.h b/Engine/Source/Runtime/Overlay/Public/Overlays.h
--- a/Engine/Source/Runtime/Overlay/Public/Overlays.h
+++ b/Engine/Source/Runtime/Overlay/Public/Overlays.h
@@ -54,7 +54,7 @@ public:
 	/**
 	 * Retrieves the set of overlays that should be used by this object (must be implemented by child classes)
 	 */
-	OVERLAY_API virtual TArray<FOverlayItem> GetAllOverlays() const PURE_VIRTUAL(UOverlays::GetAllOverlays, return TArray<FOverlayItem>(););
+	virtual TArray<FOverlayItem> GetAllOverlays() const PURE_VIRTUAL(UOverlays::GetAllOverlays, return TArray<FOverlayItem>(););
 
 	/**
 	 * Retrieves the set of overlays associated with this object for the given timespan (must be implemented by child classes)
@@ -62,5 +62,5 @@ public:
 	 * @param	Time			Determines what overlays should be displayed
 	 * @param	OutOverlays		All overlays that should be displayed for the given Time. Expect this to be emptied.
 	 */
-	OVERLAY_API virtual void GetOverlaysForTime(const FTimespan& Time, TArray<FOverlayItem>& OutOverlays) const PURE_VIRTUAL(UOverlays::GetAllOverlays, return;);
+	virtual void GetOverlaysForTime(const FTimespan& Time, TArray<FOverlayItem>& OutOverlays) const PURE_VIRTUAL(UOverlays::GetAllOverlays, return;);
 };
diff --git a/Engine/Source/Runtime/PacketHandlers/PacketHandler/Classes/HandlerComponentFactory.h b/Engine/Source/Runtime/PacketHandlers/PacketHandler/Classes/HandlerComponentFactory.h
--- a/Engine/Source/Runtime/PacketHandlers/PacketHandler/Classes/HandlerComponentFactory.h
+++ b/Engine/Source/Runtime/PacketHandlers/PacketHandler/Classes/HandlerComponentFactory.h
@@ -19,6 +19,6 @@ class UHandlerComponentFactory : public UObject
 	GENERATED_UCLASS_BODY()
 
 public:
-	PACKETHANDLER_API virtual TSharedPtr<HandlerComponent> CreateComponentInstance(FString& Options)
+	virtual TSharedPtr<HandlerComponent> CreateComponentInstance(FString& Options)
 		PURE_VIRTUAL(UHandlerComponentFactory::CreateComponentInstance, return nullptr;);
 };
diff --git a/Engine/Source/Runtime/TimeManagement/Public/FixedFrameRateCustomTimeStep.h b/Engine/Source/Runtime/TimeManagement/Public/FixedFrameRateCustomTimeStep.h
--- a/Engine/Source/Runtime/TimeManagement/Public/FixedFrameRateCustomTimeStep.h
+++ b/Engine/Source/Runtime/TimeManagement/Public/FixedFrameRateCustomTimeStep.h
@@ -29,7 +29,7 @@ public:
 	//~ End UEngineCustomTimeStep
 
 	/** Get The fixed FrameRate */
-	TIMEMANAGEMENT_API virtual FFrameRate GetFixedFrameRate() const PURE_VIRTUAL(UFixedFrameRateCustomTimeStep::GetFixedFrameRate, return GetFixedFrameRate_PureVirtual(););
+	virtual FFrameRate GetFixedFrameRate() const PURE_VIRTUAL(UFixedFrameRateCustomTimeStep::GetFixedFrameRate, return GetFixedFrameRate_PureVirtual(););
 
 protected:
 	/** Default behavior of the engine. Used FixedFrameRate */
diff --git a/Engine/Source/Runtime/TimeManagement/Public/GenlockedCustomTimeStep.h b/Engine/Source/Runtime/TimeManagement/Public/GenlockedCustomTimeStep.h
--- a/Engine/Source/Runtime/TimeManagement/Public/GenlockedCustomTimeStep.h
+++ b/Engine/Source/Runtime/TimeManagement/Public/GenlockedCustomTimeStep.h
@@ -25,23 +25,23 @@ class UGenlockedCustomTimeStep : public UFixedFrameRateCustomTimeStep
 protected:
 
 	/** Blocks until it gets a sync signal. Returns false if unsuccessful */
-	TIMEMANAGEMENT_API virtual bool WaitForSync() PURE_VIRTUAL(UGenlockedCustomTimeStep::WaitForSync, return false;);
+	virtual bool WaitForSync() PURE_VIRTUAL(UGenlockedCustomTimeStep::WaitForSync, return false;);
 
 	/** Update FApp CurrentTime, IdleTime and DeltaTime */
 	TIMEMANAGEMENT_API void UpdateAppTimes(const double& TimeBeforeSync, const double& TimeAfterSync) const;
 
 public:
 	//~ UFixedFrameRateCustomTimeStep interface
-	TIMEMANAGEMENT_API virtual FFrameRate GetFixedFrameRate() const PURE_VIRTUAL(UGenlockedCustomTimeStep::GetFixedFrameRate, return FFrameRate(24,1););
+	virtual FFrameRate GetFixedFrameRate() const PURE_VIRTUAL(UGenlockedCustomTimeStep::GetFixedFrameRate, return FFrameRate(24,1););
 
 	/** Get the sync rate (not always the same as the fixed frame rate) */
 	virtual FFrameRate GetSyncRate() const { return GetFixedFrameRate(); };
 
 	/** Returns how many syncs occurred since the last tick */
-	TIMEMANAGEMENT_API virtual uint32 GetLastSyncCountDelta() const PURE_VIRTUAL(UGenlockedCustomTimeStep::GetLastSyncCountDelta, return 1;);
+	virtual uint32 GetLastSyncCountDelta() const PURE_VIRTUAL(UGenlockedCustomTimeStep::GetLastSyncCountDelta, return 1;);
 
 	/** Returns true if the Sync related functions will return valid data */
-	TIMEMANAGEMENT_API virtual bool IsLastSyncDataValid() const PURE_VIRTUAL(UGenlockedCustomTimeStep::IsLastSyncDataValid, return false;);
+	virtual bool IsLastSyncDataValid() const PURE_VIRTUAL(UGenlockedCustomTimeStep::IsLastSyncDataValid, return false;);
 
 	/** Returns how many sync counts are expected between ticks*/
 	TIMEMANAGEMENT_API virtual uint32 GetExpectedSyncCountDelta() const;
diff --git a/Engine/Source/Runtime/TimeManagement/Public/TimeSynchronizationSource.h b/Engine/Source/Runtime/TimeManagement/Public/TimeSynchronizationSource.h
--- a/Engine/Source/Runtime/TimeManagement/Public/TimeSynchronizationSource.h
+++ b/Engine/Source/Runtime/TimeManagement/Public/TimeSynchronizationSource.h
@@ -79,31 +79,31 @@ public:
 	 * Get the time of the newest available sample (relative to this source's frame rate).
 	 * Note, in cases where Rollover is allowed and has occurred, this may have a lower value than GetOldestSampleTime. 
 	 */
-	TIMEMANAGEMENT_API virtual FFrameTime GetNewestSampleTime() const PURE_VIRTUAL(UTimeSynchronizationSource::GetNewestSampleTime, return FFrameTime();)
+	virtual FFrameTime GetNewestSampleTime() const PURE_VIRTUAL(UTimeSynchronizationSource::GetNewestSampleTime, return FFrameTime();)
 
 	/**
 	 * Get the time of the oldest available sample (relative to this source's frame rate).
 	 * Note, in cases where Rollover is allowed and has occurred, this may have a higher value than GetNewestSampleTime.
 	 */
-	TIMEMANAGEMENT_API virtual FFrameTime GetOldestSampleTime() const PURE_VIRTUAL(UTimeSynchronizationSource::GetOldestSampleTime, return FFrameTime();)
+	virtual FFrameTime GetOldestSampleTime() const PURE_VIRTUAL(UTimeSynchronizationSource::GetOldestSampleTime, return FFrameTime();)
 
 	/** Get the source actual FrameRate */
-	TIMEMANAGEMENT_API virtual FFrameRate GetFrameRate() const PURE_VIRTUAL(UTimeSynchronizationSource::GetFrameRate, return FFrameRate();)
+	virtual FFrameRate GetFrameRate() const PURE_VIRTUAL(UTimeSynchronizationSource::GetFrameRate, return FFrameRate();)
 
 	/** Used to know if the source is ready to be used for synchronization. */
-	TIMEMANAGEMENT_API virtual bool IsReady() const PURE_VIRTUAL(UTimeSynchronizationSource::IsReady, return false;)
+	virtual bool IsReady() const PURE_VIRTUAL(UTimeSynchronizationSource::IsReady, return false;)
 
 	/** Called when synchronization is started to notify this source to begin buffering frames. */
-	TIMEMANAGEMENT_API virtual bool Open(const FTimeSynchronizationOpenData& OpenData) PURE_VIRTUAL(UTimeSynchronizationSource::Open, return false;)
+	virtual bool Open(const FTimeSynchronizationOpenData& OpenData) PURE_VIRTUAL(UTimeSynchronizationSource::Open, return false;)
 
 	/** Start playing samples. */
-	TIMEMANAGEMENT_API virtual void Start(const FTimeSynchronizationStartData& StartData) PURE_VIRTUAL(UTimeSynchronizationSource::Start, return;)
+	virtual void Start(const FTimeSynchronizationStartData& StartData) PURE_VIRTUAL(UTimeSynchronizationSource::Start, return;)
 
 	/** Called when synchronization has been completed. The source may discard any unnecessary frames. */
-	TIMEMANAGEMENT_API virtual void Close() PURE_VIRTUAL(UTimeSynchronizationSource::Close, return;)
+	virtual void Close() PURE_VIRTUAL(UTimeSynchronizationSource::Close, return;)
 
 	/** Name to used when displaying an error message or to used in UI. */
-	TIMEMANAGEMENT_API virtual FString GetDisplayName() const PURE_VIRTUAL(UTimeSynchronizationSource::GetDisplayName, return FString();)
+	virtual FString GetDisplayName() const PURE_VIRTUAL(UTimeSynchronizationSource::GetDisplayName, return FString();)
 
 public:
 
