Making TileGenerationRadius and TileRemovalRadius variables on UNavigationInvokerComponent BlueprintReadOnly so blueprint can read from them

-This is a QOL change. The example use case is to read these in BP to enable tooling & debugging.

#rnx

diff --git a/Engine/Source/Runtime/NavigationSystem/Public/NavigationInvokerComponent.h b/Engine/Source/Runtime/NavigationSystem/Public/NavigationInvokerComponent.h
--- a/Engine/Source/Runtime/NavigationSystem/Public/NavigationInvokerComponent.h
+++ b/Engine/Source/Runtime/NavigationSystem/Public/NavigationInvokerComponent.h
@@ -22,12 +22,12 @@ class UNavigationInvokerComponent : public UActorComponent
 protected:
 
 	/** Navigation data is requested within a TileGenerationRadius circle around the component owner. */
-	UPROPERTY(EditAnywhere, Category = Navigation, meta = (ClampMin = "0.1", ClampMax = "6400000", UIMin = "0.1", UIMax = "6400000"))
+	UPROPERTY(EditAnywhere, BlueprintReadOnly, Category = Navigation, meta = (ClampMin = "0.1", ClampMax = "6400000", UIMin = "0.1", UIMax = "6400000"))
 	float TileGenerationRadius;
 
 	/** Navigation data can be discarded when outside of a TileRemovalRadius circle around the component owner.
 	 * This is computed for all navigation invokers, so a navigation data must be outside of all navigation invokers TileRemovalRadius circles to be discarded. */
-	UPROPERTY(EditAnywhere, Category = Navigation, meta = (ClampMin = "0.1", ClampMax = "6400000", UIMin = "0.1", UIMax = "6400000"))
+	UPROPERTY(EditAnywhere, BlueprintReadOnly, Category = Navigation, meta = (ClampMin = "0.1", ClampMax = "6400000", UIMin = "0.1", UIMax = "6400000"))
 	float TileRemovalRadius;
 
 	/** restrict navigation generation to specific agents */
