Integration: MassNavMeshPathFollowTask: remove overlapping points when building corridor.

diff --git a/Engine/Plugins/AI/MassAI/Source/MassAIBehavior/Private/Tasks/MassNavMeshPathFollowTask.cpp b/Engine/Plugins/AI/MassAI/Source/MassAIBehavior/Private/Tasks/MassNavMeshPathFollowTask.cpp
--- a/Engine/Plugins/AI/MassAI/Source/MassAIBehavior/Private/Tasks/MassNavMeshPathFollowTask.cpp
+++ b/Engine/Plugins/AI/MassAI/Source/MassAIBehavior/Private/Tasks/MassNavMeshPathFollowTask.cpp
@@ -69,6 +69,8 @@ bool FMassNavMeshPathFollowTask::RequestPath(FStateTreeExecutionContext& Context
 
 	if (Result.IsSuccessful())
 	{
+		Result.Path->RemoveOverlappingPoints(FNavCorridor::OverlappingPointTolerance);
+		
 		// @todo: Investigate single point paths but for now, only move if we have more than one point.
 		if (Result.Path.Get()->GetPathPoints().Num() > 1)
 		{
diff --git a/Engine/Plugins/Runtime/NavCorridor/Source/NavCorridor/Private/NavCorridor.cpp b/Engine/Plugins/Runtime/NavCorridor/Source/NavCorridor/Private/NavCorridor.cpp
--- a/Engine/Plugins/Runtime/NavCorridor/Source/NavCorridor/Private/NavCorridor.cpp
+++ b/Engine/Plugins/Runtime/NavCorridor/Source/NavCorridor/Private/NavCorridor.cpp
@@ -1355,6 +1355,7 @@ FString FNavCorridorParams::ToString() const
 //-------------------------------------------------------
 // FNavCorridor
 //-------------------------------------------------------
+const FVector FNavCorridor::OverlappingPointTolerance = FVector(1,1,30); 
 
 void FNavCorridor::Reset()
 {
diff --git a/Engine/Plugins/Runtime/NavCorridor/Source/NavCorridor/Private/NavCorridorTestingComponent.cpp b/Engine/Plugins/Runtime/NavCorridor/Source/NavCorridor/Private/NavCorridorTestingComponent.cpp
--- a/Engine/Plugins/Runtime/NavCorridor/Source/NavCorridor/Private/NavCorridorTestingComponent.cpp
+++ b/Engine/Plugins/Runtime/NavCorridor/Source/NavCorridor/Private/NavCorridorTestingComponent.cpp
@@ -137,7 +137,8 @@ void UNavCorridorTestingComponent::UpdateTests()
 
 		const FPathFindingQuery PathQuery(this, *NavData, ActorLocation, GoalLocation, UNavigationQueryFilter::GetQueryFilter(*NavData, this, FilterClass));
 		const FSharedConstNavQueryFilter NavQueryFilter = PathQuery.QueryFilter ? PathQuery.QueryFilter : NavData->GetDefaultQueryFilter();
-		const FPathFindingResult PathResult = NavSys->FindPathSync(NavAgentProps, PathQuery);
+		FPathFindingResult PathResult = NavSys->FindPathSync(NavAgentProps, PathQuery);
+
 		Path = PathResult.Path;
 
 		EndTime = FPlatformTime::Seconds();
@@ -148,6 +149,7 @@ void UNavCorridorTestingComponent::UpdateTests()
 
 		if (PathResult.IsSuccessful())
 		{
+			PathResult.Path->RemoveOverlappingPoints(FNavCorridor::OverlappingPointTolerance);
 			Corridor.BuildFromPath(*Path, NavQueryFilter, CorridorParams);
 			Corridor.OffsetPathLocationsFromWalls(PathOffset);
 		}
diff --git a/Engine/Plugins/Runtime/NavCorridor/Source/NavCorridor/Public/NavCorridor.h b/Engine/Plugins/Runtime/NavCorridor/Source/NavCorridor/Public/NavCorridor.h
--- a/Engine/Plugins/Runtime/NavCorridor/Source/NavCorridor/Public/NavCorridor.h
+++ b/Engine/Plugins/Runtime/NavCorridor/Source/NavCorridor/Public/NavCorridor.h
@@ -155,6 +155,11 @@ struct FNavCorridor : public TSharedFromThis<FNavCorridor>
 
 	/** Portal points defining the corridor. */
 	TArray<FNavCorridorPortal> Portals;
+
+	/** Tolerance used to remove overlappint points in path before building the corridor (arbitrary value used to fix common cases).
+	 * @todo: Make corridor support overlapping points and remove this.
+	 */
+	UE_API const static FVector OverlappingPointTolerance;
 };
 
 #undef UE_API
diff --git a/Engine/Source/Runtime/NavigationSystem/Private/NavigationData.cpp b/Engine/Source/Runtime/NavigationSystem/Private/NavigationData.cpp
--- a/Engine/Source/Runtime/NavigationSystem/Private/NavigationData.cpp
+++ b/Engine/Source/Runtime/NavigationSystem/Private/NavigationData.cpp
@@ -151,6 +151,27 @@ FSupportedAreaData::FSupportedAreaData(TSubclassOf<UNavArea> NavAreaClass, int32
 	}
 }
 
+//----------------------------------------------------------------------//
+// FNavigationPath
+//----------------------------------------------------------------------//
+void FNavigationPath::RemoveOverlappingPoints(const FVector& Tolerance)
+{
+	// Remove overlapping points according to the tolerance.
+	for (int32 Index = 0; Index+1 < PathPoints.Num();)
+	{
+		const FVector Delta = PathPoints[Index+1].Location - PathPoints[Index].Location;
+		const bool bSameLocation = FMath::Abs(Delta.X) <= Tolerance.X && FMath::Abs(Delta.Y) <= Tolerance.Y && FMath::Abs(Delta.Z) <= Tolerance.Z;
+		if (bSameLocation)
+		{
+			PathPoints.RemoveAt(Index);
+		}
+		else
+		{
+			++Index;
+		}
+	}
+}
+
 //----------------------------------------------------------------------//
 // ANavigationData                                                                
 //----------------------------------------------------------------------//
diff --git a/Engine/Source/Runtime/NavigationSystem/Public/NavigationData.h b/Engine/Source/Runtime/NavigationSystem/Public/NavigationData.h
--- a/Engine/Source/Runtime/NavigationSystem/Public/NavigationData.h
+++ b/Engine/Source/Runtime/NavigationSystem/Public/NavigationData.h
@@ -264,6 +264,9 @@ struct FNavigationPath : public TSharedFromThis<FNavigationPath, ESPMode::Thread
 	  * This function will NOT reset setup variables like goal actor, filter, observer, etc */
 	NAVIGATIONSYSTEM_API virtual void ResetForRepath();
 
+	/** Remove points that are at the same location. */
+	NAVIGATIONSYSTEM_API void RemoveOverlappingPoints(const FVector& Tolerance);
+
 	NAVIGATIONSYSTEM_API virtual void DebugDraw(const ANavigationData* NavData, const FColor PathColor, class UCanvas* Canvas, const bool bPersistent, const float LifeTime, const uint32 NextPathPointIndex = 0) const;
 	
 #if ENABLE_VISUAL_LOG
