Replace some usages of FORCEINLINE with inline in AI modules.

For context about this change, see this doc: https://docs.google.com/document/d/1LjpT7NpzTecfXdFNyv5KM6sCu67827BRljNi1iaWVSE/edit?tab=t.0#heading=h.fvtz8gjukz7x

Includes modules: HTNPlanner, StateTreeEditorModule, StateTreeModule, ZoneGraph, ZoneGraphAnnotations, Gaze, AIModule, NavigationSystem

Reviewers were automatically determined. Please comment if anyone else should be added.

#rnx

diff --git a/Engine/Plugins/AI/HTNPlanner/Source/HTNPlanner/Public/HTNDomain.h b/Engine/Plugins/AI/HTNPlanner/Source/HTNPlanner/Public/HTNDomain.h
--- a/Engine/Plugins/AI/HTNPlanner/Source/HTNPlanner/Public/HTNDomain.h
+++ b/Engine/Plugins/AI/HTNPlanner/Source/HTNPlanner/Public/HTNDomain.h
@@ -99,7 +99,7 @@ struct THTNWorldStateOperation
 		return *this;
 	}
 
-	FORCEINLINE bool IsRHSAbsolute() const { return KeyRightHand == FHTNPolicy::InvalidWSKey; }
+	inline bool IsRHSAbsolute() const { return KeyRightHand == FHTNPolicy::InvalidWSKey; }
 
 	bool IsValid() const { return (Operation != FHTNPolicy::InvalidWSOperation); }
 };
@@ -184,11 +184,11 @@ struct HTNPLANNER_API FHTNDomain : public TSharedFromThis<FHTNDomain>
 	UE_NONCOPYABLE(FHTNDomain);
 
 	FHTNPolicy::FTaskID FindTaskID(const FName& TaskName) const;
-	FORCEINLINE bool IsPrimitiveTask(const FHTNPolicy::FTaskID TaskID) const
+	inline bool IsPrimitiveTask(const FHTNPolicy::FTaskID TaskID) const
 	{
 		return TaskID < FirstCompositeTaskID;
 	}
-	FORCEINLINE bool IsCompositeTask(const FHTNPolicy::FTaskID TaskID) const
+	inline bool IsCompositeTask(const FHTNPolicy::FTaskID TaskID) const
 	{
 		return !IsPrimitiveTask(TaskID);
 	}
@@ -236,7 +236,7 @@ struct HTNPLANNER_API FHTNWorldState
 	void Reinit(const uint32 NewWorldStateSize = 128);
 
 	bool CheckCondition(const FHTNCondition& Condition) const;
-	FORCEINLINE bool CheckConditions(const FHTNCondition* Conditions, const int32 ConditionsCount) const
+	inline bool CheckConditions(const FHTNCondition* Conditions, const int32 ConditionsCount) const
 	{
 		for (int32 ConditionIndex = 0; ConditionIndex < ConditionsCount; ++ConditionIndex)
 		{
@@ -249,7 +249,7 @@ struct HTNPLANNER_API FHTNWorldState
 	}
 
 	void ApplyEffect(const FHTNEffect& Effect);
-	FORCEINLINE void ApplyEffects(const FHTNEffect* Effects, const int32 EffectsCount)
+	inline void ApplyEffects(const FHTNEffect* Effects, const int32 EffectsCount)
 	{
 		for (int32 EffectIndex = 0; EffectIndex < EffectsCount; ++EffectIndex)
 		{
diff --git a/Engine/Plugins/Runtime/StateTree/Source/StateTreeEditorModule/Public/StateTreeEditorNode.h b/Engine/Plugins/Runtime/StateTree/Source/StateTreeEditorModule/Public/StateTreeEditorNode.h
--- a/Engine/Plugins/Runtime/StateTree/Source/StateTreeEditorModule/Public/StateTreeEditorNode.h
+++ b/Engine/Plugins/Runtime/StateTree/Source/StateTreeEditorModule/Public/StateTreeEditorNode.h
@@ -91,8 +91,8 @@ template <typename T>
 struct TStateTreeEditorNode : public FStateTreeEditorNode
 {
 	using NodeType = T;
-	FORCEINLINE T& GetNode() { return Node.template GetMutable<T>(); }
-	FORCEINLINE typename T::FInstanceDataType& GetInstanceData() { return Instance.template GetMutable<typename T::FInstanceDataType>(); }
+	inline T& GetNode() { return Node.template GetMutable<T>(); }
+	inline typename T::FInstanceDataType& GetInstanceData() { return Instance.template GetMutable<typename T::FInstanceDataType>(); }
 };
 
 #undef UE_API
diff --git a/Engine/Plugins/Runtime/StateTree/Source/StateTreeModule/Public/StateTreeEvents.h b/Engine/Plugins/Runtime/StateTree/Source/StateTreeModule/Public/StateTreeEvents.h
--- a/Engine/Plugins/Runtime/StateTree/Source/StateTreeModule/Public/StateTreeEvents.h
+++ b/Engine/Plugins/Runtime/StateTree/Source/StateTreeModule/Public/StateTreeEvents.h
@@ -43,7 +43,7 @@ struct FStateTreeEvent
 	{
 	}
 
-	friend FORCEINLINE uint32 GetTypeHash(const FStateTreeEvent& Event)
+	friend inline uint32 GetTypeHash(const FStateTreeEvent& Event)
 	{
 		uint32 Hash = GetTypeHash(Event.Tag);
 		
diff --git a/Engine/Plugins/Runtime/StateTree/Source/StateTreeModule/Public/StateTreeExecutionTypes.h b/Engine/Plugins/Runtime/StateTree/Source/StateTreeModule/Public/StateTreeExecutionTypes.h
--- a/Engine/Plugins/Runtime/StateTree/Source/StateTreeModule/Public/StateTreeExecutionTypes.h
+++ b/Engine/Plugins/Runtime/StateTree/Source/StateTreeModule/Public/StateTreeExecutionTypes.h
@@ -532,14 +532,14 @@ struct FStateTreeActiveStates
 	}
 
 	/** Returns a specified state in the array. */
-	FORCEINLINE const FStateTreeStateHandle& operator[](const int32 Index) const
+	inline const FStateTreeStateHandle& operator[](const int32 Index) const
 	{
 		check(Index >= 0 && Index < (int32)NumStates);
 		return States[Index];
 	}
 
 	/** Returns mutable reference to a specified state in the array. */
-	FORCEINLINE FStateTreeStateHandle& operator[](const int32 Index)
+	inline FStateTreeStateHandle& operator[](const int32 Index)
 	{
 		check(Index >= 0 && Index < (int32)NumStates);
 		return States[Index];
@@ -567,10 +567,10 @@ struct FStateTreeActiveStates
 	 * DO NOT USE DIRECTLY
 	 * STL-like iterators to enable range-based for loop support.
 	 */
-	FORCEINLINE FStateTreeStateHandle* begin() { return &States[0]; }
-	FORCEINLINE FStateTreeStateHandle* end  () { return &States[0] + Num(); }
-	FORCEINLINE const FStateTreeStateHandle* begin() const { return &States[0]; }
-	FORCEINLINE const FStateTreeStateHandle* end  () const { return &States[0] + Num(); }
+	inline FStateTreeStateHandle* begin() { return &States[0]; }
+	inline FStateTreeStateHandle* end  () { return &States[0] + Num(); }
+	inline const FStateTreeStateHandle* begin() const { return &States[0]; }
+	inline const FStateTreeStateHandle* end  () const { return &States[0] + Num(); }
 
 	UE::StateTree::FActiveStateID StateIDs[MaxStates];
 
diff --git a/Engine/Plugins/Runtime/StateTree/Source/StateTreeModule/Public/StateTreeIndexTypes.h b/Engine/Plugins/Runtime/StateTree/Source/StateTreeModule/Public/StateTreeIndexTypes.h
--- a/Engine/Plugins/Runtime/StateTree/Source/StateTreeModule/Public/StateTreeIndexTypes.h
+++ b/Engine/Plugins/Runtime/StateTree/Source/StateTreeModule/Public/StateTreeIndexTypes.h
@@ -16,7 +16,7 @@ struct FStateTreeIndex16
 	static constexpr uint16 InvalidValue = MAX_uint16;
 	static UE_API const FStateTreeIndex16 Invalid;
 
-	friend FORCEINLINE uint32 GetTypeHash(const FStateTreeIndex16 Index)
+	friend inline uint32 GetTypeHash(const FStateTreeIndex16 Index)
 	{
 		return GetTypeHash(Index.Value);
 	}
diff --git a/Engine/Plugins/Runtime/StateTree/Source/StateTreeModule/Public/StateTreeTypes.h b/Engine/Plugins/Runtime/StateTree/Source/StateTreeModule/Public/StateTreeTypes.h
--- a/Engine/Plugins/Runtime/StateTree/Source/StateTreeModule/Public/StateTreeTypes.h
+++ b/Engine/Plugins/Runtime/StateTree/Source/StateTreeModule/Public/StateTreeTypes.h
@@ -238,7 +238,7 @@ struct FStateTreeStateHandle
 		return Index >= 0 && Index < (int32)MAX_uint16;
 	}
 
-	friend FORCEINLINE uint32 GetTypeHash(const FStateTreeStateHandle& Handle)
+	friend inline uint32 GetTypeHash(const FStateTreeStateHandle& Handle)
 	{
 		return GetTypeHash(Handle.Index);
 	}
@@ -347,7 +347,7 @@ struct FStateTreeDataHandle
 		return Index >= 0 && Index < (int32)InvalidIndex;
 	}
 
-	friend FORCEINLINE uint32 GetTypeHash(const FStateTreeDataHandle& Handle)
+	friend inline uint32 GetTypeHash(const FStateTreeDataHandle& Handle)
 	{
 		uint32 Hash = GetTypeHash(Handle.Index);
 		Hash = HashCombineFast(Hash, GetTypeHash(Handle.Source));
diff --git a/Engine/Plugins/Runtime/ZoneGraph/Source/ZoneGraph/Public/ZoneGraphAStar.h b/Engine/Plugins/Runtime/ZoneGraph/Source/ZoneGraph/Public/ZoneGraphAStar.h
--- a/Engine/Plugins/Runtime/ZoneGraph/Source/ZoneGraph/Public/ZoneGraphAStar.h
+++ b/Engine/Plugins/Runtime/ZoneGraph/Source/ZoneGraph/Public/ZoneGraphAStar.h
@@ -22,7 +22,7 @@ struct FZoneGraphAStarWrapper
 	// FGraphAStar: TGraph
 	typedef int32/*lane index type*/ FNodeRef;
 
-	FORCEINLINE bool IsValidRef(const FNodeRef NodeRef) const
+	inline bool IsValidRef(const FNodeRef NodeRef) const
 	{
 		return NodeRef != INDEX_NONE;
 	}
@@ -41,7 +41,7 @@ struct FZoneGraphAStarNode : public FGraphAStarDefaultNode<FZoneGraphAStarWrappe
 	typedef FGraphAStarDefaultNode<FZoneGraphAStarWrapper> Super;
 	typedef int32/*lane index type*/ FNodeRef;
 
-	FORCEINLINE FZoneGraphAStarNode(const FNodeRef InNodeRef = INDEX_NONE, const FVector InPosition = FVector(TNumericLimits<FVector::FReal>::Max()))
+	inline FZoneGraphAStarNode(const FNodeRef InNodeRef = INDEX_NONE, const FVector InPosition = FVector(TNumericLimits<FVector::FReal>::Max()))
 		: Super(InNodeRef)
 		, Position(InPosition)
 	{}
@@ -71,7 +71,7 @@ struct FZoneGraphPathFilter
 	bool IsStart(const FZoneGraphAStarNode& Node) const { return Node.NodeRef == StartLocation.LaneHandle.Index; }
 	bool IsEnd(const FZoneGraphAStarNode& Node) const { return Node.NodeRef == EndLocation.LaneHandle.Index; }
 
-	FORCEINLINE FVector::FReal GetHeuristicScale() const { return 1.; }
+	inline FVector::FReal GetHeuristicScale() const { return 1.; }
 
 	FVector::FReal GetHeuristicCost(const FZoneGraphAStarNode& NeighbourNode, const FZoneGraphAStarNode& EndNode) const
 	{
@@ -167,14 +167,14 @@ struct FZoneGraphPathFilter
 		}
 	}
 
-	FORCEINLINE bool IsTraversalAllowed(const FNodeRef StartNodeRef, const FNodeRef& Neighbour) const
+	inline bool IsTraversalAllowed(const FNodeRef StartNodeRef, const FNodeRef& Neighbour) const
 	{
 		const FZoneGraphTagMask& LaneTagMask = ZoneStorage.Lanes[Neighbour].Tags;
 		return ZoneTagFilter.Pass(LaneTagMask);
 	}
 
-	FORCEINLINE bool WantsPartialSolution() const { return false; }
-	FORCEINLINE bool ShouldIncludeStartNodeInPath() const { return true; }
+	inline bool WantsPartialSolution() const { return false; }
+	inline bool ShouldIncludeStartNodeInPath() const { return true; }
 
 protected:
 	const FZoneGraphStorage& ZoneStorage;
diff --git a/Engine/Plugins/Runtime/ZoneGraph/Source/ZoneGraph/Public/ZoneGraphData.h b/Engine/Plugins/Runtime/ZoneGraph/Source/ZoneGraph/Public/ZoneGraphData.h
--- a/Engine/Plugins/Runtime/ZoneGraph/Source/ZoneGraph/Public/ZoneGraphData.h
+++ b/Engine/Plugins/Runtime/ZoneGraph/Source/ZoneGraph/Public/ZoneGraphData.h
@@ -32,10 +32,10 @@ public:
 #endif // WITH_EDITOR
 	//~ End UObject/AActor Interface
 
-	FORCEINLINE bool IsDrawingEnabled() const { return bEnableDrawing; }
+	inline bool IsDrawingEnabled() const { return bEnableDrawing; }
 	UE_API void UpdateDrawing() const;
 
-	FORCEINLINE bool IsRegistered() const { return bRegistered; }
+	inline bool IsRegistered() const { return bRegistered; }
 	UE_API void OnRegistered(const FZoneGraphDataHandle DataHandle);
 	UE_API void OnUnregistered();
 
diff --git a/Engine/Plugins/Runtime/ZoneGraphAnnotations/Source/ZoneGraphAnnotations/Public/ZoneGraphAnnotationTypes.h b/Engine/Plugins/Runtime/ZoneGraphAnnotations/Source/ZoneGraphAnnotations/Public/ZoneGraphAnnotationTypes.h
--- a/Engine/Plugins/Runtime/ZoneGraphAnnotations/Source/ZoneGraphAnnotations/Public/ZoneGraphAnnotationTypes.h
+++ b/Engine/Plugins/Runtime/ZoneGraphAnnotations/Source/ZoneGraphAnnotations/Public/ZoneGraphAnnotationTypes.h
@@ -44,7 +44,7 @@ private:
 	uint64 Value = MAX_uint64;
 };
 
-FORCEINLINE uint32 GetTypeHash(const FMassLaneObstacleID& Obs)
+inline uint32 GetTypeHash(const FMassLaneObstacleID& Obs)
 {
 	return ::GetTypeHash(Obs.GetValue());
 }
diff --git a/Engine/Source/Runtime/AIModule/Classes/AIController.h b/Engine/Source/Runtime/AIModule/Classes/AIController.h
--- a/Engine/Source/Runtime/AIModule/Classes/AIController.h
+++ b/Engine/Source/Runtime/AIModule/Classes/AIController.h
@@ -314,7 +314,7 @@ public:
 	UFUNCTION(BlueprintCallable, Category = "AI")
 	AIMODULE_API AActor* GetFocusActor() const;
 
-	FORCEINLINE AActor* GetFocusActorForPriority(EAIFocusPriority::Type InPriority) const { return FocusInformation.Priorities.IsValidIndex(InPriority) ? FocusInformation.Priorities[InPriority].Actor.Get() : nullptr; }
+	inline AActor* GetFocusActorForPriority(EAIFocusPriority::Type InPriority) const { return FocusInformation.Priorities.IsValidIndex(InPriority) ? FocusInformation.Priorities[InPriority].Actor.Get() : nullptr; }
 
 	/** Clears Focus, will also clear FocalPoint as a result */
 	UFUNCTION(BlueprintCallable, Category = "AI", meta = (DisplayName = "ClearFocus", ScriptName = "ClearFocus"))
@@ -411,7 +411,7 @@ public:
 	virtual AActor* GetGameplayTaskAvatar(const UGameplayTask* Task) const override { return GetPawn(); }
 	virtual uint8 GetGameplayTaskDefaultPriority() const { return FGameplayTasks::DefaultPriority - 1; }
 
-	FORCEINLINE UGameplayTasksComponent* GetGameplayTasksComponent() const { return CachedGameplayTasksComponent; }
+	inline UGameplayTasksComponent* GetGameplayTasksComponent() const { return CachedGameplayTasksComponent; }
 
 	// add empty overrides to fix linker errors if project implements a child class without adding GameplayTasks module dependency
 	virtual void OnGameplayTaskInitialized(UGameplayTask& Task) override {}
@@ -457,7 +457,7 @@ public:
 //----------------------------------------------------------------------//
 namespace FAISystem
 {
-	FORCEINLINE bool IsValidControllerAndHasValidPawn(const AController* Controller)
+	inline bool IsValidControllerAndHasValidPawn(const AController* Controller)
 	{
 		return Controller != nullptr && Controller->IsPendingKillPending() == false
 			&& Controller->GetPawn() != nullptr && Controller->GetPawn()->IsPendingKillPending() == false;
diff --git a/Engine/Source/Runtime/AIModule/Classes/AISystem.h b/Engine/Source/Runtime/AIModule/Classes/AISystem.h
--- a/Engine/Source/Runtime/AIModule/Classes/AISystem.h
+++ b/Engine/Source/Runtime/AIModule/Classes/AISystem.h
@@ -153,41 +153,41 @@ public:
 	// UAISystemBase end
 
 	/** Behavior tree manager getter */
-	FORCEINLINE UBehaviorTreeManager* GetBehaviorTreeManager() { return BehaviorTreeManager; }
+	inline UBehaviorTreeManager* GetBehaviorTreeManager() { return BehaviorTreeManager; }
 	/** Behavior tree manager const getter */
-	FORCEINLINE const UBehaviorTreeManager* GetBehaviorTreeManager() const { return BehaviorTreeManager; }
+	inline const UBehaviorTreeManager* GetBehaviorTreeManager() const { return BehaviorTreeManager; }
 
 	/** Environment Query manager getter */
-	FORCEINLINE UEnvQueryManager* GetEnvironmentQueryManager() { return EnvironmentQueryManager; }
+	inline UEnvQueryManager* GetEnvironmentQueryManager() { return EnvironmentQueryManager; }
 	/** Environment Query manager const getter */
-	FORCEINLINE const UEnvQueryManager* GetEnvironmentQueryManager() const { return EnvironmentQueryManager; }
+	inline const UEnvQueryManager* GetEnvironmentQueryManager() const { return EnvironmentQueryManager; }
 
-	FORCEINLINE UAIPerceptionSystem* GetPerceptionSystem() { return PerceptionSystem; }
-	FORCEINLINE const UAIPerceptionSystem* GetPerceptionSystem() const { return PerceptionSystem; }
+	inline UAIPerceptionSystem* GetPerceptionSystem() { return PerceptionSystem; }
+	inline const UAIPerceptionSystem* GetPerceptionSystem() const { return PerceptionSystem; }
 
-	FORCEINLINE UAIHotSpotManager* GetHotSpotManager() { return HotSpotManager; }
-	FORCEINLINE const UAIHotSpotManager* GetHotSpotManager() const { return HotSpotManager; }
+	inline UAIHotSpotManager* GetHotSpotManager() { return HotSpotManager; }
+	inline const UAIHotSpotManager* GetHotSpotManager() const { return HotSpotManager; }
 
-	FORCEINLINE UNavLocalGridManager* GetNavLocalGridManager() { return NavLocalGrids; }
-	FORCEINLINE const UNavLocalGridManager* GetNavLocalGridManager() const { return NavLocalGrids; }
+	inline UNavLocalGridManager* GetNavLocalGridManager() { return NavLocalGrids; }
+	inline const UNavLocalGridManager* GetNavLocalGridManager() const { return NavLocalGrids; }
 
-	FORCEINLINE static UAISystem* GetCurrentSafe(UWorld* World) 
+	inline static UAISystem* GetCurrentSafe(UWorld* World) 
 	{ 
 		return World != nullptr ? Cast<UAISystem>(World->GetAISystem()) : NULL;
 	}
 
-	FORCEINLINE static UAISystem* GetCurrent(UWorld& World)
+	inline static UAISystem* GetCurrent(UWorld& World)
 	{
 		return Cast<UAISystem>(World.GetAISystem());
 	}
 
-	FORCEINLINE UWorld* GetOuterWorld() const { return Cast<UWorld>(GetOuter()); }
+	inline UWorld* GetOuterWorld() const { return Cast<UWorld>(GetOuter()); }
 
 	virtual UWorld* GetWorld() const override { return GetOuterWorld(); }
 	
-	FORCEINLINE void AddReferenceFromProxyObject(UAIAsyncTaskBlueprintProxy* BlueprintProxy) { AllProxyObjects.AddUnique(BlueprintProxy); }
+	inline void AddReferenceFromProxyObject(UAIAsyncTaskBlueprintProxy* BlueprintProxy) { AllProxyObjects.AddUnique(BlueprintProxy); }
 
-	FORCEINLINE void RemoveReferenceToProxyObject(UAIAsyncTaskBlueprintProxy* BlueprintProxy) { AllProxyObjects.RemoveSwap(BlueprintProxy); }
+	inline void RemoveReferenceToProxyObject(UAIAsyncTaskBlueprintProxy* BlueprintProxy) { AllProxyObjects.RemoveSwap(BlueprintProxy); }
 
 	//----------------------------------------------------------------------//
 	// cheats
@@ -210,13 +210,13 @@ public:
 	public:
 		FBlackboardDataToComponentsIterator(FBlackboardDataToComponentsMap& BlackboardDataToComponentsMap, class UBlackboardData* BlackboardAsset);
 
-		FORCEINLINE FBlackboardDataToComponentsIterator& operator++()
+		inline FBlackboardDataToComponentsIterator& operator++()
 		{
 			++GetCurrentIteratorRef();
 			TryMoveIteratorToParentBlackboard();
 			return *this;
 		}
-		FORCEINLINE FBlackboardDataToComponentsIterator operator++(int)
+		inline FBlackboardDataToComponentsIterator operator++(int)
 		{
 			FBlackboardDataToComponentsIterator Tmp(*this);
 			++GetCurrentIteratorRef();
@@ -224,15 +224,15 @@ public:
 			return Tmp;
 		}
 
-		FORCEINLINE explicit operator bool() const { return CurrentIteratorIndex < Iterators.Num() && (bool)GetCurrentIteratorRef(); }
-		FORCEINLINE bool operator !() const { return !(bool)*this; }
+		inline explicit operator bool() const { return CurrentIteratorIndex < Iterators.Num() && (bool)GetCurrentIteratorRef(); }
+		inline bool operator !() const { return !(bool)*this; }
 
-		FORCEINLINE UBlackboardData* Key() const { return GetCurrentIteratorRef().Key().Get(); }
-		FORCEINLINE UBlackboardComponent* Value() const { return GetCurrentIteratorRef().Value().Get(); }
+		inline UBlackboardData* Key() const { return GetCurrentIteratorRef().Key().Get(); }
+		inline UBlackboardComponent* Value() const { return GetCurrentIteratorRef().Value().Get(); }
 
 	private:
-		FORCEINLINE const FBlackboardDataToComponentsMap::TConstKeyIterator& GetCurrentIteratorRef() const { return Iterators[CurrentIteratorIndex]; }
-		FORCEINLINE FBlackboardDataToComponentsMap::TConstKeyIterator& GetCurrentIteratorRef() { return Iterators[CurrentIteratorIndex]; }
+		inline const FBlackboardDataToComponentsMap::TConstKeyIterator& GetCurrentIteratorRef() const { return Iterators[CurrentIteratorIndex]; }
+		inline FBlackboardDataToComponentsMap::TConstKeyIterator& GetCurrentIteratorRef() { return Iterators[CurrentIteratorIndex]; }
 
 		void TryMoveIteratorToParentBlackboard()
 		{
diff --git a/Engine/Source/Runtime/AIModule/Classes/AITypes.h b/Engine/Source/Runtime/AIModule/Classes/AITypes.h
--- a/Engine/Source/Runtime/AIModule/Classes/AITypes.h
+++ b/Engine/Source/Runtime/AIModule/Classes/AITypes.h
@@ -33,24 +33,24 @@ namespace FAISystem
 	constexpr float InfiniteInterval = -FLT_MAX;
 	constexpr uint32 InvalidUnsignedID = uint32(INDEX_NONE);
 
-	FORCEINLINE bool IsValidLocation(const FVector& TestLocation)
+	inline bool IsValidLocation(const FVector& TestLocation)
 	{
 		return -InvalidLocation.X < TestLocation.X && TestLocation.X < InvalidLocation.X
 			&& -InvalidLocation.Y < TestLocation.Y && TestLocation.Y < InvalidLocation.Y
 			&& -InvalidLocation.Z < TestLocation.Z && TestLocation.Z < InvalidLocation.Z;
 	}
 
-	FORCEINLINE bool IsValidDirection(const FVector& TestVector)
+	inline bool IsValidDirection(const FVector& TestVector)
 	{
 		return IsValidLocation(TestVector) == true && TestVector.IsZero() == false;
 	}
 
-	FORCEINLINE bool IsValidRotation(const FRotator& TestRotation)
+	inline bool IsValidRotation(const FRotator& TestRotation)
 	{
 		return TestRotation != InvalidRotation;
 	}
 
-	FORCEINLINE bool IsValidOrientation(const FQuat& TestOrientation)
+	inline bool IsValidOrientation(const FQuat& TestOrientation)
 	{
 		return TestOrientation != InvalidOrientation;
 	}
@@ -81,12 +81,12 @@ enum class FAIDistanceType : uint8
 
 namespace FAISystem
 {
-	FORCEINLINE bool PickAIOption(EAIOptionFlag::Type Option, bool DefaultOption)
+	inline bool PickAIOption(EAIOptionFlag::Type Option, bool DefaultOption)
 	{
 		return Option == EAIOptionFlag::Default ? DefaultOption : (Option == EAIOptionFlag::Enable);
 	}
 
-	FORCEINLINE EAIOptionFlag::Type BoolToAIOption(bool Value)
+	inline EAIOptionFlag::Type BoolToAIOption(bool Value)
 	{
 		return Value ? EAIOptionFlag::Enable : EAIOptionFlag::Disable;
 	}
@@ -217,7 +217,7 @@ public:
 		return InvalidIDInstance;
 	}
 
-	friend FORCEINLINE uint32 GetTypeHash(const FAINamedID& ID)
+	friend inline uint32 GetTypeHash(const FAINamedID& ID)
 	{
 		return GetTypeHash(ID.Index);
 	}
@@ -273,7 +273,7 @@ public:
 		return InvalidIDInstance;
 	}
 
-	friend FORCEINLINE uint32 GetTypeHash(const FAIGenericID& ID)
+	friend inline uint32 GetTypeHash(const FAIGenericID& ID)
 	{
 		return GetTypeHash(ID.Index);
 	}
@@ -341,12 +341,12 @@ struct FAIResourceLock
 	/** force-clears all locks */
 	AIMODULE_API void ForceClearAllLocks();
 
-	FORCEINLINE bool IsLocked() const
+	inline bool IsLocked() const
 	{
 		return Locks != 0;
 	}
 
-	FORCEINLINE bool IsLockedBy(EAIRequestPriority::Type LockPriority) const
+	inline bool IsLockedBy(EAIRequestPriority::Type LockPriority) const
 	{
 		return (Locks & (1 << LockPriority)) != 0;
 	}
@@ -354,7 +354,7 @@ struct FAIResourceLock
 	/** Answers the question if given priority is allowed to use this resource.
 	 *	@Note that if resource is locked with priority LockPriority this function will
 	 *	return false as well */
-	FORCEINLINE bool IsAvailableFor(EAIRequestPriority::Type LockPriority) const
+	inline bool IsAvailableFor(EAIRequestPriority::Type LockPriority) const
 	{
 		for (int32 Priority = EAIRequestPriority::MAX - 1; Priority >= LockPriority; --Priority)
 		{
@@ -416,22 +416,22 @@ public:
 
 	/** returns true if given ID is identical to stored ID or any of considered
 	 *	IDs is FAIRequestID::AnyRequest*/
-	FORCEINLINE bool IsEquivalent(uint32 OtherID) const 
+	inline bool IsEquivalent(uint32 OtherID) const 
 	{
 		return OtherID != InvalidRequestID && this->IsValid() && (RequestID == OtherID || RequestID == AnyRequestID || OtherID == AnyRequestID);
 	}
 
-	FORCEINLINE bool IsEquivalent(FAIRequestID Other) const
+	inline bool IsEquivalent(FAIRequestID Other) const
 	{
 		return IsEquivalent(Other.RequestID);
 	}
 
-	FORCEINLINE bool IsValid() const
+	inline bool IsValid() const
 	{
 		return RequestID != InvalidRequestID;
 	}
 
-	FORCEINLINE uint32 GetID() const { return RequestID; }
+	inline uint32 GetID() const { return RequestID; }
 
 	void operator=(uint32 OtherID)
 	{
diff --git a/Engine/Source/Runtime/AIModule/Classes/BehaviorTree/BTAuxiliaryNode.h b/Engine/Source/Runtime/AIModule/Classes/BehaviorTree/BTAuxiliaryNode.h
--- a/Engine/Source/Runtime/AIModule/Classes/BehaviorTree/BTAuxiliaryNode.h
+++ b/Engine/Source/Runtime/AIModule/Classes/BehaviorTree/BTAuxiliaryNode.h
@@ -119,7 +119,7 @@ protected:
 	InitNotifyFlags(&NodeType::TickNode, &NodeType::OnBecomeRelevant, &NodeType::OnCeaseRelevant); \
 	} while (false)
 
-FORCEINLINE uint8 UBTAuxiliaryNode::GetChildIndex() const
+inline uint8 UBTAuxiliaryNode::GetChildIndex() const
 {
 	return ChildIndex;
 }
diff --git a/Engine/Source/Runtime/AIModule/Classes/BehaviorTree/BTCompositeNode.h b/Engine/Source/Runtime/AIModule/Classes/BehaviorTree/BTCompositeNode.h
--- a/Engine/Source/Runtime/AIModule/Classes/BehaviorTree/BTCompositeNode.h
+++ b/Engine/Source/Runtime/AIModule/Classes/BehaviorTree/BTCompositeNode.h
@@ -294,7 +294,7 @@ protected:
 //////////////////////////////////////////////////////////////////////////
 // Inlines
 
-FORCEINLINE UBTNode* UBTCompositeNode::GetChildNode(int32 Index) const
+inline UBTNode* UBTCompositeNode::GetChildNode(int32 Index) const
 {
 	return Children.IsValidIndex(Index) ?
 		(Children[Index].ChildComposite ?
@@ -303,17 +303,17 @@ FORCEINLINE UBTNode* UBTCompositeNode::GetChildNode(int32 Index) const
 		nullptr;
 }
 
-FORCEINLINE int32 UBTCompositeNode::GetChildrenNum() const
+inline int32 UBTCompositeNode::GetChildrenNum() const
 {
 	return Children.Num();
 }
 
-FORCEINLINE uint16 UBTCompositeNode::GetLastExecutionIndex() const
+inline uint16 UBTCompositeNode::GetLastExecutionIndex() const
 {
 	return LastExecutionIndex;
 }
 
-FORCEINLINE bool UBTCompositeNode::IsApplyingDecoratorScope() const
+inline bool UBTCompositeNode::IsApplyingDecoratorScope() const
 {
 	return bApplyDecoratorScope;
 }
diff --git a/Engine/Source/Runtime/AIModule/Classes/BehaviorTree/BTDecorator.h b/Engine/Source/Runtime/AIModule/Classes/BehaviorTree/BTDecorator.h
--- a/Engine/Source/Runtime/AIModule/Classes/BehaviorTree/BTDecorator.h
+++ b/Engine/Source/Runtime/AIModule/Classes/BehaviorTree/BTDecorator.h
@@ -154,12 +154,12 @@ protected:
 //////////////////////////////////////////////////////////////////////////
 // Inlines
 
-FORCEINLINE EBTFlowAbortMode::Type UBTDecorator::GetFlowAbortMode() const
+inline EBTFlowAbortMode::Type UBTDecorator::GetFlowAbortMode() const
 {
 	return FlowAbortMode;
 }
 
-FORCEINLINE bool UBTDecorator::IsInversed() const
+inline bool UBTDecorator::IsInversed() const
 {
 	return bInverseCondition;
 }
diff --git a/Engine/Source/Runtime/AIModule/Classes/BehaviorTree/BTNode.h b/Engine/Source/Runtime/AIModule/Classes/BehaviorTree/BTNode.h
--- a/Engine/Source/Runtime/AIModule/Classes/BehaviorTree/BTNode.h
+++ b/Engine/Source/Runtime/AIModule/Classes/BehaviorTree/BTNode.h
@@ -260,49 +260,49 @@ protected:
 //////////////////////////////////////////////////////////////////////////
 // Inlines
 
-FORCEINLINE UBehaviorTree* UBTNode::GetTreeAsset() const
+inline UBehaviorTree* UBTNode::GetTreeAsset() const
 {
 	return TreeAsset;
 }
 
-FORCEINLINE UBTCompositeNode* UBTNode::GetParentNode() const
+inline UBTCompositeNode* UBTNode::GetParentNode() const
 {
 	return ParentNode;
 }
 
 #if USE_BEHAVIORTREE_DEBUGGER
-FORCEINLINE UBTNode* UBTNode::GetNextNode() const
+inline UBTNode* UBTNode::GetNextNode() const
 {
 	return NextExecutionNode;
 }
 #endif
 
-FORCEINLINE uint16 UBTNode::GetExecutionIndex() const
+inline uint16 UBTNode::GetExecutionIndex() const
 {
 	return ExecutionIndex;
 }
 
-FORCEINLINE uint16 UBTNode::GetMemoryOffset() const
+inline uint16 UBTNode::GetMemoryOffset() const
 {
 	return MemoryOffset;
 }
 
-FORCEINLINE uint8 UBTNode::GetTreeDepth() const
+inline uint8 UBTNode::GetTreeDepth() const
 {
 	return TreeDepth;
 }
 
-FORCEINLINE void UBTNode::MarkInjectedNode()
+inline void UBTNode::MarkInjectedNode()
 {
 	bIsInjected = true;
 }
 
-FORCEINLINE bool UBTNode::IsInjected() const
+inline bool UBTNode::IsInjected() const
 {
 	return bIsInjected;
 }
 
-FORCEINLINE void UBTNode::ForceInstancing(bool bEnable)
+inline void UBTNode::ForceInstancing(bool bEnable)
 {
 	// allow only in not initialized trees, side effect: root node always blocked
 	check(ParentNode == NULL);
@@ -310,12 +310,12 @@ FORCEINLINE void UBTNode::ForceInstancing(bool bEnable)
 	bCreateNodeInstance = bEnable;
 }
 
-FORCEINLINE bool UBTNode::HasInstance() const
+inline bool UBTNode::HasInstance() const
 {
 	return bCreateNodeInstance;
 }
 
-FORCEINLINE bool UBTNode::IsInstanced() const
+inline bool UBTNode::IsInstanced() const
 {
 	return bIsInstanced;
 }
diff --git a/Engine/Source/Runtime/AIModule/Classes/BehaviorTree/BTTaskNode.h b/Engine/Source/Runtime/AIModule/Classes/BehaviorTree/BTTaskNode.h
--- a/Engine/Source/Runtime/AIModule/Classes/BehaviorTree/BTTaskNode.h
+++ b/Engine/Source/Runtime/AIModule/Classes/BehaviorTree/BTTaskNode.h
@@ -141,7 +141,7 @@ protected:
 	InitNotifyFlags(&NodeType::TickTask, &NodeType::OnTaskFinished); \
 	} while (false)
 
-FORCEINLINE bool UBTTaskNode::ShouldIgnoreRestartSelf() const
+inline bool UBTTaskNode::ShouldIgnoreRestartSelf() const
 {
 	return bIgnoreRestartSelf;
 }
diff --git a/Engine/Source/Runtime/AIModule/Classes/BehaviorTree/BehaviorTreeComponent.h b/Engine/Source/Runtime/AIModule/Classes/BehaviorTree/BehaviorTreeComponent.h
--- a/Engine/Source/Runtime/AIModule/Classes/BehaviorTree/BehaviorTreeComponent.h
+++ b/Engine/Source/Runtime/AIModule/Classes/BehaviorTree/BehaviorTreeComponent.h
@@ -570,32 +570,32 @@ protected:
 //////////////////////////////////////////////////////////////////////////
 // Inlines
 
-FORCEINLINE UBehaviorTree* UBehaviorTreeComponent::GetCurrentTree() const
+inline UBehaviorTree* UBehaviorTreeComponent::GetCurrentTree() const
 {
 	return InstanceStack.Num() ? KnownInstances[InstanceStack[ActiveInstanceIdx].InstanceIdIndex].TreeAsset : NULL;
 }
 
-FORCEINLINE UBehaviorTree* UBehaviorTreeComponent::GetRootTree() const
+inline UBehaviorTree* UBehaviorTreeComponent::GetRootTree() const
 {
 	return InstanceStack.Num() ? KnownInstances[InstanceStack[0].InstanceIdIndex].TreeAsset : NULL;
 }
 
-FORCEINLINE const UBTNode* UBehaviorTreeComponent::GetActiveNode() const
+inline const UBTNode* UBehaviorTreeComponent::GetActiveNode() const
 {
 	return InstanceStack.Num() ? InstanceStack[ActiveInstanceIdx].ActiveNode : NULL;
 }
 
-FORCEINLINE uint16 UBehaviorTreeComponent::GetActiveInstanceIdx() const
+inline uint16 UBehaviorTreeComponent::GetActiveInstanceIdx() const
 {
 	return ActiveInstanceIdx;
 }
 
-FORCEINLINE bool UBehaviorTreeComponent::IsRestartPending() const
+inline bool UBehaviorTreeComponent::IsRestartPending() const
 {
 	return ExecutionRequest.ExecuteNode && !ExecutionRequest.bTryNextChild;
 }
 
-FORCEINLINE bool UBehaviorTreeComponent::IsAbortPending() const
+inline bool UBehaviorTreeComponent::IsAbortPending() const
 {
 	return bWaitingForLatentAborts || PendingExecution.IsSet();
 }
diff --git a/Engine/Source/Runtime/AIModule/Classes/BehaviorTree/BehaviorTreeTypes.h b/Engine/Source/Runtime/AIModule/Classes/BehaviorTree/BehaviorTreeTypes.h
--- a/Engine/Source/Runtime/AIModule/Classes/BehaviorTree/BehaviorTreeTypes.h
+++ b/Engine/Source/Runtime/AIModule/Classes/BehaviorTree/BehaviorTreeTypes.h
@@ -260,7 +260,7 @@ struct FBehaviorTreeDebuggerInstance
 	/** runtime descriptions for each execution index */
 	TArray<FString> RuntimeDesc;
 
-	FORCEINLINE bool IsValid() const { return ActivePath.Num() != 0; }
+	inline bool IsValid() const { return ActivePath.Num() != 0; }
 };
 
 /** debugger data about current execution step */
@@ -361,9 +361,9 @@ struct FBehaviorTreeInstance
 	void DecMemoryStats() const;
 	uint32 GetAllocatedSize() const;
 #else
-	FORCEINLINE uint32 GetAllocatedSize() const { return 0; }
-	FORCEINLINE void IncMemoryStats() const {}
-	FORCEINLINE void DecMemoryStats() const {}
+	inline uint32 GetAllocatedSize() const { return 0; }
+	inline void IncMemoryStats() const {}
+	inline void DecMemoryStats() const {}
 #endif // STATS
 
 	/** initialize memory and create node instances */
@@ -462,11 +462,11 @@ struct FBTNodeIndex
 	bool TakesPriorityOver(const FBTNodeIndex& Other) const;
 	bool IsSet() const { return InstanceIndex < InvalidIndex; }
 
-	FORCEINLINE bool operator==(const FBTNodeIndex& Other) const { return Other.ExecutionIndex == ExecutionIndex && Other.InstanceIndex == InstanceIndex; }
-	FORCEINLINE bool operator!=(const FBTNodeIndex& Other) const { return !operator==(Other); }
-	FORCEINLINE friend uint32 GetTypeHash(const FBTNodeIndex& Other) { return Other.ExecutionIndex ^ Other.InstanceIndex; }
+	inline bool operator==(const FBTNodeIndex& Other) const { return Other.ExecutionIndex == ExecutionIndex && Other.InstanceIndex == InstanceIndex; }
+	inline bool operator!=(const FBTNodeIndex& Other) const { return !operator==(Other); }
+	inline friend uint32 GetTypeHash(const FBTNodeIndex& Other) { return Other.ExecutionIndex ^ Other.InstanceIndex; }
 
-	FORCEINLINE FString Describe() const { return FString::Printf(TEXT("[%d:%d]"), InstanceIndex, ExecutionIndex); }
+	inline FString Describe() const { return FString::Printf(TEXT("[%d:%d]"), InstanceIndex, ExecutionIndex); }
 };
 
 struct FBTNodeIndexRange
@@ -658,7 +658,7 @@ public:
 		
 	void AllowNoneAsValue(bool bAllow) { bNoneIsAllowedValue = bAllow; }
 
-	FORCEINLINE FBlackboard::FKey GetSelectedKeyID() const { return FBlackboard::FKey(IntCastChecked<uint16>(SelectedKeyID)); }
+	inline FBlackboard::FKey GetSelectedKeyID() const { return FBlackboard::FKey(IntCastChecked<uint16>(SelectedKeyID)); }
 
 	/** helper functions for setting basic filters */
 	AIMODULE_API void AddObjectFilter(UObject* Owner, FName PropertyName, TSubclassOf<UObject> AllowedClass);
@@ -680,10 +680,10 @@ public:
 		AddStructFilter(Owner, PropertyName, TBaseStructure<T>::Get());
 	}
 
-	FORCEINLINE bool IsNone() const { return bNoneIsAllowedValue && GetSelectedKeyID() == FBlackboard::InvalidKey; }
-	FORCEINLINE bool IsSet() const { return GetSelectedKeyID() != FBlackboard::InvalidKey; }
-	FORCEINLINE bool NeedsResolving() const { return GetSelectedKeyID() == FBlackboard::InvalidKey && SelectedKeyName.IsNone() == false; }
-	FORCEINLINE void InvalidateResolvedKey() { SelectedKeyID = FBlackboard::InvalidKey; }
+	inline bool IsNone() const { return bNoneIsAllowedValue && GetSelectedKeyID() == FBlackboard::InvalidKey; }
+	inline bool IsSet() const { return GetSelectedKeyID() != FBlackboard::InvalidKey; }
+	inline bool NeedsResolving() const { return GetSelectedKeyID() == FBlackboard::InvalidKey && SelectedKeyName.IsNone() == false; }
+	inline void InvalidateResolvedKey() { SelectedKeyID = FBlackboard::InvalidKey; }
 
 #if WITH_EDITOR
 	void PreSave(const UObject* Outer, const UBlackboardData& Blackboard, FName PropertyName);
diff --git a/Engine/Source/Runtime/AIModule/Classes/BehaviorTree/Blackboard/BlackboardKeyType.h b/Engine/Source/Runtime/AIModule/Classes/BehaviorTree/Blackboard/BlackboardKeyType.h
--- a/Engine/Source/Runtime/AIModule/Classes/BehaviorTree/Blackboard/BlackboardKeyType.h
+++ b/Engine/Source/Runtime/AIModule/Classes/BehaviorTree/Blackboard/BlackboardKeyType.h
@@ -166,22 +166,22 @@ protected:
 //////////////////////////////////////////////////////////////////////////
 // Inlines
 
-FORCEINLINE uint16 UBlackboardKeyType::GetValueSize() const
+inline uint16 UBlackboardKeyType::GetValueSize() const
 {
 	return ValueSize;
 }
 
-FORCEINLINE EBlackboardKeyOperation::Type UBlackboardKeyType::GetTestOperation() const
+inline EBlackboardKeyOperation::Type UBlackboardKeyType::GetTestOperation() const
 {
 	return SupportedOp; 
 }
 
-FORCEINLINE bool UBlackboardKeyType::HasInstance() const
+inline bool UBlackboardKeyType::HasInstance() const
 {
 	return bCreateKeyInstance;
 }
 
-FORCEINLINE bool UBlackboardKeyType::IsInstanced() const
+inline bool UBlackboardKeyType::IsInstanced() const
 {
 	return bIsInstanced;
 }
diff --git a/Engine/Source/Runtime/AIModule/Classes/BehaviorTree/BlackboardComponent.h b/Engine/Source/Runtime/AIModule/Classes/BehaviorTree/BlackboardComponent.h
--- a/Engine/Source/Runtime/AIModule/Classes/BehaviorTree/BlackboardComponent.h
+++ b/Engine/Source/Runtime/AIModule/Classes/BehaviorTree/BlackboardComponent.h
@@ -200,14 +200,14 @@ public:
 	typename TDataClass::FDataType GetValue(FBlackboard::FKey KeyID) const;
 
 	/** get pointer to raw data for given key */
-	FORCEINLINE uint8* GetKeyRawData(const FName& KeyName) { return GetKeyRawData(GetKeyID(KeyName)); }
-	FORCEINLINE uint8* GetKeyRawData(FBlackboard::FKey KeyID) { return ValueMemory.Num() && ValueOffsets.IsValidIndex(KeyID) ? (ValueMemory.GetData() + ValueOffsets[KeyID]) : NULL; }
+	inline uint8* GetKeyRawData(const FName& KeyName) { return GetKeyRawData(GetKeyID(KeyName)); }
+	inline uint8* GetKeyRawData(FBlackboard::FKey KeyID) { return ValueMemory.Num() && ValueOffsets.IsValidIndex(KeyID) ? (ValueMemory.GetData() + ValueOffsets[KeyID]) : NULL; }
 
-	FORCEINLINE const uint8* GetKeyRawData(const FName& KeyName) const { return GetKeyRawData(GetKeyID(KeyName)); }
-	FORCEINLINE const uint8* GetKeyRawData(FBlackboard::FKey KeyID) const { return ValueMemory.Num() && ValueOffsets.IsValidIndex(KeyID) ? (ValueMemory.GetData() + ValueOffsets[KeyID]) : NULL; }
+	inline const uint8* GetKeyRawData(const FName& KeyName) const { return GetKeyRawData(GetKeyID(KeyName)); }
+	inline const uint8* GetKeyRawData(FBlackboard::FKey KeyID) const { return ValueMemory.Num() && ValueOffsets.IsValidIndex(KeyID) ? (ValueMemory.GetData() + ValueOffsets[KeyID]) : NULL; }
 
 	PRAGMA_DISABLE_DEPRECATION_WARNINGS // re BlackboardAsset
-	FORCEINLINE bool IsValidKey(FBlackboard::FKey KeyID) const { check(BlackboardAsset); return BlackboardAsset->IsValidKey(KeyID); }
+	inline bool IsValidKey(FBlackboard::FKey KeyID) const { check(BlackboardAsset); return BlackboardAsset->IsValidKey(KeyID); }
 	AIMODULE_API PRAGMA_ENABLE_DEPRECATION_WARNINGS // re BlackboardAsset
 
 	/** compares blackboard's values under specified keys */
@@ -312,7 +312,7 @@ protected:
 
 PRAGMA_DISABLE_DEPRECATION_WARNINGS // re BlackboardAsset
 
-FORCEINLINE bool UBlackboardComponent::HasValidAsset() const
+inline bool UBlackboardComponent::HasValidAsset() const
 {
 	return BlackboardAsset && BlackboardAsset->IsValid();
 }
@@ -457,13 +457,13 @@ public:
 	}
 
 	template<typename T2>
-	FORCEINLINE bool SetValue(UBlackboardComponent& BBComponent, const T2 InValue)
+	inline bool SetValue(UBlackboardComponent& BBComponent, const T2 InValue)
 	{
 		return SetValue(BBComponent, FStoredType(InValue));
 	}
 
 	/** @return True is value has changed*/
-	FORCEINLINE bool SetValue(UBlackboardComponent& BBComponent, const FStoredType InValue)
+	inline bool SetValue(UBlackboardComponent& BBComponent, const FStoredType InValue)
 	{
 		ensure(BBKey != FBlackboard::InvalidKey);
 		if (InValue != CachedValue)
@@ -475,14 +475,14 @@ public:
 		return false;
 	}
 
-	FORCEINLINE const FStoredType& Get() const
+	inline const FStoredType& Get() const
 	{
 		ensure(BBKey != FBlackboard::InvalidKey);
 		return CachedValue;
 	}
 
 	template<typename T2>
-	FORCEINLINE T2 Get() const
+	inline T2 Get() const
 	{
 		ensure(BBKey != FBlackboard::InvalidKey);
 		return (T2)CachedValue;
diff --git a/Engine/Source/Runtime/AIModule/Classes/BehaviorTree/BlackboardData.h b/Engine/Source/Runtime/AIModule/Classes/BehaviorTree/BlackboardData.h
--- a/Engine/Source/Runtime/AIModule/Classes/BehaviorTree/BlackboardData.h
+++ b/Engine/Source/Runtime/AIModule/Classes/BehaviorTree/BlackboardData.h
@@ -67,7 +67,7 @@ private:
 
 public:
 
-	FORCEINLINE bool HasSynchronizedKeys() const { return bHasSynchronizedKeys; }
+	inline bool HasSynchronizedKeys() const { return bHasSynchronizedKeys; }
 
 	/** @return true if the key is instance synced */
 	AIMODULE_API bool IsKeyInstanceSynced(FBlackboard::FKey KeyID) const;
@@ -84,7 +84,7 @@ public:
 	/** @return number of defined keys, including parent chain */
 	AIMODULE_API int32 GetNumKeys() const;
 
-	FORCEINLINE FBlackboard::FKey GetFirstKeyID() const { return FirstKeyID; }
+	inline FBlackboard::FKey GetFirstKeyID() const { return FirstKeyID; }
 
 	bool IsValidKey(FBlackboard::FKey KeyID) const { return KeyID != FBlackboard::InvalidKey && (int32)KeyID < GetNumKeys(); }
 
diff --git a/Engine/Source/Runtime/AIModule/Classes/BehaviorTree/Decorators/BTDecorator_BlackboardBase.h b/Engine/Source/Runtime/AIModule/Classes/BehaviorTree/Decorators/BTDecorator_BlackboardBase.h
--- a/Engine/Source/Runtime/AIModule/Classes/BehaviorTree/Decorators/BTDecorator_BlackboardBase.h
+++ b/Engine/Source/Runtime/AIModule/Classes/BehaviorTree/Decorators/BTDecorator_BlackboardBase.h
@@ -47,7 +47,7 @@ protected:
 //////////////////////////////////////////////////////////////////////////
 // Inlines
 
-FORCEINLINE FName UBTDecorator_BlackboardBase::GetSelectedBlackboardKey() const
+inline FName UBTDecorator_BlackboardBase::GetSelectedBlackboardKey() const
 {
 	return BlackboardKey.SelectedKeyName;
 }
diff --git a/Engine/Source/Runtime/AIModule/Classes/BehaviorTree/Decorators/BTDecorator_BlueprintBase.h b/Engine/Source/Runtime/AIModule/Classes/BehaviorTree/Decorators/BTDecorator_BlueprintBase.h
--- a/Engine/Source/Runtime/AIModule/Classes/BehaviorTree/Decorators/BTDecorator_BlueprintBase.h
+++ b/Engine/Source/Runtime/AIModule/Classes/BehaviorTree/Decorators/BTDecorator_BlueprintBase.h
@@ -221,7 +221,7 @@ protected:
 	UFUNCTION(BlueprintCallable, Category="AI|BehaviorTree")
 	AIMODULE_API bool IsDecoratorObserverActive() const;
 
-	FORCEINLINE bool GetNeedsTickForConditionChecking() const { return PerformConditionCheckImplementations != 0 && (bIsObservingBB == false || bCheckConditionOnlyBlackBoardChanges == false); }
+	inline bool GetNeedsTickForConditionChecking() const { return PerformConditionCheckImplementations != 0 && (bIsObservingBB == false || bCheckConditionOnlyBlackBoardChanges == false); }
 
 	/** Handle abort request calls from the specified abort type */
 	AIMODULE_API void RequestAbort(UBehaviorTreeComponent& OwnerComp, const EAbortType Type);
diff --git a/Engine/Source/Runtime/AIModule/Classes/BehaviorTree/Services/BTService_BlackboardBase.h b/Engine/Source/Runtime/AIModule/Classes/BehaviorTree/Services/BTService_BlackboardBase.h
--- a/Engine/Source/Runtime/AIModule/Classes/BehaviorTree/Services/BTService_BlackboardBase.h
+++ b/Engine/Source/Runtime/AIModule/Classes/BehaviorTree/Services/BTService_BlackboardBase.h
@@ -33,7 +33,7 @@ protected:
 //////////////////////////////////////////////////////////////////////////
 // Inlines
 
-FORCEINLINE FName UBTService_BlackboardBase::GetSelectedBlackboardKey() const
+inline FName UBTService_BlackboardBase::GetSelectedBlackboardKey() const
 {
 	return BlackboardKey.SelectedKeyName;
 }
diff --git a/Engine/Source/Runtime/AIModule/Classes/BehaviorTree/Tasks/BTTask_BlackboardBase.h b/Engine/Source/Runtime/AIModule/Classes/BehaviorTree/Tasks/BTTask_BlackboardBase.h
--- a/Engine/Source/Runtime/AIModule/Classes/BehaviorTree/Tasks/BTTask_BlackboardBase.h
+++ b/Engine/Source/Runtime/AIModule/Classes/BehaviorTree/Tasks/BTTask_BlackboardBase.h
@@ -34,7 +34,7 @@ protected:
 //////////////////////////////////////////////////////////////////////////
 // Inlines
 
-FORCEINLINE FName UBTTask_BlackboardBase::GetSelectedBlackboardKey() const
+inline FName UBTTask_BlackboardBase::GetSelectedBlackboardKey() const
 {
 	return BlackboardKey.SelectedKeyName;
 }
diff --git a/Engine/Source/Runtime/AIModule/Classes/BehaviorTree/Tasks/BTTask_RunBehavior.h b/Engine/Source/Runtime/AIModule/Classes/BehaviorTree/Tasks/BTTask_RunBehavior.h
--- a/Engine/Source/Runtime/AIModule/Classes/BehaviorTree/Tasks/BTTask_RunBehavior.h
+++ b/Engine/Source/Runtime/AIModule/Classes/BehaviorTree/Tasks/BTTask_RunBehavior.h
@@ -51,12 +51,12 @@ protected:
 //////////////////////////////////////////////////////////////////////////
 // Inlines
 
-FORCEINLINE UBehaviorTree* UBTTask_RunBehavior::GetSubtreeAsset() const
+inline UBehaviorTree* UBTTask_RunBehavior::GetSubtreeAsset() const
 {
 	return BehaviorAsset;
 }
 
-FORCEINLINE int32 UBTTask_RunBehavior::GetInjectedNodesCount() const
+inline int32 UBTTask_RunBehavior::GetInjectedNodesCount() const
 {
 	return BehaviorAsset ? BehaviorAsset->RootDecorators.Num() : 0;
 }
diff --git a/Engine/Source/Runtime/AIModule/Classes/BehaviorTree/Tasks/BTTask_RunBehaviorDynamic.h b/Engine/Source/Runtime/AIModule/Classes/BehaviorTree/Tasks/BTTask_RunBehaviorDynamic.h
--- a/Engine/Source/Runtime/AIModule/Classes/BehaviorTree/Tasks/BTTask_RunBehaviorDynamic.h
+++ b/Engine/Source/Runtime/AIModule/Classes/BehaviorTree/Tasks/BTTask_RunBehaviorDynamic.h
@@ -64,17 +64,17 @@ protected:
 //////////////////////////////////////////////////////////////////////////
 // Inlines
 
-FORCEINLINE bool UBTTask_RunBehaviorDynamic::HasMatchingTag(const FGameplayTag& Tag) const
+inline bool UBTTask_RunBehaviorDynamic::HasMatchingTag(const FGameplayTag& Tag) const
 {
 	return InjectionTag == Tag;
 }
 
-FORCEINLINE UBehaviorTree* UBTTask_RunBehaviorDynamic::GetDefaultBehaviorAsset() const
+inline UBehaviorTree* UBTTask_RunBehaviorDynamic::GetDefaultBehaviorAsset() const
 {
 	return DefaultBehaviorAsset;
 }
 
-FORCEINLINE const FGameplayTag& UBTTask_RunBehaviorDynamic::GetInjectionTag() const
+inline const FGameplayTag& UBTTask_RunBehaviorDynamic::GetInjectionTag() const
 {
 	return InjectionTag;
 }
diff --git a/Engine/Source/Runtime/AIModule/Classes/BehaviorTree/Tasks/BTTask_WaitBlackboardTime.h b/Engine/Source/Runtime/AIModule/Classes/BehaviorTree/Tasks/BTTask_WaitBlackboardTime.h
--- a/Engine/Source/Runtime/AIModule/Classes/BehaviorTree/Tasks/BTTask_WaitBlackboardTime.h
+++ b/Engine/Source/Runtime/AIModule/Classes/BehaviorTree/Tasks/BTTask_WaitBlackboardTime.h
@@ -39,7 +39,7 @@ protected:
 //////////////////////////////////////////////////////////////////////////
 // Inlines
 
-FORCEINLINE FName UBTTask_WaitBlackboardTime::GetSelectedBlackboardKey() const
+inline FName UBTTask_WaitBlackboardTime::GetSelectedBlackboardKey() const
 {
 	return BlackboardKey.SelectedKeyName;
 }
diff --git a/Engine/Source/Runtime/AIModule/Classes/BrainComponent.h b/Engine/Source/Runtime/AIModule/Classes/BrainComponent.h
--- a/Engine/Source/Runtime/AIModule/Classes/BrainComponent.h
+++ b/Engine/Source/Runtime/AIModule/Classes/BrainComponent.h
@@ -86,9 +86,9 @@ public:
 	AIMODULE_API void OnMessage(const FAIMessage& Message);
 	AIMODULE_API FString DescribeObservedMessage() const;
 	
-	FORCEINLINE FName GetObservedMessageType() const { return MessageType; }
-	FORCEINLINE FAIRequestID GetObservedMessageID() const { return MessageID; }
-	FORCEINLINE bool IsObservingMessageID() const { return bFilterByID; }
+	inline FName GetObservedMessageType() const { return MessageType; }
+	inline FAIRequestID GetObservedMessageID() const { return MessageID; }
+	inline bool IsObservingMessageID() const { return bFilterByID; }
 
 private:
 
@@ -197,7 +197,7 @@ public:
 	/** @return blackboard used with this component */
 	AIMODULE_API const UBlackboardComponent* GetBlackboardComponent() const;
 
-	FORCEINLINE AAIController* GetAIOwner() const { return AIOwner; }
+	inline AAIController* GetAIOwner() const { return AIOwner; }
 
 protected:
 
@@ -224,12 +224,12 @@ public:
 //////////////////////////////////////////////////////////////////////////
 // Inlines
 
-FORCEINLINE UBlackboardComponent* UBrainComponent::GetBlackboardComponent()
+inline UBlackboardComponent* UBrainComponent::GetBlackboardComponent()
 {
 	return BlackboardComp;
 }
 
-FORCEINLINE const UBlackboardComponent* UBrainComponent::GetBlackboardComponent() const
+inline const UBlackboardComponent* UBrainComponent::GetBlackboardComponent() const
 {
 	return BlackboardComp;
 }
diff --git a/Engine/Source/Runtime/AIModule/Classes/DataProviders/AIDataProvider.h b/Engine/Source/Runtime/AIModule/Classes/DataProviders/AIDataProvider.h
--- a/Engine/Source/Runtime/AIModule/Classes/DataProviders/AIDataProvider.h
+++ b/Engine/Source/Runtime/AIModule/Classes/DataProviders/AIDataProvider.h
@@ -59,7 +59,7 @@ public:
 	/** bind data in provider and cache property for faster access */
 	AIMODULE_API void BindData(const UObject* Owner, int32 RequestId) const;
 
-	FORCEINLINE bool IsDynamic() const { return DataBinding != nullptr; }
+	inline bool IsDynamic() const { return DataBinding != nullptr; }
 
 	FAIDataProviderValue() :
 		CachedProperty(nullptr),
diff --git a/Engine/Source/Runtime/AIModule/Classes/EnvironmentQuery/EnvQueryInstanceBlueprintWrapper.h b/Engine/Source/Runtime/AIModule/Classes/EnvironmentQuery/EnvQueryInstanceBlueprintWrapper.h
--- a/Engine/Source/Runtime/AIModule/Classes/EnvironmentQuery/EnvQueryInstanceBlueprintWrapper.h
+++ b/Engine/Source/Runtime/AIModule/Classes/EnvironmentQuery/EnvQueryInstanceBlueprintWrapper.h
@@ -84,7 +84,7 @@ public:
 	AIMODULE_API void SetInstigator(const UObject* Object);
 
 	/** Used in UMassEQSBlueprintLibrary::GetEnviromentQueryResultAsEntityInfo */
-	FORCEINLINE EEnvQueryRunMode::Type GetRunMode() const { return RunMode; };
+	inline EEnvQueryRunMode::Type GetRunMode() const { return RunMode; };
 
 	FEQSQueryDoneSignature& GetOnQueryFinishedEvent() { return OnQueryFinishedEvent; }
 
diff --git a/Engine/Source/Runtime/AIModule/Classes/EnvironmentQuery/EnvQueryManager.h b/Engine/Source/Runtime/AIModule/Classes/EnvironmentQuery/EnvQueryManager.h
--- a/Engine/Source/Runtime/AIModule/Classes/EnvironmentQuery/EnvQueryManager.h
+++ b/Engine/Source/Runtime/AIModule/Classes/EnvironmentQuery/EnvQueryManager.h
@@ -84,22 +84,22 @@ struct FEnvQueryRequest
 	AIMODULE_API FEnvQueryRequest& SetDynamicParam(const FAIDynamicParam& Param, const UBlackboardComponent* BlackboardComponent = nullptr);
 
 	// set named params
-	FORCEINLINE FEnvQueryRequest& SetFloatParam(FName ParamName, float Value) { NamedParams.Add(ParamName, Value); return *this; }
-	FORCEINLINE FEnvQueryRequest& SetIntParam(FName ParamName, int32 Value) { NamedParams.Add(ParamName, *((float*)&Value)); return *this; }
-	FORCEINLINE FEnvQueryRequest& SetBoolParam(FName ParamName, bool Value) { NamedParams.Add(ParamName, Value ? 1.0f : -1.0f); return *this; }
-	FORCEINLINE FEnvQueryRequest& SetNamedParam(const FEnvNamedValue& ParamData) { NamedParams.Add(ParamData.ParamName, ParamData.Value); return *this; }
+	inline FEnvQueryRequest& SetFloatParam(FName ParamName, float Value) { NamedParams.Add(ParamName, Value); return *this; }
+	inline FEnvQueryRequest& SetIntParam(FName ParamName, int32 Value) { NamedParams.Add(ParamName, *((float*)&Value)); return *this; }
+	inline FEnvQueryRequest& SetBoolParam(FName ParamName, bool Value) { NamedParams.Add(ParamName, Value ? 1.0f : -1.0f); return *this; }
+	inline FEnvQueryRequest& SetNamedParam(const FEnvNamedValue& ParamData) { NamedParams.Add(ParamData.ParamName, ParamData.Value); return *this; }
 	AIMODULE_API FEnvQueryRequest& SetNamedParams(const TArray<FEnvNamedValue>& Params);
 
 	// set world (for accessing query manager) when owner can't provide it
-	FORCEINLINE FEnvQueryRequest& SetWorldOverride(UWorld* InWorld) { World = InWorld; return *this; }
+	inline FEnvQueryRequest& SetWorldOverride(UWorld* InWorld) { World = InWorld; return *this; }
 
 	template< class UserClass >	
-	FORCEINLINE int32 Execute(EEnvQueryRunMode::Type Mode, UserClass* InObj, typename FQueryFinishedSignature::TMethodPtr< UserClass > InMethod)
+	inline int32 Execute(EEnvQueryRunMode::Type Mode, UserClass* InObj, typename FQueryFinishedSignature::TMethodPtr< UserClass > InMethod)
 	{
 		return Execute(Mode, FQueryFinishedSignature::CreateUObject(InObj, InMethod));
 	}
 	template< class UserClass >	
-	FORCEINLINE int32 Execute(EEnvQueryRunMode::Type Mode, UserClass* InObj, typename FQueryFinishedSignature::TConstMethodPtr< UserClass > InMethod)
+	inline int32 Execute(EEnvQueryRunMode::Type Mode, UserClass* InObj, typename FQueryFinishedSignature::TConstMethodPtr< UserClass > InMethod)
 	{
 		return Execute(Mode, FQueryFinishedSignature::CreateUObject(InObj, InMethod));
 	}
@@ -196,7 +196,7 @@ protected:
 
 FArchive& operator<<(FArchive& Ar, FEQSDebugger::FStatsInfo& Data);
 
-FORCEINLINE bool operator== (const FEQSDebugger::FEnvQueryInfo & Left, const FEQSDebugger::FEnvQueryInfo & Right)
+inline bool operator== (const FEQSDebugger::FEnvQueryInfo & Left, const FEQSDebugger::FEnvQueryInfo & Right)
 {
 	return Left.Instance == Right.Instance;
 }
diff --git a/Engine/Source/Runtime/AIModule/Classes/EnvironmentQuery/EnvQueryNode.h b/Engine/Source/Runtime/AIModule/Classes/EnvironmentQuery/EnvQueryNode.h
--- a/Engine/Source/Runtime/AIModule/Classes/EnvironmentQuery/EnvQueryNode.h
+++ b/Engine/Source/Runtime/AIModule/Classes/EnvironmentQuery/EnvQueryNode.h
@@ -29,7 +29,7 @@ class UEnvQueryNode : public UObject
 	 * Returns false by default, unless overridden.
 	 * If overridden, will return whether or not this Node is currently being processed asynchronously. 
 	 */
-	virtual FORCEINLINE bool IsCurrentlyRunningAsync() const { return false; }
+	virtual inline bool IsCurrentlyRunningAsync() const { return false; }
 
 #if WITH_EDITOR
 	AIMODULE_API virtual void PostEditChangeProperty(FPropertyChangedEvent& PropertyChangedEvent) override;
diff --git a/Engine/Source/Runtime/AIModule/Classes/EnvironmentQuery/EnvQueryTest.h b/Engine/Source/Runtime/AIModule/Classes/EnvironmentQuery/EnvQueryTest.h
--- a/Engine/Source/Runtime/AIModule/Classes/EnvironmentQuery/EnvQueryTest.h
+++ b/Engine/Source/Runtime/AIModule/Classes/EnvironmentQuery/EnvQueryTest.h
@@ -151,9 +151,9 @@ class UEnvQueryTest : public UEnvQueryNode
 #endif
 
 	AIMODULE_API void SetWorkOnFloatValues(bool bWorkOnFloats);
-	FORCEINLINE bool GetWorkOnFloatValues() const { return bWorkOnFloatValues; }
+	inline bool GetWorkOnFloatValues() const { return bWorkOnFloatValues; }
 
-	FORCEINLINE bool CanRunAsFinalCondition() const 
+	inline bool CanRunAsFinalCondition() const 
 	{ 
 		return (TestPurpose != EEnvTestPurpose::Score)							// We are filtering and...
 				&& ((TestPurpose == EEnvTestPurpose::Filter)					// Either we are NOT scoring at ALL or...
@@ -181,7 +181,7 @@ public:
 	AIMODULE_API FVector GetItemLocation(FEnvQueryInstance& QueryInstance, int32 ItemIndex) const;
 
 	/** helper: get location of item */
-	FORCEINLINE FVector GetItemLocation(FEnvQueryInstance& QueryInstance, const FEnvQueryInstance::ItemIterator& Iterator) const
+	inline FVector GetItemLocation(FEnvQueryInstance& QueryInstance, const FEnvQueryInstance::ItemIterator& Iterator) const
 	{
 		return GetItemLocation(QueryInstance, Iterator.GetIndex());
 	}
@@ -190,7 +190,7 @@ public:
 	AIMODULE_API FRotator GetItemRotation(FEnvQueryInstance& QueryInstance, int32 ItemIndex) const;
 
 	/** helper: get location of item */
-	FORCEINLINE FRotator GetItemRotation(FEnvQueryInstance& QueryInstance, const FEnvQueryInstance::ItemIterator& Iterator) const
+	inline FRotator GetItemRotation(FEnvQueryInstance& QueryInstance, const FEnvQueryInstance::ItemIterator& Iterator) const
 	{
 		return GetItemRotation(QueryInstance, Iterator.GetIndex());
 	}
@@ -199,7 +199,7 @@ public:
 	AIMODULE_API AActor* GetItemActor(FEnvQueryInstance& QueryInstance, int32 ItemIndex) const;
 		
 	/** helper: get actor from item */
-	FORCEINLINE AActor* GetItemActor(FEnvQueryInstance& QueryInstance, const FEnvQueryInstance::ItemIterator& Iterator) const
+	inline AActor* GetItemActor(FEnvQueryInstance& QueryInstance, const FEnvQueryInstance::ItemIterator& Iterator) const
 	{
 		return GetItemActor(QueryInstance, Iterator.GetIndex());
 	}
@@ -207,8 +207,8 @@ public:
 	/** normalize scores in range */
 	AIMODULE_API void NormalizeItemScores(FEnvQueryInstance& QueryInstance);
 
-	FORCEINLINE bool IsScoring() const { return (TestPurpose != EEnvTestPurpose::Filter); } 
-	FORCEINLINE bool IsFiltering() const { return (TestPurpose != EEnvTestPurpose::Score); }
+	inline bool IsScoring() const { return (TestPurpose != EEnvTestPurpose::Filter); } 
+	inline bool IsFiltering() const { return (TestPurpose != EEnvTestPurpose::Score); }
 
 	AIMODULE_API FText DescribeFloatTestParams() const;
 	AIMODULE_API FText DescribeBoolTestParams(const FString& ConditionDesc) const;
@@ -229,7 +229,7 @@ public:
 //////////////////////////////////////////////////////////////////////////
 // Inlines
 
-FORCEINLINE bool UEnvQueryTest::IsSupportedItem(TSubclassOf<UEnvQueryItemType> ItemType) const
+inline bool UEnvQueryTest::IsSupportedItem(TSubclassOf<UEnvQueryItemType> ItemType) const
 {
 	return ItemType && (ItemType == ValidItemType || ItemType->IsChildOf(ValidItemType));
 };
diff --git a/Engine/Source/Runtime/AIModule/Classes/EnvironmentQuery/EnvQueryTypes.h b/Engine/Source/Runtime/AIModule/Classes/EnvironmentQuery/EnvQueryTypes.h
--- a/Engine/Source/Runtime/AIModule/Classes/EnvironmentQuery/EnvQueryTypes.h
+++ b/Engine/Source/Runtime/AIModule/Classes/EnvironmentQuery/EnvQueryTypes.h
@@ -493,8 +493,8 @@ struct FEnvQueryItem
 	/** has this item been discarded? */
 	uint32 bIsDiscarded:1;
 
-	FORCEINLINE bool IsValid() const { return DataOffset >= 0 && !bIsDiscarded; }
-	FORCEINLINE void Discard() { bIsDiscarded = true; }
+	inline bool IsValid() const { return DataOffset >= 0 && !bIsDiscarded; }
+	inline void Discard() { bIsDiscarded = true; }
 
 	bool operator<(const FEnvQueryItem& Other) const
 	{
@@ -545,7 +545,7 @@ public:
 	/** instance owner. Mind that it doesn't have to be the query's "Querier". This is just the object that is responsible for this query instance. */
 	TWeakObjectPtr<UObject> Owner;
 
-	FORCEINLINE float GetItemScore(int32 Index) const { return Items.IsValidIndex(Index) ? Items[Index].Score : 0.0f; }
+	inline float GetItemScore(int32 Index) const { return Items.IsValidIndex(Index) ? Items[Index].Score : 0.0f; }
 
 	/** item accessors for basic types */
 	AIMODULE_API AActor* GetItemAsActor(int32 Index) const;
@@ -574,19 +574,19 @@ public:
 	FEnvQueryResult() : ItemType(NULL), Status(EEnvQueryStatus::Processing), OptionIndex(0), QueryID(0) {}
 	FEnvQueryResult(const EEnvQueryStatus::Type& InStatus) : ItemType(NULL), Status(InStatus), OptionIndex(0), QueryID(0) {}
 
-	FORCEINLINE bool IsFinished() const { return Status != EEnvQueryStatus::Processing; }
-	FORCEINLINE bool IsAborted() const { return Status == EEnvQueryStatus::Aborted; }
-	FORCEINLINE bool IsSuccessful() const { return Status == EEnvQueryStatus::Success; }
-	FORCEINLINE void MarkAsMissingParam() { Status = EEnvQueryStatus::MissingParam; }
-	FORCEINLINE void MarkAsAborted() { Status = EEnvQueryStatus::Aborted; }
-	FORCEINLINE void MarkAsFailed() { Status = EEnvQueryStatus::Failed; }
-	FORCEINLINE void MarkAsFinishedWithoutIssues() { Status = EEnvQueryStatus::Success; }
-	FORCEINLINE void MarkAsOwnerLost() { Status = EEnvQueryStatus::OwnerLost; }
+	inline bool IsFinished() const { return Status != EEnvQueryStatus::Processing; }
+	inline bool IsAborted() const { return Status == EEnvQueryStatus::Aborted; }
+	inline bool IsSuccessful() const { return Status == EEnvQueryStatus::Success; }
+	inline void MarkAsMissingParam() { Status = EEnvQueryStatus::MissingParam; }
+	inline void MarkAsAborted() { Status = EEnvQueryStatus::Aborted; }
+	inline void MarkAsFailed() { Status = EEnvQueryStatus::Failed; }
+	inline void MarkAsFinishedWithoutIssues() { Status = EEnvQueryStatus::Success; }
+	inline void MarkAsOwnerLost() { Status = EEnvQueryStatus::OwnerLost; }
 
-	FORCEINLINE EEnvQueryStatus::Type GetRawStatus() const { return Status; }
+	inline EEnvQueryStatus::Type GetRawStatus() const { return Status; }
 
 	UE_DEPRECATED(5.0, "FEnvQueryResult::IsSuccsessful is deprecated. Use FEnvQueryResult::IsSuccessful instead.")
-	FORCEINLINE bool IsSuccsessful() const { return IsSuccessful(); }
+	inline bool IsSuccsessful() const { return IsSuccessful(); }
 };
 
 
@@ -627,7 +627,7 @@ struct FEnvQueryItemDetails
 #endif
 	}
 
-	FORCEINLINE uint32 GetAllocatedSize() const
+	inline uint32 GetAllocatedSize() const
 	{
 		return sizeof(*this) +
 #if USE_EQS_DEBUGGER
@@ -652,7 +652,7 @@ struct FEnvQueryContextData
 		: NumValues(0)
 	{}
 
-	FORCEINLINE uint32 GetAllocatedSize() const { return sizeof(*this) + RawData.GetAllocatedSize(); }
+	inline uint32 GetAllocatedSize() const { return sizeof(*this) + RawData.GetAllocatedSize(); }
 };
 
 struct FEnvQueryOptionInstance
@@ -672,7 +672,7 @@ struct FEnvQueryOptionInstance
 	/** if set, items will be shuffled after tests */
 	uint32 bHasNavLocations : 1;
 
-	FORCEINLINE uint32 GetAllocatedSize() const { return sizeof(*this) + Tests.GetAllocatedSize(); }
+	inline uint32 GetAllocatedSize() const { return sizeof(*this) + Tests.GetAllocatedSize(); }
 };
 
 #if NO_LOGGING
@@ -961,13 +961,13 @@ public:
 	AIMODULE_API void StripRedundantData();
 
 #if STATS
-	FORCEINLINE void IncStats()
+	inline void IncStats()
 	{
 		INC_MEMORY_STAT_BY(STAT_AI_EQS_InstanceMemory, GetAllocatedSize());
 		INC_DWORD_STAT_BY(STAT_AI_EQS_NumItems, Items.Num());
 	}
 
-	FORCEINLINE void DecStats()
+	inline void DecStats()
 	{
 		DEC_MEMORY_STAT_BY(STAT_AI_EQS_InstanceMemory, GetAllocatedSize()); 
 		DEC_DWORD_STAT_BY(STAT_AI_EQS_NumItems, Items.Num());
@@ -976,10 +976,10 @@ public:
 	AIMODULE_API uint32 GetAllocatedSize() const;
 	AIMODULE_API uint32 GetContextAllocatedSize() const;
 #else
-	FORCEINLINE uint32 GetAllocatedSize() const { return 0; }
-	FORCEINLINE uint32 GetContextAllocatedSize() const { return 0; }
-	FORCEINLINE void IncStats() {}
-	FORCEINLINE void DecStats() {}
+	inline uint32 GetAllocatedSize() const { return 0; }
+	inline uint32 GetContextAllocatedSize() const { return 0; }
+	inline void IncStats() {}
+	inline void DecStats() {}
 #endif // STATS
 
 #if !NO_LOGGING
@@ -1037,7 +1037,7 @@ public:
 			return CurrentItem;
 		}
 
-		FORCEINLINE explicit operator bool() const
+		inline explicit operator bool() const
 		{
 			return CurrentItem < Instance.Items.Num();
 		}
@@ -1186,7 +1186,7 @@ public:
 			return *this;
 		}
 
-		FORCEINLINE explicit operator bool() const
+		inline explicit operator bool() const
 		{
 			return CurrentItem < Instance.Items.Num() && !Instance.bFoundSingleResult && (Deadline < 0 || FPlatformTime::Seconds() < Deadline);
 		}
@@ -1225,7 +1225,7 @@ public:
 		AIMODULE_API void HandleFailedTestResult();
 		AIMODULE_API void StoreTestResult();
 
-		FORCEINLINE void SetScoreInternal(float Score)
+		inline void SetScoreInternal(float Score)
 		{
 			switch (CachedScoreOp)
 			{
@@ -1253,7 +1253,7 @@ public:
 			}
 		}
 
-		FORCEINLINE void CheckItemPassed()
+		inline void CheckItemPassed()
 		{
 			if (!bForced)
 			{
diff --git a/Engine/Source/Runtime/AIModule/Classes/EnvironmentQuery/Generators/EnvQueryGenerator_ProjectedPoints.h b/Engine/Source/Runtime/AIModule/Classes/EnvironmentQuery/Generators/EnvQueryGenerator_ProjectedPoints.h
--- a/Engine/Source/Runtime/AIModule/Classes/EnvironmentQuery/Generators/EnvQueryGenerator_ProjectedPoints.h
+++ b/Engine/Source/Runtime/AIModule/Classes/EnvironmentQuery/Generators/EnvQueryGenerator_ProjectedPoints.h
@@ -28,7 +28,7 @@ class UEnvQueryGenerator_ProjectedPoints : public UEnvQueryGenerator
 
 		FSortByHeight(const FVector& OriginalPt) : OriginalZ(OriginalPt.Z) {}
 
-		FORCEINLINE bool operator()(const FNavLocation& A, const FNavLocation& B) const
+		inline bool operator()(const FNavLocation& A, const FNavLocation& B) const
 		{
 			return FMath::Abs(A.Location.Z - OriginalZ) < FMath::Abs(B.Location.Z - OriginalZ);
 		}
diff --git a/Engine/Source/Runtime/AIModule/Classes/EnvironmentQuery/Items/EnvQueryItemType.h b/Engine/Source/Runtime/AIModule/Classes/EnvironmentQuery/Items/EnvQueryItemType.h
--- a/Engine/Source/Runtime/AIModule/Classes/EnvironmentQuery/Items/EnvQueryItemType.h
+++ b/Engine/Source/Runtime/AIModule/Classes/EnvironmentQuery/Items/EnvQueryItemType.h
@@ -25,7 +25,7 @@ public:
 	AIMODULE_API UEnvQueryItemType(const FObjectInitializer& ObjectInitializer = FObjectInitializer::Get());
 
 	/** get ValueSize */
-	FORCEINLINE uint16 GetValueSize() const { return ValueSize; }
+	inline uint16 GetValueSize() const { return ValueSize; }
 
 	/** add filters for blackboard key selector */
 	AIMODULE_API virtual void AddBlackboardFilters(FBlackboardKeySelector& KeySelector, UObject* FilterOwner) const;
diff --git a/Engine/Source/Runtime/AIModule/Classes/GenericTeamAgentInterface.h b/Engine/Source/Runtime/AIModule/Classes/GenericTeamAgentInterface.h
--- a/Engine/Source/Runtime/AIModule/Classes/GenericTeamAgentInterface.h
+++ b/Engine/Source/Runtime/AIModule/Classes/GenericTeamAgentInterface.h
@@ -40,9 +40,9 @@ public:
 	FGenericTeamId(uint8 InTeamID = NoTeamId) : TeamID(InTeamID)
 	{}
 
-	FORCEINLINE operator uint8() const { return TeamID; }
+	inline operator uint8() const { return TeamID; }
 
-	FORCEINLINE uint8 GetId() const { return TeamID; }
+	inline uint8 GetId() const { return TeamID; }
 	//FORCEINLINE void SetId(uint8 NewId) { TeamID = NewId; }
 	
 	static AIMODULE_API FGenericTeamId GetTeamIdentifier(const AActor* TeamMember);
@@ -66,7 +66,7 @@ protected:
 public:
 	static AIMODULE_API const FGenericTeamId NoTeam;
 
-	friend FORCEINLINE uint32 GetTypeHash(const FGenericTeamId Value)
+	friend inline uint32 GetTypeHash(const FGenericTeamId Value)
 	{
 		return Value.GetId();
 	}
diff --git a/Engine/Source/Runtime/AIModule/Classes/Navigation/CrowdFollowingComponent.h b/Engine/Source/Runtime/AIModule/Classes/Navigation/CrowdFollowingComponent.h
--- a/Engine/Source/Runtime/AIModule/Classes/Navigation/CrowdFollowingComponent.h
+++ b/Engine/Source/Runtime/AIModule/Classes/Navigation/CrowdFollowingComponent.h
@@ -104,34 +104,34 @@ class UCrowdFollowingComponent : public UPathFollowingComponent, public ICrowdAg
 	AIMODULE_API void SetGroupsToAvoid(int32 GroupFlags, bool bUpdateAgent = true);
 	AIMODULE_API void SetGroupsToIgnore(int32 GroupFlags, bool bUpdateAgent = true);
 
-	FORCEINLINE bool IsCrowdSimulationEnabled() const { return SimulationState == ECrowdSimulationState::Enabled; }
-	FORCEINLINE bool IsCrowdSimulatioSuspended() const { return bSuspendCrowdSimulation; }
-	FORCEINLINE bool IsCrowdAnticipateTurnsEnabled() const { return bEnableAnticipateTurns; }
-	FORCEINLINE bool IsCrowdObstacleAvoidanceEnabled() const { return bEnableObstacleAvoidance; }
-	FORCEINLINE bool IsCrowdSeparationEnabled() const { return bEnableSeparation; }
-	FORCEINLINE bool IsCrowdOptimizeVisibilityEnabled() const { return bEnableOptimizeVisibility; /** don't check suspend here! */ }
-	FORCEINLINE bool IsCrowdOptimizeTopologyEnabled() const { return bEnableOptimizeTopology; }
-	FORCEINLINE bool IsCrowdPathOffsetEnabled() const { return bEnablePathOffset; }
-	FORCEINLINE bool IsCrowdSlowdownAtGoalEnabled() const { return bEnableSlowdownAtGoal; }
-	FORCEINLINE bool IsCrowdAffectFallingVelocityEnabled() const { return bAffectFallingVelocity; }
-	FORCEINLINE bool IsCrowdRotateToVelocityEnabled() const { return bRotateToVelocity; }
-
-	FORCEINLINE ECrowdSimulationState GetCrowdSimulationState() const { return SimulationState; }
-	FORCEINLINE bool IsCrowdSimulationActive() const { return IsCrowdSimulationEnabled() && !IsCrowdSimulatioSuspended(); }
+	inline bool IsCrowdSimulationEnabled() const { return SimulationState == ECrowdSimulationState::Enabled; }
+	inline bool IsCrowdSimulatioSuspended() const { return bSuspendCrowdSimulation; }
+	inline bool IsCrowdAnticipateTurnsEnabled() const { return bEnableAnticipateTurns; }
+	inline bool IsCrowdObstacleAvoidanceEnabled() const { return bEnableObstacleAvoidance; }
+	inline bool IsCrowdSeparationEnabled() const { return bEnableSeparation; }
+	inline bool IsCrowdOptimizeVisibilityEnabled() const { return bEnableOptimizeVisibility; /** don't check suspend here! */ }
+	inline bool IsCrowdOptimizeTopologyEnabled() const { return bEnableOptimizeTopology; }
+	inline bool IsCrowdPathOffsetEnabled() const { return bEnablePathOffset; }
+	inline bool IsCrowdSlowdownAtGoalEnabled() const { return bEnableSlowdownAtGoal; }
+	inline bool IsCrowdAffectFallingVelocityEnabled() const { return bAffectFallingVelocity; }
+	inline bool IsCrowdRotateToVelocityEnabled() const { return bRotateToVelocity; }
+
+	inline ECrowdSimulationState GetCrowdSimulationState() const { return SimulationState; }
+	inline bool IsCrowdSimulationActive() const { return IsCrowdSimulationEnabled() && !IsCrowdSimulatioSuspended(); }
 	/** checks if bEnableAnticipateTurns is set to true, and if crowd simulation is not suspended */
-	FORCEINLINE bool IsCrowdAnticipateTurnsActive() const { return IsCrowdAnticipateTurnsEnabled() && !IsCrowdSimulatioSuspended(); }
+	inline bool IsCrowdAnticipateTurnsActive() const { return IsCrowdAnticipateTurnsEnabled() && !IsCrowdSimulatioSuspended(); }
 	/** checks if bEnableObstacleAvoidance is set to true, and if crowd simulation is not suspended */
-	FORCEINLINE bool IsCrowdObstacleAvoidanceActive() const { return IsCrowdObstacleAvoidanceEnabled() && !IsCrowdSimulatioSuspended(); }
+	inline bool IsCrowdObstacleAvoidanceActive() const { return IsCrowdObstacleAvoidanceEnabled() && !IsCrowdSimulatioSuspended(); }
 	/** checks if bEnableSeparation is set to true, and if crowd simulation is not suspended */
-	FORCEINLINE bool IsCrowdSeparationActive() const { return IsCrowdSeparationEnabled() && !IsCrowdSimulatioSuspended(); }
+	inline bool IsCrowdSeparationActive() const { return IsCrowdSeparationEnabled() && !IsCrowdSimulatioSuspended(); }
 	/** checks if bEnableOptimizeTopology is set to true, and if crowd simulation is not suspended */
-	FORCEINLINE bool IsCrowdOptimizeTopologyActive() const { return IsCrowdOptimizeTopologyEnabled() && !IsCrowdSimulatioSuspended(); }
+	inline bool IsCrowdOptimizeTopologyActive() const { return IsCrowdOptimizeTopologyEnabled() && !IsCrowdSimulatioSuspended(); }
 
-	FORCEINLINE float GetCrowdSeparationWeight() const { return SeparationWeight; }
-	FORCEINLINE float GetCrowdCollisionQueryRange() const { return CollisionQueryRange; }
-	FORCEINLINE float GetCrowdPathOptimizationRange() const { return PathOptimizationRange; }
-	FORCEINLINE ECrowdAvoidanceQuality::Type GetCrowdAvoidanceQuality() const { return AvoidanceQuality; }
-	FORCEINLINE float GetCrowdAvoidanceRangeMultiplier() const { return AvoidanceRangeMultiplier; }
+	inline float GetCrowdSeparationWeight() const { return SeparationWeight; }
+	inline float GetCrowdCollisionQueryRange() const { return CollisionQueryRange; }
+	inline float GetCrowdPathOptimizationRange() const { return PathOptimizationRange; }
+	inline ECrowdAvoidanceQuality::Type GetCrowdAvoidanceQuality() const { return AvoidanceQuality; }
+	inline float GetCrowdAvoidanceRangeMultiplier() const { return AvoidanceRangeMultiplier; }
 	AIMODULE_API int32 GetAvoidanceGroup() const;
 	AIMODULE_API int32 GetGroupsToAvoid() const;
 	AIMODULE_API int32 GetGroupsToIgnore() const;
diff --git a/Engine/Source/Runtime/AIModule/Classes/Navigation/PathFollowingComponent.h b/Engine/Source/Runtime/AIModule/Classes/Navigation/PathFollowingComponent.h
--- a/Engine/Source/Runtime/AIModule/Classes/Navigation/PathFollowingComponent.h
+++ b/Engine/Source/Runtime/AIModule/Classes/Navigation/PathFollowingComponent.h
@@ -266,7 +266,7 @@ class UPathFollowingComponent : public UActorComponent, public IAIResourceInterf
 	/** notify about finished movement */
 	AIMODULE_API virtual void OnPathFinished(const FPathFollowingResult& Result);
 
-	FORCEINLINE void OnPathFinished(EPathFollowingResult::Type ResultCode, uint16 ExtraResultFlags) { OnPathFinished(FPathFollowingResult(ResultCode, ExtraResultFlags)); }
+	inline void OnPathFinished(EPathFollowingResult::Type ResultCode, uint16 ExtraResultFlags) { OnPathFinished(FPathFollowingResult(ResultCode, ExtraResultFlags)); }
 
 	/** notify about finishing move along current path segment */
 	AIMODULE_API virtual void OnSegmentFinished();
@@ -319,10 +319,10 @@ class UPathFollowingComponent : public UActorComponent, public IAIResourceInterf
 	bool IsDecelerating() const { return bIsDecelerating; };
 
 	/** @returns state of movement stopping on finish */
-	FORCEINLINE bool IsStopMovementOnFinishActive() const { return bStopMovementOnFinish; }
+	inline bool IsStopMovementOnFinishActive() const { return bStopMovementOnFinish; }
 	
 	/** set whether movement is stopped on finish of move. */
-	FORCEINLINE void SetStopMovementOnFinish(bool bEnable) { bStopMovementOnFinish = bEnable; }
+	inline void SetStopMovementOnFinish(bool bEnable) { bStopMovementOnFinish = bEnable; }
 
 	/** set threshold for precise reach tests in intermediate goals (minimal test radius)  */
 	AIMODULE_API void SetPreciseReachThreshold(float AgentRadiusMultiplier, float AgentHalfHeightMultiplier);
@@ -337,31 +337,31 @@ class UPathFollowingComponent : public UActorComponent, public IAIResourceInterf
 	/** Returns current location on navigation data */
 	AIMODULE_API FNavLocation GetCurrentNavLocation() const;
 
-	FORCEINLINE EPathFollowingStatus::Type GetStatus() const
+	inline EPathFollowingStatus::Type GetStatus() const
 	{
 		PRAGMA_DISABLE_DEPRECATION_WARNINGS
 		return Status;
 		PRAGMA_ENABLE_DEPRECATION_WARNINGS
 	}
-	FORCEINLINE float GetAcceptanceRadius() const { return AcceptanceRadius; }
-	FORCEINLINE float GetDefaultAcceptanceRadius() const { return MyDefaultAcceptanceRadius; }
+	inline float GetAcceptanceRadius() const { return AcceptanceRadius; }
+	inline float GetDefaultAcceptanceRadius() const { return MyDefaultAcceptanceRadius; }
 	AIMODULE_API void SetAcceptanceRadius(const float InAcceptanceRadius);
-	FORCEINLINE AActor* GetMoveGoal() const { return DestinationActor.Get(); }
-	FORCEINLINE bool HasPartialPath() const { return Path.IsValid() && Path->IsPartial(); }
-	FORCEINLINE bool DidMoveReachGoal() const
+	inline AActor* GetMoveGoal() const { return DestinationActor.Get(); }
+	inline bool HasPartialPath() const { return Path.IsValid() && Path->IsPartial(); }
+	inline bool DidMoveReachGoal() const
 	{
 		PRAGMA_DISABLE_DEPRECATION_WARNINGS
 		return bLastMoveReachedGoal && (Status == EPathFollowingStatus::Idle);
 		PRAGMA_ENABLE_DEPRECATION_WARNINGS
 	}
 
-	FORCEINLINE FAIRequestID GetCurrentRequestId() const { return CurrentRequestId; }
-	FORCEINLINE uint32 GetCurrentPathIndex() const { return MoveSegmentStartIndex; }
-	FORCEINLINE uint32 GetNextPathIndex() const { return MoveSegmentEndIndex; }
-	FORCEINLINE UObject* GetCurrentCustomLinkOb() const { return CurrentCustomLinkOb.Get(); }
-	FORCEINLINE FVector GetCurrentTargetLocation() const { return *CurrentDestination; }
-	FORCEINLINE FBasedPosition GetCurrentTargetLocationBased() const { return CurrentDestination; }
-	FORCEINLINE FVector GetMoveGoalLocationOffset() const { return MoveOffset; }
+	inline FAIRequestID GetCurrentRequestId() const { return CurrentRequestId; }
+	inline uint32 GetCurrentPathIndex() const { return MoveSegmentStartIndex; }
+	inline uint32 GetNextPathIndex() const { return MoveSegmentEndIndex; }
+	inline UObject* GetCurrentCustomLinkOb() const { return CurrentCustomLinkOb.Get(); }
+	inline FVector GetCurrentTargetLocation() const { return *CurrentDestination; }
+	inline FBasedPosition GetCurrentTargetLocationBased() const { return CurrentDestination; }
+	inline FVector GetMoveGoalLocationOffset() const { return MoveOffset; }
 	bool HasStartedNavLinkMove() const { return bWalkingNavLinkStart; }
 	AIMODULE_API bool IsCurrentSegmentNavigationLink() const;
 	AIMODULE_API FVector GetCurrentDirection() const;
@@ -369,10 +369,10 @@ class UPathFollowingComponent : public UActorComponent, public IAIResourceInterf
 	FVector GetCurrentMoveInput() const { return CurrentMoveInput; }
 
 	/** check if path following has authority over movement (e.g. not falling) and can update own state */
-	FORCEINLINE bool HasMovementAuthority() const { return (NavMovementInterface == nullptr) || NavMovementInterface->CanStopPathFollowing(); }
+	inline bool HasMovementAuthority() const { return (NavMovementInterface == nullptr) || NavMovementInterface->CanStopPathFollowing(); }
 
-	FORCEINLINE const FNavPathSharedPtr GetPath() const { return Path; }
-	FORCEINLINE bool HasValidPath() const { return Path.IsValid() && Path->IsValid(); }
+	inline const FNavPathSharedPtr GetPath() const { return Path; }
+	inline bool HasValidPath() const { return Path.IsValid() && Path->IsValid(); }
 	AIMODULE_API bool HasDirectPath() const;
 
 	/** readable name of current status */
@@ -663,12 +663,12 @@ protected:
 	AIMODULE_API bool IsBlocked() const;
 
 	/** switch to next segment on path */
-	FORCEINLINE void SetNextMoveSegment() { SetMoveSegment(GetNextPathIndex()); }
+	inline void SetNextMoveSegment() { SetMoveSegment(GetNextPathIndex()); }
 
 	/** assign new request Id */
-	FORCEINLINE void StoreRequestId() { CurrentRequestId = UPathFollowingComponent::GetNextRequestId(); }
+	inline void StoreRequestId() { CurrentRequestId = UPathFollowingComponent::GetNextRequestId(); }
 
-	FORCEINLINE static uint32 GetNextRequestId() { return NextRequestId++; }
+	inline static uint32 GetNextRequestId() { return NextRequestId++; }
 
 	/** Checks if this PathFollowingComponent is already on path, and
 	*	if so determines index of next path point
diff --git a/Engine/Source/Runtime/AIModule/Classes/Navigation/RecastGraphAStar.h b/Engine/Source/Runtime/AIModule/Classes/Navigation/RecastGraphAStar.h
--- a/Engine/Source/Runtime/AIModule/Classes/Navigation/RecastGraphAStar.h
+++ b/Engine/Source/Runtime/AIModule/Classes/Navigation/RecastGraphAStar.h
@@ -38,13 +38,13 @@ struct FRecastNeighbour
 	friend FRecastGraphWrapper;
 	friend FRecastGraphAStarFilter;
 
-	FORCEINLINE operator dtPolyRef() const
+	inline operator dtPolyRef() const
 	{
 		return NodeRef;
 	}
 
 protected:
-	FORCEINLINE FRecastNeighbour(dtPolyRef InNodeRef, unsigned char InSide = 0)
+	inline FRecastNeighbour(dtPolyRef InNodeRef, unsigned char InSide = 0)
 		: NodeRef{ InNodeRef }
 		, Side{ InSide }
 	{}
@@ -87,22 +87,22 @@ public:
 	// FGraphAStar: TGraph
 	typedef dtPolyRef FNodeRef;
 
-	FORCEINLINE bool IsValidRef(const dtPolyRef& NodeRef) const
+	inline bool IsValidRef(const dtPolyRef& NodeRef) const
 	{
 		return NodeRef != INVALID_NAVNODEREF;
 	}
 	AIMODULE_API FRecastNeighbour GetNeighbour(const FRecastAStarSearchNode& Node, const int32 NeighbourIndex) const;
 	//////////////////////////////////////////////////////////////////////////
 
-	FORCEINLINE const dtNavMeshQuery& GetRecastQuery() const { return RecastQuery; }
+	inline const dtNavMeshQuery& GetRecastQuery() const { return RecastQuery; }
 
 protected:
 
 	friend FRecastGraphAStarFilter;
 	friend FRecastAStarSearchNode;
 
-	FORCEINLINE const ARecastNavMesh* GetRecastNavMeshActor() const { checkSlow(RecastNavMeshActor);  return RecastNavMeshActor; }
-	FORCEINLINE const dtNavMesh* GetDetourNavMesh() const { checkSlow(DetourNavMesh); return DetourNavMesh; }
+	inline const ARecastNavMesh* GetRecastNavMeshActor() const { checkSlow(RecastNavMeshActor);  return RecastNavMeshActor; }
+	inline const dtNavMesh* GetDetourNavMesh() const { checkSlow(DetourNavMesh); return DetourNavMesh; }
 	AIMODULE_API void BindFilter(FRecastGraphAStarFilter& AStarFilter);
 #endif // WITH_RECAST
 
@@ -124,7 +124,7 @@ struct FRecastAStarSearchNode : public FGraphAStarDefaultNode<FRecastGraphWrappe
 {
 	typedef FGraphAStarDefaultNode<FRecastGraphWrapper> Super;
 
-	FORCEINLINE FRecastAStarSearchNode(const dtPolyRef InNodeRef = INVALID_NAVNODEREF, FVector InPosition = FVector(TNumericLimits<FVector::FReal>::Max(), TNumericLimits<FVector::FReal>::Max(), TNumericLimits<FVector::FReal>::Max()) )
+	inline FRecastAStarSearchNode(const dtPolyRef InNodeRef = INVALID_NAVNODEREF, FVector InPosition = FVector(TNumericLimits<FVector::FReal>::Max(), TNumericLimits<FVector::FReal>::Max(), TNumericLimits<FVector::FReal>::Max()) )
 		: Super(InNodeRef)
 		, Tile{ nullptr }
 		, Poly{ nullptr }
@@ -138,22 +138,22 @@ struct FRecastAStarSearchNode : public FGraphAStarDefaultNode<FRecastGraphWrappe
 	mutable const dtPoly* Poly;
 	mutable FVector::FReal Position[3]; // Position in Recast World coordinate system
 
-	FORCEINLINE operator dtPolyRef() const
+	inline operator dtPolyRef() const
 	{
 		return NodeRef;
 	}
 
-	FORCEINLINE void Initialize(const FRecastGraphWrapper& RecastGraphWrapper) const
+	inline void Initialize(const FRecastGraphWrapper& RecastGraphWrapper) const
 	{
 		RecastGraphWrapper.GetDetourNavMesh()->getTileAndPolyByRefUnsafe(NodeRef, &Tile, &Poly);
 	}
 
-	FORCEINLINE bool HasValidCacheInfo() const
+	inline bool HasValidCacheInfo() const
 	{
 		return Tile != nullptr && Poly != nullptr && Position[0] != TNumericLimits<FVector::FReal>::Max() && Position[1] != TNumericLimits<FVector::FReal>::Max() && Position[2] != TNumericLimits<FVector::FReal>::Max();
 	}
 
-	FORCEINLINE void CacheInfo(const FRecastGraphWrapper& RecastGraphWrapper, const FRecastAStarSearchNode& FromNode) const
+	inline void CacheInfo(const FRecastGraphWrapper& RecastGraphWrapper, const FRecastAStarSearchNode& FromNode) const
 	{
 		checkSlow(FromNode.HasValidCacheInfo());
 		if (!HasValidCacheInfo())
@@ -176,16 +176,16 @@ struct FRecastGraphAStarFilter
 {
 	AIMODULE_API FRecastGraphAStarFilter(FRecastGraphWrapper& InRecastGraphWrapper, const FRecastQueryFilter& InFilter, uint32 InMaxSearchNodes, const FVector::FReal InCostLimit, const UObject* Owner);
 
-	FORCEINLINE bool WantsPartialSolution() const
+	inline bool WantsPartialSolution() const
 	{ 
 		return true; 
 	}
-	FORCEINLINE FVector::FReal GetHeuristicScale() const
+	inline FVector::FReal GetHeuristicScale() const
 	{ 
 		return Filter.getHeuristicScale();
 	}
 
-	FORCEINLINE FVector::FReal GetHeuristicCost(const FRecastAStarSearchNode& StartNode, const FRecastAStarSearchNode& EndNode) const
+	inline FVector::FReal GetHeuristicCost(const FRecastAStarSearchNode& StartNode, const FRecastAStarSearchNode& EndNode) const
 	{
 		check(EndNode.HasValidCacheInfo());
 
@@ -194,7 +194,7 @@ struct FRecastGraphAStarFilter
 		return Cost;
 	}
 
-	FORCEINLINE FVector::FReal GetTraversalCost(const FRecastAStarSearchNode& StartNode, const FRecastAStarSearchNode& EndNode) const
+	inline FVector::FReal GetTraversalCost(const FRecastAStarSearchNode& StartNode, const FRecastAStarSearchNode& EndNode) const
 	{
 		EndNode.CacheInfo(RecastGraphWrapper, StartNode);
 		const FVector::FReal Cost = Filter.getCost(
@@ -206,7 +206,7 @@ struct FRecastGraphAStarFilter
 		return Cost;
 	}
 
-	FORCEINLINE bool IsTraversalAllowed(const dtPolyRef& NodeA, const FRecastNeighbour& NodeB) const
+	inline bool IsTraversalAllowed(const dtPolyRef& NodeA, const FRecastNeighbour& NodeB) const
 	{
 		if (!Filter.isValidLinkSide(NodeB.Side))
 		{
@@ -218,22 +218,22 @@ struct FRecastGraphAStarFilter
 		return Filter.passFilter(NodeB.NodeRef, TileB, PolyB) && RecastGraphWrapper.GetRecastQuery().passLinkFilterByRef(TileB, NodeB.NodeRef);
 	}
 
-	FORCEINLINE bool ShouldIgnoreClosedNodes() const
+	inline bool ShouldIgnoreClosedNodes() const
 	{
 		return Filter.getShouldIgnoreClosedNodes();
 	}
 
-	FORCEINLINE bool ShouldIncludeStartNodeInPath() const
+	inline bool ShouldIncludeStartNodeInPath() const
 	{
 		return true;
 	}
 
-	FORCEINLINE uint32 GetMaxSearchNodes() const
+	inline uint32 GetMaxSearchNodes() const
 	{
 		return MaxSearchNodes;
 	}
 
-	FORCEINLINE FVector::FReal GetCostLimit() const
+	inline FVector::FReal GetCostLimit() const
 	{
 		return CostLimit;
 	}
@@ -242,8 +242,8 @@ protected:
 	friend FRecastAStarSearchNode;
 	friend FRecastGraphWrapper;
 
-	FORCEINLINE const FRecastQueryFilter& GetFilter() const { return Filter; }
-	FORCEINLINE FRecastSpeciaLinkFilter& GetLinkFilter() { return LinkFilter; }
+	inline const FRecastQueryFilter& GetFilter() const { return Filter; }
+	inline FRecastSpeciaLinkFilter& GetLinkFilter() { return LinkFilter; }
 
 private:
 	const FRecastQueryFilter& Filter;
diff --git a/Engine/Source/Runtime/AIModule/Classes/Perception/AIPerceptionComponent.h b/Engine/Source/Runtime/AIModule/Classes/Perception/AIPerceptionComponent.h
--- a/Engine/Source/Runtime/AIModule/Classes/Perception/AIPerceptionComponent.h
+++ b/Engine/Source/Runtime/AIModule/Classes/Perception/AIPerceptionComponent.h
@@ -139,7 +139,7 @@ struct FActorPerceptionInfo
 	*
 	* @return Location of the last sensed stimuli or FAISystem::InvalidLocation if given sense has never registered related Target actor or if last stimuli has expired.
 	*/
-	FORCEINLINE FVector GetStimulusLocation(const FAISenseID Sense) const
+	inline FVector GetStimulusLocation(const FAISenseID Sense) const
 	{
 		return LastSensedStimuli.IsValidIndex(Sense) && (LastSensedStimuli[Sense].IsValid() && (LastSensedStimuli[Sense].IsExpired() == false)) ? LastSensedStimuli[Sense].StimulusLocation : FAISystem::InvalidLocation;
 	}
@@ -149,7 +149,7 @@ struct FActorPerceptionInfo
 	*
 	* @return Location of the receiver for the last sensed stimuli or FAISystem::InvalidLocation if given sense has never registered related Target actor or last stimuli has expired.
 	*/
-	FORCEINLINE FVector GetReceiverLocation(const FAISenseID Sense) const
+	inline FVector GetReceiverLocation(const FAISenseID Sense) const
 	{
 		return LastSensedStimuli.IsValidIndex(Sense) && (LastSensedStimuli[Sense].IsValid() && (LastSensedStimuli[Sense].IsExpired() == false)) ? LastSensedStimuli[Sense].ReceiverLocation : FAISystem::InvalidLocation;
 	}
@@ -159,7 +159,7 @@ struct FActorPerceptionInfo
 	*
 	* @return True if a target has been registered (even if not currently sensed) for the given sense and the stimuli is not expired.
 	*/
-	FORCEINLINE bool HasKnownStimulusOfSense(const FAISenseID Sense) const
+	inline bool HasKnownStimulusOfSense(const FAISenseID Sense) const
 	{
 		return LastSensedStimuli.IsValidIndex(Sense) && (LastSensedStimuli[Sense].IsValid() && (LastSensedStimuli[Sense].IsExpired() == false));
 	}
@@ -169,7 +169,7 @@ struct FActorPerceptionInfo
 	*
 	* @return True if a target is still sensed for the given sense and the stimuli is not expired.
 	*/
-	FORCEINLINE bool IsSenseActive(const FAISenseID Sense) const
+	inline bool IsSenseActive(const FAISenseID Sense) const
 	{
 		return LastSensedStimuli.IsValidIndex(Sense) && LastSensedStimuli[Sense].IsActive();
 	}
@@ -275,15 +275,15 @@ public:
 	AIMODULE_API const AActor* GetBodyActor() const;
 	AIMODULE_API AActor* GetMutableBodyActor();
 
-	FORCEINLINE FPerceptionChannelAllowList GetPerceptionFilter() const { return PerceptionFilter; }
+	inline FPerceptionChannelAllowList GetPerceptionFilter() const { return PerceptionFilter; }
 
 	AIMODULE_API FGenericTeamId GetTeamIdentifier() const;
-	FORCEINLINE FPerceptionListenerID GetListenerId() const { return PerceptionListenerId; }
+	inline FPerceptionListenerID GetListenerId() const { return PerceptionListenerId; }
 
 	AIMODULE_API FVector GetActorLocation(const AActor& Actor) const;
-	FORCEINLINE const FActorPerceptionInfo* GetActorInfo(const AActor& Actor) const { return PerceptualData.Find(&Actor); }
-	FORCEINLINE FActorPerceptionContainer::TIterator GetPerceptualDataIterator() { return FActorPerceptionContainer::TIterator(PerceptualData); }
-	FORCEINLINE FActorPerceptionContainer::TConstIterator GetPerceptualDataConstIterator() const { return FActorPerceptionContainer::TConstIterator(PerceptualData); }
+	inline const FActorPerceptionInfo* GetActorInfo(const AActor& Actor) const { return PerceptualData.Find(&Actor); }
+	inline FActorPerceptionContainer::TIterator GetPerceptualDataIterator() { return FActorPerceptionContainer::TIterator(PerceptualData); }
+	inline FActorPerceptionContainer::TConstIterator GetPerceptualDataConstIterator() const { return FActorPerceptionContainer::TConstIterator(PerceptualData); }
 
 	AIMODULE_API virtual void GetHostileActors(TArray<AActor*>& OutActors) const;
 	
@@ -298,8 +298,8 @@ public:
 	AIMODULE_API const FActorPerceptionInfo* GetFreshestTrace(const FAISenseID Sense) const;
 	
 	AIMODULE_API void SetDominantSense(TSubclassOf<UAISense> InDominantSense);
-	FORCEINLINE FAISenseID GetDominantSenseID() const { return DominantSenseID; }
-	FORCEINLINE TSubclassOf<UAISense> GetDominantSense() const { return DominantSense; }
+	inline FAISenseID GetDominantSenseID() const { return DominantSenseID; }
+	inline TSubclassOf<UAISense> GetDominantSense() const { return DominantSense; }
 	AIMODULE_API UAISenseConfig* GetSenseConfig(const FAISenseID& SenseID);
 	AIMODULE_API const UAISenseConfig* GetSenseConfig(const FAISenseID& SenseID) const;
 
diff --git a/Engine/Source/Runtime/AIModule/Classes/Perception/AIPerceptionSystem.h b/Engine/Source/Runtime/AIModule/Classes/Perception/AIPerceptionSystem.h
--- a/Engine/Source/Runtime/AIModule/Classes/Perception/AIPerceptionSystem.h
+++ b/Engine/Source/Runtime/AIModule/Classes/Perception/AIPerceptionSystem.h
@@ -75,7 +75,7 @@ protected:
 			: SenseID(InSenseID), Source(SourceActor)
 		{}
 
-		FORCEINLINE bool operator==(const FPerceptionSourceRegistration& Other) const
+		inline bool operator==(const FPerceptionSourceRegistration& Other) const
 		{
 			return SenseID == Other.SenseID && Source == Other.Source;
 		}
@@ -84,7 +84,7 @@ protected:
 
 public:
 
-	FORCEINLINE bool IsSenseInstantiated(const FAISenseID& SenseID) const { return SenseID.IsValid() && Senses.IsValidIndex(SenseID) && Senses[SenseID] != nullptr; }
+	inline bool IsSenseInstantiated(const FAISenseID& SenseID) const { return SenseID.IsValid() && Senses.IsValidIndex(SenseID) && Senses[SenseID] != nullptr; }
 
 	/** Registers listener if not registered */
 	AIMODULE_API void UpdateListener(UAIPerceptionComponent& Listener);
@@ -199,7 +199,7 @@ protected:
 	AIMODULE_API bool AgeStimuli(const float Amount);
 
 	friend class UAISense;
-	FORCEINLINE AIPerception::FListenerMap& GetListenersMap() { return ListenerContainer; }
+	inline AIPerception::FListenerMap& GetListenersMap() { return ListenerContainer; }
 
 	friend class UAISystem;
 	AIMODULE_API virtual void OnNewPawn(APawn& Pawn);
diff --git a/Engine/Source/Runtime/AIModule/Classes/Perception/AIPerceptionTypes.h b/Engine/Source/Runtime/AIModule/Classes/Perception/AIPerceptionTypes.h
--- a/Engine/Source/Runtime/AIModule/Classes/Perception/AIPerceptionTypes.h
+++ b/Engine/Source/Runtime/AIModule/Classes/Perception/AIPerceptionTypes.h
@@ -52,30 +52,30 @@ struct FPerceptionChannelAllowList
 		return (AcceptedChannelsMask == 0);
 	}
 
-	FORCEINLINE FPerceptionChannelAllowList& FilterOutChannel(FAISenseID Channel)
+	inline FPerceptionChannelAllowList& FilterOutChannel(FAISenseID Channel)
 	{
 		AcceptedChannelsMask &= ~(1 << Channel);
 		return *this;
 	}
 
-	FORCEINLINE_DEBUGGABLE FPerceptionChannelAllowList& AcceptChannel(FAISenseID Channel)
+	inline FPerceptionChannelAllowList& AcceptChannel(FAISenseID Channel)
 	{
 		AcceptedChannelsMask |= (1 << Channel);
 		return *this;
 	}
 
-	FORCEINLINE bool ShouldRespondToChannel(FAISenseID Channel) const
+	inline bool ShouldRespondToChannel(FAISenseID Channel) const
 	{
 		return (AcceptedChannelsMask & (1 << Channel)) != 0;
 	}
 
-	FORCEINLINE FPerceptionChannelAllowList& MergeFilterIn(const FPerceptionChannelAllowList& OtherFilter)
+	inline FPerceptionChannelAllowList& MergeFilterIn(const FPerceptionChannelAllowList& OtherFilter)
 	{
 		AcceptedChannelsMask |= OtherFilter.AcceptedChannelsMask;
 		return *this;
 	}
 
-	FORCEINLINE FFlagsContainer GetAcceptedChannelsMask() const 
+	inline FFlagsContainer GetAcceptedChannelsMask() const 
 	{ 
 		return AcceptedChannelsMask;
 	}
@@ -96,7 +96,7 @@ struct FPerceptionChannelAllowList
 			FindNextAcceptedChannel();
 		}
 
-		FORCEINLINE void FindNextAcceptedChannel()
+		inline void FindNextAcceptedChannel()
 		{
 			const FFlagsContainer& Flags = AllowList.GetAcceptedChannelsMask();
 
@@ -106,17 +106,17 @@ struct FPerceptionChannelAllowList
 			}
 		}
 
-		FORCEINLINE explicit operator bool() const
+		inline explicit operator bool() const
 		{
 			return (RemainingChannelsToTest & AllowList.GetAcceptedChannelsMask()) != 0;
 		}
 
-		FORCEINLINE int32 operator*() const
+		inline int32 operator*() const
 		{
 			return CurrentIndex;
 		}
 
-		FORCEINLINE void operator++()
+		inline void operator++()
 		{
 			// mark "old" index as already used
 			RemainingChannelsToTest &= ~(1 << CurrentIndex);
@@ -182,20 +182,20 @@ public:
 	FAIStimulus& SetStimulusAge(float StimulusAge) { Age = StimulusAge; return *this; }
 	FAIStimulus& SetWantsNotifyOnlyOnValueChange(bool InEnable) { bWantsToNotifyOnlyOnValueChange = InEnable; return *this; }
 	
-	FORCEINLINE float GetAge() const { return Strength > 0 ? Age : NeverHappenedAge; }
+	inline float GetAge() const { return Strength > 0 ? Age : NeverHappenedAge; }
 	/** @return false when this stimulus is no longer valid, when it is Expired */
-	FORCEINLINE bool AgeStimulus(float ConstPerceptionAgingRate) 
+	inline bool AgeStimulus(float ConstPerceptionAgingRate) 
 	{ 
 		Age += ConstPerceptionAgingRate; 
 		return Age < ExpirationAge;
 	}
-	FORCEINLINE bool WasSuccessfullySensed() const { return bSuccessfullySensed; }
-	FORCEINLINE bool IsExpired() const { return bExpired; }
-	FORCEINLINE void MarkNoLongerSensed() { bSuccessfullySensed = false; }
-	FORCEINLINE void MarkExpired() { bExpired = true; MarkNoLongerSensed(); }
-	FORCEINLINE bool IsActive() const { return WasSuccessfullySensed() == true && IsValid(); }
-	FORCEINLINE bool WantsToNotifyOnlyOnPerceptionChange() const { return bWantsToNotifyOnlyOnValueChange; }
-	FORCEINLINE bool IsValid() const { return Type != FAISenseID::InvalidID() && GetAge() < NeverHappenedAge; }
+	inline bool WasSuccessfullySensed() const { return bSuccessfullySensed; }
+	inline bool IsExpired() const { return bExpired; }
+	inline void MarkNoLongerSensed() { bSuccessfullySensed = false; }
+	inline void MarkExpired() { bExpired = true; MarkNoLongerSensed(); }
+	inline bool IsActive() const { return WasSuccessfullySensed() == true && IsValid(); }
+	inline bool WantsToNotifyOnlyOnPerceptionChange() const { return bWantsToNotifyOnlyOnValueChange; }
+	inline bool IsValid() const { return Type != FAISenseID::InvalidID() && GetAge() < NeverHappenedAge; }
 
 #if !(UE_BUILD_SHIPPING || UE_BUILD_TEST)
 	AIMODULE_API FString GetDebugDescription() const;
@@ -224,9 +224,9 @@ struct FAISenseAffiliationFilter
 	uint32 bDetectFriendlies : 1 = false;
 	
 	uint8 GetAsFlags() const { return IntCastChecked<uint8>((bDetectEnemies << ETeamAttitude::Hostile) | (bDetectNeutrals << ETeamAttitude::Neutral) | (bDetectFriendlies << ETeamAttitude::Friendly)); }
-	FORCEINLINE bool ShouldDetectAll() const { return (bDetectEnemies && bDetectNeutrals && bDetectFriendlies); }
+	inline bool ShouldDetectAll() const { return (bDetectEnemies && bDetectNeutrals && bDetectFriendlies); }
 
-	static FORCEINLINE uint8 DetectAllFlags() { return (1 << ETeamAttitude::Hostile) | (1 << ETeamAttitude::Neutral) | (1 << ETeamAttitude::Friendly); }
+	static inline uint8 DetectAllFlags() { return (1 << ETeamAttitude::Hostile) | (1 << ETeamAttitude::Neutral) | (1 << ETeamAttitude::Friendly); }
 
 	static bool ShouldSenseTeam(FGenericTeamId TeamA, FGenericTeamId TeamB, uint8 AffiliationFlags)
 	{
@@ -272,15 +272,15 @@ public:
 
 	AIMODULE_API void RegisterStimulus(AActor* Source, const FAIStimulus& Stimulus);
 
-	FORCEINLINE bool HasAnyNewStimuli() const { return bHasStimulusToProcess; }
+	inline bool HasAnyNewStimuli() const { return bHasStimulusToProcess; }
 	AIMODULE_API void ProcessStimuli();
 
-	FORCEINLINE bool HasSense(FAISenseID SenseID) const { return Filter.ShouldRespondToChannel(SenseID); }
+	inline bool HasSense(FAISenseID SenseID) const { return Filter.ShouldRespondToChannel(SenseID); }
 
 	// used to remove "dead" listeners
 	static AIMODULE_API const FPerceptionListener NullListener;
 
-	FORCEINLINE FPerceptionListenerID GetListenerID() const { return ListenerID; }
+	inline FPerceptionListenerID GetListenerID() const { return ListenerID; }
 
 	AIMODULE_API FName GetBodyActorName() const;
 	AIMODULE_API uint32 GetBodyActorUniqueID() const;
@@ -292,8 +292,8 @@ public:
 
 private:
 	friend class UAIPerceptionSystem;
-	FORCEINLINE void SetListenerID(FPerceptionListenerID InListenerID) { ListenerID = InListenerID; }
-	FORCEINLINE void MarkForStimulusProcessing() { bHasStimulusToProcess = true; }
+	inline void SetListenerID(FPerceptionListenerID InListenerID) { ListenerID = InListenerID; }
+	inline void MarkForStimulusProcessing() { bHasStimulusToProcess = true; }
 };
 
 struct FPerceptionStimuliSource
diff --git a/Engine/Source/Runtime/AIModule/Classes/Perception/AISense.h b/Engine/Source/Runtime/AIModule/Classes/Perception/AISense.h
--- a/Engine/Source/Runtime/AIModule/Classes/Perception/AISense.h
+++ b/Engine/Source/Runtime/AIModule/Classes/Perception/AISense.h
@@ -72,9 +72,9 @@ public:
 	{ 
 		return GetDefault<TSense>()->GetSenseID();
 	}
-	FORCEINLINE FAISenseID GetSenseID() const { return SenseID; }
+	inline FAISenseID GetSenseID() const { return SenseID; }
 
-	FORCEINLINE bool WantsUpdateOnlyOnPerceptionValueChange() const { return (NotifyType == EAISenseNotifyType::OnPerceptionChange); }
+	inline bool WantsUpdateOnlyOnPerceptionValueChange() const { return (NotifyType == EAISenseNotifyType::OnPerceptionChange); }
 
 	AIMODULE_API virtual void PostInitProperties() override;
 
@@ -106,9 +106,9 @@ public:
 	virtual void OnListenerForgetsActor(const FPerceptionListener& Listener, AActor& ActorToForget) {}
 	virtual void OnListenerForgetsAll(const FPerceptionListener& Listener) {}
 
-	FORCEINLINE void OnNewListener(const FPerceptionListener& NewListener) { OnNewListenerDelegate.ExecuteIfBound(NewListener); }
-	FORCEINLINE void OnListenerUpdate(const FPerceptionListener& UpdatedListener) { OnListenerUpdateDelegate.ExecuteIfBound(UpdatedListener); }
-	FORCEINLINE void OnListenerRemoved(const FPerceptionListener& RemovedListener) { OnListenerRemovedDelegate.ExecuteIfBound(RemovedListener); }
+	inline void OnNewListener(const FPerceptionListener& NewListener) { OnNewListenerDelegate.ExecuteIfBound(NewListener); }
+	inline void OnListenerUpdate(const FPerceptionListener& UpdatedListener) { OnListenerUpdateDelegate.ExecuteIfBound(UpdatedListener); }
+	inline void OnListenerRemoved(const FPerceptionListener& RemovedListener) { OnListenerRemovedDelegate.ExecuteIfBound(RemovedListener); }
 	virtual void OnListenerConfigUpdated(const FPerceptionListener& UpdatedListener) { OnListenerUpdate(UpdatedListener); }
 
 	bool WantsNewPawnNotification() const { return bWantsNewPawnNotification; }
@@ -128,12 +128,12 @@ protected:
 	virtual float Update() { return FLT_MAX; }
 
 	/** will result in updating as soon as possible */
-	FORCEINLINE void RequestImmediateUpdate() { TimeUntilNextUpdate = 0.f; }
+	inline void RequestImmediateUpdate() { TimeUntilNextUpdate = 0.f; }
 
 	/** will result in updating in specified number of seconds */
-	FORCEINLINE void RequestUpdateInSeconds(float UpdateInSeconds) { TimeUntilNextUpdate = UpdateInSeconds; }
+	inline void RequestUpdateInSeconds(float UpdateInSeconds) { TimeUntilNextUpdate = UpdateInSeconds; }
 
-	FORCEINLINE UAIPerceptionSystem* GetPerceptionSystem() { return PerceptionSystemInstance; }
+	inline UAIPerceptionSystem* GetPerceptionSystem() { return PerceptionSystemInstance; }
 
 	AIMODULE_API void SetSenseID(FAISenseID Index);
 
diff --git a/Engine/Source/Runtime/AIModule/Classes/Perception/AISenseEvent_Damage.h b/Engine/Source/Runtime/AIModule/Classes/Perception/AISenseEvent_Damage.h
--- a/Engine/Source/Runtime/AIModule/Classes/Perception/AISenseEvent_Damage.h
+++ b/Engine/Source/Runtime/AIModule/Classes/Perception/AISenseEvent_Damage.h
@@ -19,7 +19,7 @@ public:
 	
 	AIMODULE_API virtual FAISenseID GetSenseID() const override;
 
-	FORCEINLINE FAIDamageEvent GetDamageEvent()
+	inline FAIDamageEvent GetDamageEvent()
 	{
 		Event.Compile();
 		return Event;
diff --git a/Engine/Source/Runtime/AIModule/Classes/Perception/AISenseEvent_Hearing.h b/Engine/Source/Runtime/AIModule/Classes/Perception/AISenseEvent_Hearing.h
--- a/Engine/Source/Runtime/AIModule/Classes/Perception/AISenseEvent_Hearing.h
+++ b/Engine/Source/Runtime/AIModule/Classes/Perception/AISenseEvent_Hearing.h
@@ -22,7 +22,7 @@ public:
 	AIMODULE_API UAISenseEvent_Hearing(const FObjectInitializer& ObjectInitializer = FObjectInitializer::Get());
 	AIMODULE_API virtual FAISenseID GetSenseID() const override;
 	
-	FORCEINLINE FAINoiseEvent GetNoiseEvent()
+	inline FAINoiseEvent GetNoiseEvent()
 	{
 		Event.Compile();
 		return Event;
diff --git a/Engine/Source/Runtime/AIModule/Classes/Perception/AISense_Sight.h b/Engine/Source/Runtime/AIModule/Classes/Perception/AISense_Sight.h
--- a/Engine/Source/Runtime/AIModule/Classes/Perception/AISense_Sight.h
+++ b/Engine/Source/Runtime/AIModule/Classes/Perception/AISense_Sight.h
@@ -63,13 +63,13 @@ struct FAISightTarget
 
 	AIMODULE_API FAISightTarget(AActor* InTarget = NULL, FGenericTeamId InTeamId = FGenericTeamId::NoTeam);
 
-	FORCEINLINE FVector GetLocationSimple() const
+	inline FVector GetLocationSimple() const
 	{
 		const AActor* TargetPtr = Target.Get();
 		return TargetPtr ? TargetPtr->GetActorLocation() : FVector::ZeroVector;
 	}
 
-	FORCEINLINE const AActor* GetTargetActor() const { return Target.Get(); }
+	inline const AActor* GetTargetActor() const { return Target.Get(); }
 };
 
 struct FAISightQuery
diff --git a/Engine/Source/Runtime/AIModule/Classes/Perception/AISense_Team.h b/Engine/Source/Runtime/AIModule/Classes/Perception/AISense_Team.h
--- a/Engine/Source/Runtime/AIModule/Classes/Perception/AISense_Team.h
+++ b/Engine/Source/Runtime/AIModule/Classes/Perception/AISense_Team.h
@@ -35,12 +35,12 @@ public:
 	FAITeamStimulusEvent() : Broadcaster(nullptr), Enemy(nullptr) {}
 	AIMODULE_API FAITeamStimulusEvent(AActor* InBroadcaster, AActor* InEnemy, const FVector& InLastKnowLocation, float EventRange, float PassedInfoAge = 0.f, float InStrength = 1.f);
 
-	FORCEINLINE void CacheBroadcastLocation()
+	inline void CacheBroadcastLocation()
 	{
 		BroadcastLocation = Broadcaster ? Broadcaster->GetActorLocation() : FAISystem::InvalidLocation;
 	}
 
-	FORCEINLINE const FVector& GetBroadcastLocation() const 
+	inline const FVector& GetBroadcastLocation() const 
 	{
 		return BroadcastLocation;
 	}
diff --git a/Engine/Source/Runtime/AIModule/Public/BehaviorTreeNodeInitializationData.h b/Engine/Source/Runtime/AIModule/Public/BehaviorTreeNodeInitializationData.h
--- a/Engine/Source/Runtime/AIModule/Public/BehaviorTreeNodeInitializationData.h
+++ b/Engine/Source/Runtime/AIModule/Public/BehaviorTreeNodeInitializationData.h
@@ -28,7 +28,7 @@ struct FBehaviorTreeNodeInitializationData
 
 	struct FMemorySort
 	{
-		FORCEINLINE bool operator()(const FBehaviorTreeNodeInitializationData& A, const FBehaviorTreeNodeInitializationData& B) const
+		inline bool operator()(const FBehaviorTreeNodeInitializationData& A, const FBehaviorTreeNodeInitializationData& B) const
 		{
 			return A.DataSize > B.DataSize;
 		}
diff --git a/Engine/Source/Runtime/AIModule/Public/BlueprintNodeHelpers.h b/Engine/Source/Runtime/AIModule/Public/BlueprintNodeHelpers.h
--- a/Engine/Source/Runtime/AIModule/Public/BlueprintNodeHelpers.h
+++ b/Engine/Source/Runtime/AIModule/Public/BlueprintNodeHelpers.h
@@ -32,14 +32,14 @@ namespace BlueprintNodeHelpers
 
 	AIMODULE_API void AbortLatentActions(UActorComponent& OwnerOb, const UObject& Ob);
 
-	FORCEINLINE bool HasBlueprintFunction(FName FuncName, const UObject& Object, const UClass& StopAtClass)
+	inline bool HasBlueprintFunction(FName FuncName, const UObject& Object, const UClass& StopAtClass)
 	{
 		const UFunction* Function = Object.GetClass()->FindFunctionByName(FuncName);
 		ensure(Function);
 		return (Function != nullptr) && (Function->GetOuter() != &StopAtClass);
 	}
 
-	FORCEINLINE FString GetNodeName(const UObject& NodeObject)
+	inline FString GetNodeName(const UObject& NodeObject)
 	{
 		return NodeObject.GetClass()->GetName().LeftChop(2);
 	}
diff --git a/Engine/Source/Runtime/AIModule/Public/GraphAStar.h b/Engine/Source/Runtime/AIModule/Public/GraphAStar.h
--- a/Engine/Source/Runtime/AIModule/Public/GraphAStar.h
+++ b/Engine/Source/Runtime/AIModule/Public/GraphAStar.h
@@ -48,7 +48,7 @@ struct FGraphAStarDefaultNode
 	uint8 bIsOpened : 1;
 	uint8 bIsClosed : 1;
 
-	FORCEINLINE FGraphAStarDefaultNode(const FGraphNodeRef& InNodeRef)
+	inline FGraphAStarDefaultNode(const FGraphNodeRef& InNodeRef)
 		: NodeRef(InNodeRef)
 		, ParentRef(TIsPointer<FGraphNodeRef>::Value ? (FGraphNodeRef)0 : (FGraphNodeRef)INDEX_NONE)
 		, TraversalCost(TNumericLimits<FVector::FReal>::Max())
@@ -59,12 +59,12 @@ struct FGraphAStarDefaultNode
 		, bIsClosed(false)
 	{}
 
-	FORCEINLINE void MarkOpened() { bIsOpened = true; }
-	FORCEINLINE void MarkNotOpened() { bIsOpened = false; }
-	FORCEINLINE void MarkClosed() { bIsClosed = true; }
-	FORCEINLINE void MarkNotClosed() { bIsClosed = false; }
-	FORCEINLINE bool IsOpened() const { return bIsOpened; }
-	FORCEINLINE bool IsClosed() const { return bIsClosed; }
+	inline void MarkOpened() { bIsOpened = true; }
+	inline void MarkNotOpened() { bIsOpened = false; }
+	inline void MarkClosed() { bIsClosed = true; }
+	inline void MarkNotClosed() { bIsClosed = false; }
+	inline bool IsOpened() const { return bIsOpened; }
+	inline bool IsClosed() const { return bIsClosed; }
 };
 
 #define DECLARE_OPTIONALLY_IMPLEMENTED_TEMPLATE_CLASS_FUNCTION_IMPL( TemplateClass, TemplateClassParameter, ConditionalReturnType, ConditionalFunctionName, DefaultImpl ) \
@@ -73,7 +73,7 @@ struct CQuery##ConditionalFunctionName	\
 	template<typename TemplateClass> auto Requires(TemplateClassParameter& Obj) -> decltype(Obj.ConditionalFunctionName());	\
 };	\
 template <typename TemplateClass> \
-static FORCEINLINE decltype(auto) ConditionalFunctionName(TemplateClassParameter & Obj) \
+static inline decltype(auto) ConditionalFunctionName(TemplateClassParameter & Obj) \
 { \
 	if constexpr (TModels_V<CQuery##ConditionalFunctionName, TemplateClass>) \
 	{ \
@@ -93,7 +93,7 @@ struct CQuery##ConditionalFunctionName	\
 	template<typename TemplateClass> auto Requires(TemplateClassParameter& Obj, ConditionalParamType1 Param1) -> decltype(Obj.ConditionalFunctionName(Param1));	\
 };	\
 template <typename TemplateClass> \
-static FORCEINLINE decltype(auto) ConditionalFunctionName(TemplateClassParameter & Obj, ConditionalParamType1 Param1) \
+static inline decltype(auto) ConditionalFunctionName(TemplateClassParameter & Obj, ConditionalParamType1 Param1) \
 { \
 	if constexpr (TModels_V<CQuery##ConditionalFunctionName, TemplateClass>) \
 	{ \
@@ -113,7 +113,7 @@ struct CQuery##ConditionalFunctionName	\
 	template<typename TemplateClass> auto Requires(TemplateClassParameter& Obj, ConditionalParamType1 Param1, ConditionalParamType2 Param2) -> decltype(Obj.ConditionalFunctionName(Param1,Param2));	\
 };	\
 template <typename TemplateClass> \
-static FORCEINLINE decltype(auto) ConditionalFunctionName(TemplateClassParameter & Obj, ConditionalParamType1 Param1, ConditionalParamType2 Param2) \
+static inline decltype(auto) ConditionalFunctionName(TemplateClassParameter & Obj, ConditionalParamType1 Param1, ConditionalParamType2 Param2) \
 { \
 	if constexpr (TModels_V<CQuery##ConditionalFunctionName, TemplateClass>) \
 	{ \
@@ -133,7 +133,7 @@ struct CQuery##QueryFunctionName	\
 	template<typename TemplateClass> auto Requires(const TemplateClass& Obj) -> decltype(Obj.ConditionalFunctionName());	\
 };	\
 template <typename TemplateClass> \
-static FORCEINLINE QueryReturnType QueryFunctionName(const TemplateClass& Obj, QueryParam) \
+static inline QueryReturnType QueryFunctionName(const TemplateClass& Obj, QueryParam) \
 { \
 	if constexpr (TModels_V<CQuery##QueryFunctionName, TemplateClass>) \
 	{ \
@@ -205,7 +205,7 @@ struct FGraphAStar
 			: NodePool(InNodePool)
 		{}
 
-		FORCEINLINE bool operator()(const int32 A, const int32 B) const
+		inline bool operator()(const int32 A, const int32 B) const
 		{
 			return NodePool[A].TotalCost < NodePool[B].TotalCost;
 		}
@@ -222,7 +222,7 @@ struct FGraphAStar
 			NodeMap.Reserve(FMath::RoundUpToPowerOfTwo(Policy::NodePoolSize / 4));
 		}
 
-		FORCEINLINE FSearchNode& Add(const FSearchNode& SearchNode)
+		inline FSearchNode& Add(const FSearchNode& SearchNode)
 		{
 			FSearchNode& NewNode = Super::Emplace_GetRef(SearchNode);
 			NewNode.SearchNodeIndex = UE_PTRDIFF_TO_INT32(&NewNode - Super::GetData());
@@ -230,7 +230,7 @@ struct FGraphAStar
 			return NewNode;
 		}
 
-		FORCEINLINE_DEBUGGABLE FSearchNode& FindOrAdd(const FGraphNodeRef NodeRef)
+		inline FSearchNode& FindOrAdd(const FGraphNodeRef NodeRef)
 		{
 			// first find if node already exist in node map
 			const int32 NotInMapIndex = -1;
@@ -248,19 +248,19 @@ struct FGraphAStar
 			return NewNode;
 		}
 
-		FORCEINLINE FSearchNode* Find(const FGraphNodeRef NodeRef)
+		inline FSearchNode* Find(const FGraphNodeRef NodeRef)
 		{
 			const int32* IndexPtr = NodeMap.Find(NodeRef);
 			return IndexPtr ? &(*this)[*IndexPtr] : nullptr;
 		}
 
-		FORCEINLINE void Reset()
+		inline void Reset()
 		{
 			Super::Reset(Policy::NodePoolSize);
 			NodeMap.Reset();
 		}
 
-		FORCEINLINE void ReinitNodes()
+		inline void ReinitNodes()
 		{
 			for (FSearchNode& Node : *this)
 			{
@@ -311,7 +311,7 @@ struct FGraphAStar
 		}
 
 		UE_DEPRECATED(5.4, "PopIndex with a boolean bAllowShrinking has been deprecated - please use the version without parameter")
-		FORCEINLINE int32 PopIndex(bool bAllowShrinking)
+		inline int32 PopIndex(bool bAllowShrinking)
 		{
 			return PopIndex();
 		}
@@ -354,7 +354,7 @@ struct FGraphAStar
 	 */
 	template<typename TQueryFilter>
 	UE_DEPRECATED(5.1, "Please use ProcessSingleNode() taking FVector::FReal& OutBestNodeCost instead!")
-	FORCEINLINE_DEBUGGABLE bool ProcessSingleNode(const FSearchNode& EndNode, const bool bIsBound, const TQueryFilter& Filter, int32& OutBestNodeIndex, float& OutBestNodeCost)
+	inline bool ProcessSingleNode(const FSearchNode& EndNode, const bool bIsBound, const TQueryFilter& Filter, int32& OutBestNodeIndex, float& OutBestNodeCost)
 	{
 		double BestNodeCost;
 
@@ -368,7 +368,7 @@ struct FGraphAStar
 	 * returns true if loop should be continued
 	 */
 	template<typename TQueryFilter>
-	FORCEINLINE_DEBUGGABLE bool ProcessSingleNode(const FSearchNode& EndNode, const bool bIsBound, const TQueryFilter& Filter, int32& OutBestNodeIndex, FVector::FReal& OutBestNodeCost)
+	inline bool ProcessSingleNode(const FSearchNode& EndNode, const bool bIsBound, const TQueryFilter& Filter, int32& OutBestNodeIndex, FVector::FReal& OutBestNodeCost)
 	{
 		// Pop next best node and put it on closed list
 		const int32 ConsideredNodeIndex = OpenList.PopIndex();
diff --git a/Engine/Source/Runtime/AIModule/Public/SimpleCellGrid.h b/Engine/Source/Runtime/AIModule/Public/SimpleCellGrid.h
--- a/Engine/Source/Runtime/AIModule/Public/SimpleCellGrid.h
+++ b/Engine/Source/Runtime/AIModule/Public/SimpleCellGrid.h
@@ -118,33 +118,33 @@ public:
 		WorldBounds = FBox(Origin - FVector(0, 0, BoundsSize.Z / 2), Origin + FVector(BoundsSize.X, BoundsSize.Y, BoundsSize.Z / 2));
 	}
 
-	FORCEINLINE bool IsValid() const
+	inline bool IsValid() const
 	{
 		return Cells.Num() && GridCellSize > 0;
 	}
 
-	FORCEINLINE bool IsValidIndex(const int32 CellIndex) const
+	inline bool IsValidIndex(const int32 CellIndex) const
 	{
 		return Cells.IsValidIndex(CellIndex);
 	}
 
-	FORCEINLINE bool IsValidCoord(int32 LocationX, int32 LocationY) const
+	inline bool IsValidCoord(int32 LocationX, int32 LocationY) const
 	{
 		return (LocationX >= 0) && (LocationX < (int32)GridSize.Width) && (LocationY >= 0) && (LocationY < (int32)GridSize.Height);
 	}
 
-	FORCEINLINE bool IsValidCoord(const FIntVector& CellCoords) const
+	inline bool IsValidCoord(const FIntVector& CellCoords) const
 	{
 		return IsValidCoord(CellCoords.X, CellCoords.Y);
 	}
 
-	FORCEINLINE uint32 GetAllocatedSize() const
+	inline uint32 GetAllocatedSize() const
 	{
 		return IntCastChecked<uint32>(Cells.GetAllocatedSize());
 	}
 
 	/** Convert world location to (X,Y) coords on grid, result can be outside grid */
-	FORCEINLINE FIntVector GetCellCoordsUnsafe(const FVector& WorldLocation) const
+	inline FIntVector GetCellCoordsUnsafe(const FVector& WorldLocation) const
 	{
 		return FIntVector(
 			IntCastChecked<int32>(FMath::TruncToInt((WorldLocation.X - Origin.X) / GridCellSize)),
@@ -160,19 +160,19 @@ public:
 	}
 
 	/** Convert cell index to coord X on grid, result can be invalid */
-	FORCEINLINE int32 GetCellCoordX(int32 CellIndex) const
+	inline int32 GetCellCoordX(int32 CellIndex) const
 	{
 		return CellIndex / GridSize.Height;
 	}
 
 	/** Convert cell index to coord Y on grid, result can be invalid */
-	FORCEINLINE int32 GetCellCoordY(int32 CellIndex) const
+	inline int32 GetCellCoordY(int32 CellIndex) const
 	{
 		return CellIndex % GridSize.Height;
 	}
 
 	/** Convert cell index to (X,Y) coords on grid */
-	FORCEINLINE FIntVector GetCellCoords(int32 CellIndex) const
+	inline FIntVector GetCellCoords(int32 CellIndex) const
 	{
 		return FIntVector(GetCellCoordX(CellIndex), GetCellCoordY(CellIndex), 0);
 	}
@@ -185,19 +185,19 @@ public:
 	}
 
 	/** Convert (X,Y) coords on grid to cell index, result can be invalid */
-	FORCEINLINE int32 GetCellIndexUnsafe(const FIntVector& CellCoords) const
+	inline int32 GetCellIndexUnsafe(const FIntVector& CellCoords) const
 	{
 		return GetCellIndexUnsafe(CellCoords.X, CellCoords.Y);
 	}
 
 	/** Convert (X,Y) coords on grid to cell index, result can be invalid */
-	FORCEINLINE int32 GetCellIndexUnsafe(int32 LocationX, int32 LocationY) const
+	inline int32 GetCellIndexUnsafe(int32 LocationX, int32 LocationY) const
 	{
 		return (LocationX * GridSize.Height) + LocationY;
 	}
 
 	/** Convert (X,Y) coords on grid to cell index, returns -1 for position outside grid */
-	FORCEINLINE int32 GetCellIndex(int32 LocationX, int32 LocationY) const
+	inline int32 GetCellIndex(int32 LocationX, int32 LocationY) const
 	{
 		return IsValidCoord(LocationX, LocationY) ? GetCellIndexUnsafe(LocationX, LocationY) : INDEX_NONE;
 	}
@@ -210,13 +210,13 @@ public:
 	}
 
 	/** Return the bounding box of a cell. */
-	FORCEINLINE FBox GetWorldCellBox(int32 CellIndex) const
+	inline FBox GetWorldCellBox(int32 CellIndex) const
 	{
 		return GetWorldCellBox(GetCellCoordX(CellIndex), GetCellCoordY(CellIndex));
 	}
 
 	/** Return the bounding box of a cell. */
-	FORCEINLINE FBox GetWorldCellBox(int32 LocationX, int32 LocationY) const
+	inline FBox GetWorldCellBox(int32 LocationX, int32 LocationY) const
 	{
 		return FBox(
 			Origin + FVector(LocationX * GridCellSize, LocationY * GridCellSize, -BoundsSize.Z * 0.5f), 
@@ -225,13 +225,13 @@ public:
 	}
 
 	/** Return the 2D bounding box of a cell. */
-	FORCEINLINE FBox2D GetWorldCellBox2D(int32 CellIndex) const
+	inline FBox2D GetWorldCellBox2D(int32 CellIndex) const
 	{
 		return GetWorldCellBox2D(GetCellCoordX(CellIndex), GetCellCoordY(CellIndex));
 	}
 
 	/** Return the 2D bounding box of a cell. */
-	FORCEINLINE FBox2D GetWorldCellBox2D(int32 LocationX, int32 LocationY) const
+	inline FBox2D GetWorldCellBox2D(int32 LocationX, int32 LocationY) const
 	{
 		return FBox2D(
 			FVector2D(Origin) + FVector2D(LocationX * GridCellSize, LocationY * GridCellSize), 
@@ -240,7 +240,7 @@ public:
 	}
 
 	/** Return the world bounding box of all cells included in the given rectangle. */
-	FORCEINLINE FBox GetWorldCellRectangleBox(const FIntRect& CellRect) const
+	inline FBox GetWorldCellRectangleBox(const FIntRect& CellRect) const
 	{
 		return FBox(
 			Origin + FVector(CellRect.Min.X * GridCellSize, CellRect.Min.Y * GridCellSize, -BoundsSize.Z * 0.5f), 
@@ -249,7 +249,7 @@ public:
 	}
 
 	/** Compute a rectangle of cells overlapping the given WorldBox. */
-	FORCEINLINE FIntRect GetCellRectangleFromBox(const FBox& WorldBox) const
+	inline FIntRect GetCellRectangleFromBox(const FBox& WorldBox) const
 	{
 		if (!WorldBox.IsValid)
 		{
@@ -262,17 +262,17 @@ public:
 	}
 
 	/** Return an IntRect that includes all the cells of the grid. Max is inclusive. */
-	FORCEINLINE FIntRect GetGridRectangle() const
+	inline FIntRect GetGridRectangle() const
 	{
 		return FIntRect(0, 0, GridSize.Width-1, GridSize.Height-1);
 	}
 
-	FORCEINLINE FVector GetWorldCellCenter(int32 CellIndex) const
+	inline FVector GetWorldCellCenter(int32 CellIndex) const
 	{
 		return GetWorldCellCenter(GetCellCoordX(CellIndex), GetCellCoordY(CellIndex));
 	}
 
-	FORCEINLINE FVector GetWorldCellCenter(int32 LocationX, int32 LocationY) const
+	inline FVector GetWorldCellCenter(int32 LocationX, int32 LocationY) const
 	{
 		return Origin + FVector((LocationX + 0.5f) * GridCellSize, (LocationY + 0.5f) * GridCellSize, 0);
 	}
@@ -290,21 +290,21 @@ public:
 		return (CellIndex == INDEX_NONE) ? InvalidCellInstance : Cells[CellIndex];
 	}
 
-	FORCEINLINE FCellType& operator[](int32 CellIndex) { return Cells[CellIndex]; }
-	FORCEINLINE const FCellType& operator[](int32 CellIndex) const { return Cells[CellIndex]; }
+	inline FCellType& operator[](int32 CellIndex) { return Cells[CellIndex]; }
+	inline const FCellType& operator[](int32 CellIndex) const { return Cells[CellIndex]; }
 
-	FORCEINLINE FCellType& GetCellAtIndexUnsafe(int32 CellIndex) { return Cells.GetData()[CellIndex]; }
-	FORCEINLINE const FCellType& GetCellAtIndexUnsafe(int32 CellIndex) const { return Cells.GetData()[CellIndex]; }
+	inline FCellType& GetCellAtIndexUnsafe(int32 CellIndex) { return Cells.GetData()[CellIndex]; }
+	inline const FCellType& GetCellAtIndexUnsafe(int32 CellIndex) const { return Cells.GetData()[CellIndex]; }
 
 	FCellType& GetCellAtCoordsUnsafe(int32 LocationX, int32 LocationY) { return GetCellAtIndexUnsafe(GetCellIndexUnsafe(LocationX, LocationY)); }
 	const FCellType& GetCellAtCoordsUnsafe(int32 LocationX, int32 LocationY) const { return GetCellAtIndexUnsafe(GetCellIndexUnsafe(LocationX, LocationY)); }
 
-	FORCEINLINE int32 GetCellsCount() const
+	inline int32 GetCellsCount() const
 	{
 		return Cells.Num();
 	}
 
-	FORCEINLINE int32 Num() const
+	inline int32 Num() const
 	{
 		return Cells.Num();
 	}
diff --git a/Engine/Source/Runtime/Engine/Classes/AI/Navigation/AvoidanceManager.h b/Engine/Source/Runtime/Engine/Classes/AI/Navigation/AvoidanceManager.h
--- a/Engine/Source/Runtime/Engine/Classes/AI/Navigation/AvoidanceManager.h
+++ b/Engine/Source/Runtime/Engine/Classes/AI/Navigation/AvoidanceManager.h
@@ -70,12 +70,12 @@ struct FNavAvoidanceData
 		int32 InGroupMask = 1, int32 InGroupsToAvoid = 0xffffffff, int32 InGroupsToIgnore = 0,
 		float InTestRadius2D = 500.0f);
 
-	FORCEINLINE bool ShouldBeIgnored() const
+	inline bool ShouldBeIgnored() const
 	{
 		return (RemainingTimeToLive <= 0.0f);
 	}
 
-	FORCEINLINE bool ShouldIgnoreGroup(int32 OtherGroupMask) const
+	inline bool ShouldIgnoreGroup(int32 OtherGroupMask) const
 	{
 		return ((GroupsToAvoid& OtherGroupMask) == 0) || ((GroupsToIgnore & OtherGroupMask) != 0);
 	}
diff --git a/Engine/Source/Runtime/Engine/Classes/AI/Navigation/NavAgentSelector.h b/Engine/Source/Runtime/Engine/Classes/AI/Navigation/NavAgentSelector.h
--- a/Engine/Source/Runtime/Engine/Classes/AI/Navigation/NavAgentSelector.h
+++ b/Engine/Source/Runtime/Engine/Classes/AI/Navigation/NavAgentSelector.h
@@ -58,18 +58,18 @@ struct FNavAgentSelector
 
 	ENGINE_API explicit FNavAgentSelector(const uint32 InBits = AllAgentsMask);
 
-	FORCEINLINE bool Contains(int32 AgentIndex) const
+	inline bool Contains(int32 AgentIndex) const
 	{
 		return (AgentIndex >= 0 && AgentIndex < 16) ? !!(PackedBits & (1 << AgentIndex)) : false;
 	}
 
-	FORCEINLINE bool ContainsAnyAgent() const
+	inline bool ContainsAnyAgent() const
 	{
 		return PackedBits & AllAgentsMask;
 	}
 
 
-	FORCEINLINE void Set(int32 AgentIndex)
+	inline void Set(int32 AgentIndex)
 	{
 		if (AgentIndex >= 0 && AgentIndex < 16)
 		{
@@ -77,17 +77,17 @@ struct FNavAgentSelector
 		}
 	}
 
-	FORCEINLINE bool IsInitialized() const
+	inline bool IsInitialized() const
 	{
 		return (PackedBits & InitializedBit) != 0;
 	}
 
-	FORCEINLINE void MarkInitialized()
+	inline void MarkInitialized()
 	{
 		PackedBits |= InitializedBit;
 	}
 
-	FORCEINLINE void Empty()
+	inline void Empty()
 	{
 		PackedBits = 0;
 	}
diff --git a/Engine/Source/Runtime/Engine/Classes/AI/Navigation/NavLinkDefinition.h b/Engine/Source/Runtime/Engine/Classes/AI/Navigation/NavLinkDefinition.h
--- a/Engine/Source/Runtime/Engine/Classes/AI/Navigation/NavLinkDefinition.h
+++ b/Engine/Source/Runtime/Engine/Classes/AI/Navigation/NavLinkDefinition.h
@@ -218,7 +218,7 @@ struct FNavigationLink : public FNavigationLinkBase
 		: Left(InLeft), Right(InRight)
 	{}
 
-	FORCEINLINE FNavigationLink Transform(const FTransform& Transformation) const
+	inline FNavigationLink Transform(const FTransform& Transformation) const
 	{
 		FNavigationLink Result = *this;
 		Result.Left = Transformation.TransformPosition(Result.Left);
@@ -227,7 +227,7 @@ struct FNavigationLink : public FNavigationLinkBase
 		return Result;
 	}
 
-	FORCEINLINE FNavigationLink Translate(const FVector& Translation) const
+	inline FNavigationLink Translate(const FVector& Translation) const
 	{
 		FNavigationLink Result = *this;
 		Result.Left += Translation;
@@ -236,7 +236,7 @@ struct FNavigationLink : public FNavigationLinkBase
 		return Result;
 	}
 
-	FORCEINLINE FNavigationLink Rotate(const FRotator& Rotation) const
+	inline FNavigationLink Rotate(const FRotator& Rotation) const
 	{
 		FNavigationLink Result = *this;
 		
@@ -277,7 +277,7 @@ struct FNavigationSegmentLink : public FNavigationLinkBase
 		: LeftStart(InLeftStart), LeftEnd(InLeftEnd), RightStart(InRightStart), RightEnd(InRightEnd)
 	{}
 
-	FORCEINLINE FNavigationSegmentLink Transform(const FTransform& Transformation) const
+	inline FNavigationSegmentLink Transform(const FTransform& Transformation) const
 	{
 		FNavigationSegmentLink Result = *this;
 		Result.LeftStart = Transformation.TransformPosition(Result.LeftStart);
@@ -288,7 +288,7 @@ struct FNavigationSegmentLink : public FNavigationLinkBase
 		return Result;
 	}
 
-	FORCEINLINE FNavigationSegmentLink Translate(const FVector& Translation) const
+	inline FNavigationSegmentLink Translate(const FVector& Translation) const
 	{
 		FNavigationSegmentLink Result = *this;
 		Result.LeftStart += Translation;
diff --git a/Engine/Source/Runtime/Engine/Classes/AI/Navigation/NavigationAvoidanceTypes.h b/Engine/Source/Runtime/Engine/Classes/AI/Navigation/NavigationAvoidanceTypes.h
--- a/Engine/Source/Runtime/Engine/Classes/AI/Navigation/NavigationAvoidanceTypes.h
+++ b/Engine/Source/Runtime/Engine/Classes/AI/Navigation/NavigationAvoidanceTypes.h
@@ -89,27 +89,27 @@ struct FNavAvoidanceMask
 	};
 #endif
 
-	FORCEINLINE bool HasGroup(uint8 GroupID) const
+	inline bool HasGroup(uint8 GroupID) const
 	{
 		return (Packed & (1 << GroupID)) != 0;
 	}
 
-	FORCEINLINE void SetGroup(uint8 GroupID)
+	inline void SetGroup(uint8 GroupID)
 	{
 		Packed |= (1 << GroupID);
 	}
 
-	FORCEINLINE void ClearGroup(uint8 GroupID)
+	inline void ClearGroup(uint8 GroupID)
 	{
 		Packed &= ~(1 << GroupID);
 	}
 
-	FORCEINLINE void ClearAll()
+	inline void ClearAll()
 	{
 		Packed = 0;
 	}
 
-	FORCEINLINE void SetFlagsDirectly(uint32 NewFlagset)
+	inline void SetFlagsDirectly(uint32 NewFlagset)
 	{
 		Packed = NewFlagset;
 	}
diff --git a/Engine/Source/Runtime/Engine/Classes/AI/Navigation/NavigationRelevantData.h b/Engine/Source/Runtime/Engine/Classes/AI/Navigation/NavigationRelevantData.h
--- a/Engine/Source/Runtime/Engine/Classes/AI/Navigation/NavigationRelevantData.h
+++ b/Engine/Source/Runtime/Engine/Classes/AI/Navigation/NavigationRelevantData.h
@@ -114,25 +114,25 @@ struct FNavigationRelevantData : public TSharedFromThis<FNavigationRelevantData,
 	ENGINE_API FNavigationRelevantData& operator=(FNavigationRelevantData&& Other);
 	ENGINE_API FNavigationRelevantData& operator=(const FNavigationRelevantData& Other);
 
-	FORCEINLINE bool HasGeometry() const { return VoxelData.Num() || CollisionData.Num(); }
-	FORCEINLINE bool HasModifiers() const { return !Modifiers.IsEmpty(); }
-	FORCEINLINE bool HasDynamicModifiers() const { return Modifiers.IsDynamic(); }
-	FORCEINLINE bool IsPendingLazyGeometryGathering() const { return bPendingLazyGeometryGathering; }
-	FORCEINLINE bool IsPendingLazyModifiersGathering() const { return bPendingLazyModifiersGathering; }
-	FORCEINLINE bool IsPendingChildLazyModifiersGathering() const { return bPendingChildLazyModifiersGathering; }
-	FORCEINLINE bool NeedAnyPendingLazyModifiersGathering() const { return bPendingLazyModifiersGathering || bPendingChildLazyModifiersGathering; }
-	FORCEINLINE bool SupportsGatheringGeometrySlices() const { return bSupportsGatheringGeometrySlices; }
+	inline bool HasGeometry() const { return VoxelData.Num() || CollisionData.Num(); }
+	inline bool HasModifiers() const { return !Modifiers.IsEmpty(); }
+	inline bool HasDynamicModifiers() const { return Modifiers.IsDynamic(); }
+	inline bool IsPendingLazyGeometryGathering() const { return bPendingLazyGeometryGathering; }
+	inline bool IsPendingLazyModifiersGathering() const { return bPendingLazyModifiersGathering; }
+	inline bool IsPendingChildLazyModifiersGathering() const { return bPendingChildLazyModifiersGathering; }
+	inline bool NeedAnyPendingLazyModifiersGathering() const { return bPendingLazyModifiersGathering || bPendingChildLazyModifiersGathering; }
+	inline bool SupportsGatheringGeometrySlices() const { return bSupportsGatheringGeometrySlices; }
 
 	/**
 	 * Indicates that this object will not dirty the navmesh when added or removed from the octree.
 	 * In this case we expect it to manually dirty areas (e.g. using OnObjectBoundsChanged).
 	 */
-	FORCEINLINE bool ShouldSkipDirtyAreaOnAddOrRemove() const { return bShouldSkipDirtyAreaOnAddOrRemove; }
+	inline bool ShouldSkipDirtyAreaOnAddOrRemove() const { return bShouldSkipDirtyAreaOnAddOrRemove; }
 
-	FORCEINLINE bool IsEmpty() const { return !HasGeometry() && !HasModifiers(); }
-	FORCEINLINE SIZE_T GetAllocatedSize() const { return CollisionData.GetAllocatedSize() + VoxelData.GetAllocatedSize() + Modifiers.GetAllocatedSize(); }
-	FORCEINLINE SIZE_T GetGeometryAllocatedSize() const { return CollisionData.GetAllocatedSize() + VoxelData.GetAllocatedSize(); }
-	FORCEINLINE ENavigationDirtyFlag GetDirtyFlag() const
+	inline bool IsEmpty() const { return !HasGeometry() && !HasModifiers(); }
+	inline SIZE_T GetAllocatedSize() const { return CollisionData.GetAllocatedSize() + VoxelData.GetAllocatedSize() + Modifiers.GetAllocatedSize(); }
+	inline SIZE_T GetGeometryAllocatedSize() const { return CollisionData.GetAllocatedSize() + VoxelData.GetAllocatedSize(); }
+	inline ENavigationDirtyFlag GetDirtyFlag() const
 	{
 		const bool bSetGeometryFlag = HasGeometry() || IsPendingLazyGeometryGathering() ||
 			Modifiers.GetFillCollisionUnderneathForNavmesh() || Modifiers.GetMaskFillCollisionUnderneathForNavmesh() ||
diff --git a/Engine/Source/Runtime/Engine/Classes/AI/Navigation/NavigationTypes.h b/Engine/Source/Runtime/Engine/Classes/AI/Navigation/NavigationTypes.h
--- a/Engine/Source/Runtime/Engine/Classes/AI/Navigation/NavigationTypes.h
+++ b/Engine/Source/Runtime/Engine/Classes/AI/Navigation/NavigationTypes.h
@@ -46,12 +46,12 @@ namespace FNavigationSystem
 
 	inline constexpr FVector InvalidLocation = FVector(TNumericLimits<FVector::FReal>::Max(), UE::Math::TVectorConstInit{});
 
-	FORCEINLINE bool IsValidLocation(const FVector& TestLocation)
+	inline bool IsValidLocation(const FVector& TestLocation)
 	{
 		return TestLocation != InvalidLocation;
 	}
 
-	FORCEINLINE bool BoxesAreSame(const FBox& A, const FBox& B)
+	inline bool BoxesAreSame(const FBox& A, const FBox& B)
 	{
 		return FVector::PointsAreSame(A.Min, B.Min) && FVector::PointsAreSame(A.Max, B.Max);
 	}
@@ -236,15 +236,15 @@ struct FNavigationPortalEdge
 		: Left(InLeft), Right(InRight), ToRef(InToRef)
 	{}
 
-	FORCEINLINE FVector GetPoint(const int32 Index) const
+	inline FVector GetPoint(const int32 Index) const
 	{
 		check(Index >= 0 && Index < 2);
 		return ((FVector*)&Left)[Index];
 	}
 
-	FORCEINLINE FVector::FReal GetLength() const { return FVector::Dist(Left, Right); }
+	inline FVector::FReal GetLength() const { return FVector::Dist(Left, Right); }
 
-	FORCEINLINE FVector GetMiddlePoint() const { return Left + (Right - Left) / 2; }
+	inline FVector GetMiddlePoint() const { return Left + (Right - Left) / 2; }
 };
 
 /** Describes a point in navigation data */
@@ -261,9 +261,9 @@ struct FNavLocation
 		: Location(InLocation), NodeRef(InNodeRef) {}
 
 	/** checks if location has associated navigation node ref */
-	FORCEINLINE bool HasNodeRef() const { return NodeRef != INVALID_NAVNODEREF; }
+	inline bool HasNodeRef() const { return NodeRef != INVALID_NAVNODEREF; }
 
-	FORCEINLINE operator FVector() const { return Location; }
+	inline operator FVector() const { return Location; }
 
 	bool operator==(const FNavLocation& Other) const
 	{
@@ -482,12 +482,12 @@ struct FNavAgentProperties : public FMovementProperties
 
 	ENGINE_API void UpdateWithCollisionComponent(class UShapeComponent* CollisionComponent);
 
-	FORCEINLINE bool IsValid() const { return AgentRadius >= 0.f && AgentHeight >= 0.f; }
-	FORCEINLINE bool HasStepHeightOverride() const { return AgentStepHeight >= 0.0f; }
+	inline bool IsValid() const { return AgentRadius >= 0.f && AgentHeight >= 0.f; }
+	inline bool HasStepHeightOverride() const { return AgentStepHeight >= 0.0f; }
 
 	ENGINE_API bool IsNavDataMatching(const FNavAgentProperties& Other) const;
 
-	FORCEINLINE bool IsEquivalent(const FNavAgentProperties& Other, float Precision = 5.f) const
+	inline bool IsEquivalent(const FNavAgentProperties& Other, float Precision = 5.f) const
 	{
 		return FGenericPlatformMath::Abs(AgentRadius - Other.AgentRadius) < Precision
 			&& FGenericPlatformMath::Abs(AgentHeight - Other.AgentHeight) < Precision
@@ -679,14 +679,14 @@ namespace NavMeshMemory
 				}
 			}
 
-			FORCEINLINE_DEBUGGABLE void ResizeAllocation(int32 PreviousNumElements, int32 NumElements, int32 NumBytesPerElement, uint32 AlignmentOfElement)
+			inline void ResizeAllocation(int32 PreviousNumElements, int32 NumElements, int32 NumBytesPerElement, uint32 AlignmentOfElement)
 			{
 				// Maintain existing behavior, call the default aligned version of this function.
 				// We currently rely on this as we are often storing actual structs into uint8's here.
 				ResizeAllocation(PreviousNumElements, NumElements, NumBytesPerElement);
 			}
 
-			FORCEINLINE_DEBUGGABLE void ResizeAllocation(SizeType PreviousNumElements, SizeType NumElements, SIZE_T NumBytesPerElement)
+			inline void ResizeAllocation(SizeType PreviousNumElements, SizeType NumElements, SIZE_T NumBytesPerElement)
 			{
 				const int64 NewSize = NumElements * NumBytesPerElement;
 				INC_DWORD_STAT_BY(STAT_NavigationMemory, NewSize - AllocatedSize);
@@ -746,5 +746,5 @@ struct FNavHeightfieldSamples
 	ENGINE_API void GetResourceSizeEx(FResourceSizeEx& CumulativeResourceSize);
 
 	ENGINE_API void Empty();
-	FORCEINLINE bool IsEmpty() const { return Heights.IsEmpty() && Holes.IsEmpty();  }
+	inline bool IsEmpty() const { return Heights.IsEmpty() && Holes.IsEmpty();  }
 };
diff --git a/Engine/Source/Runtime/Engine/Classes/AI/NavigationSystemBase.h b/Engine/Source/Runtime/Engine/Classes/AI/NavigationSystemBase.h
--- a/Engine/Source/Runtime/Engine/Classes/AI/NavigationSystemBase.h
+++ b/Engine/Source/Runtime/Engine/Classes/AI/NavigationSystemBase.h
@@ -112,13 +112,13 @@ namespace FNavigationSystem
 	ENGINE_API bool IsEditorRunMode(const FNavigationSystemRunMode Mode);
 	
 	template<typename TNavSys>
-	FORCEINLINE TNavSys* GetCurrent(UWorld* World)
+	inline TNavSys* GetCurrent(UWorld* World)
 	{
 		return World ? Cast<TNavSys>(World->GetNavigationSystem()) : (TNavSys*)nullptr;
 	}
 
 	template<typename TNavSys>
-	FORCEINLINE const TNavSys* GetCurrent(const UWorld* World)
+	inline const TNavSys* GetCurrent(const UWorld* World)
 	{
 		return World ? Cast<TNavSys>(World->GetNavigationSystem()) : (const TNavSys*)nullptr;
 	}
diff --git a/Engine/Source/Runtime/Engine/Public/AI/Navigation/NavQueryFilter.h b/Engine/Source/Runtime/Engine/Public/AI/Navigation/NavQueryFilter.h
--- a/Engine/Source/Runtime/Engine/Public/AI/Navigation/NavQueryFilter.h
+++ b/Engine/Source/Runtime/Engine/Public/AI/Navigation/NavQueryFilter.h
@@ -77,7 +77,7 @@ public:
 	void SetMaxSearchNodes(const uint32 MaxNodes) { MaxSearchNodes = MaxNodes; }
 
 	/** get node limit for A* loop */
-	FORCEINLINE uint32 GetMaxSearchNodes() const { return MaxSearchNodes; }
+	inline uint32 GetMaxSearchNodes() const { return MaxSearchNodes; }
 
 	/** get heuristic scaling factor */
 	float GetHeuristicScale() const { return QueryFilterImpl->GetHeuristicScale(); }
@@ -98,18 +98,18 @@ public:
 		QueryFilterImpl = MakeShareable(new FilterType());
 	}
 
-	FORCEINLINE_DEBUGGABLE void SetFilterImplementation(const INavigationQueryFilterInterface* InQueryFilterImpl)
+	inline void SetFilterImplementation(const INavigationQueryFilterInterface* InQueryFilterImpl)
 	{
 		QueryFilterImpl = MakeShareable(InQueryFilterImpl->CreateCopy());
 	}
 
-	FORCEINLINE const INavigationQueryFilterInterface* GetImplementation() const { return QueryFilterImpl.Get(); }
-	FORCEINLINE INavigationQueryFilterInterface* GetImplementation() { return QueryFilterImpl.Get(); }
+	inline const INavigationQueryFilterInterface* GetImplementation() const { return QueryFilterImpl.Get(); }
+	inline INavigationQueryFilterInterface* GetImplementation() { return QueryFilterImpl.Get(); }
 	void Reset() { GetImplementation()->Reset(); }
 
 	ENGINE_API FSharedNavQueryFilter GetCopy() const;
 
-	FORCEINLINE bool operator==(const FNavigationQueryFilter& Other) const
+	inline bool operator==(const FNavigationQueryFilter& Other) const
 	{
 		const INavigationQueryFilterInterface* Impl0 = GetImplementation();
 		const INavigationQueryFilterInterface* Impl1 = Other.GetImplementation();
diff --git a/Engine/Source/Runtime/Engine/Public/AI/NavigationModifier.h b/Engine/Source/Runtime/Engine/Public/AI/NavigationModifier.h
--- a/Engine/Source/Runtime/Engine/Public/AI/NavigationModifier.h
+++ b/Engine/Source/Runtime/Engine/Public/AI/NavigationModifier.h
@@ -20,7 +20,7 @@ template<typename InElementType> class TNavStatArray;
 struct FNavigationModifier
 {
 	FNavigationModifier() : bHasMetaAreas(false) {}
-	FORCEINLINE bool HasMetaAreas() const { return !!bHasMetaAreas; }
+	inline bool HasMetaAreas() const { return !!bHasMetaAreas; }
 
 protected:
 	/** set to true if any of areas used by this modifier is a meta nav area (UNavArea::IsMetaArea is true)*/
@@ -113,16 +113,16 @@ struct FAreaNavModifier : public FNavigationModifier
 	ENGINE_API void InitializePerInstanceConvex(const TNavStatArray<FVector>& Points, const int32 FirstIndex, const int32 LastIndex, const TSubclassOf<UNavAreaBase> AreaClass);
 	ENGINE_API void InitializeConvex(const TNavStatArray<FVector>& Points, const int32 FirstIndex, const int32 LastIndex, const FTransform& LocalToWorld, const TSubclassOf<UNavAreaBase> AreaClass);
 
-	FORCEINLINE const FBox& GetBounds() const { return Bounds; }
-	FORCEINLINE ENavigationShapeType::Type GetShapeType() const { return ShapeType; }
-	FORCEINLINE ENavigationAreaMode::Type GetApplyMode() const { return ApplyMode; }
-	FORCEINLINE bool IsLowAreaModifier() const { return bIsLowAreaModifier; }
-	FORCEINLINE bool ShouldExpandTopByCellHeight() const { return bExpandTopByCellHeight; }
-	FORCEINLINE bool ShouldIncludeAgentHeight() const { return bIncludeAgentHeight; }
-	FORCEINLINE void SetExpandTopByCellHeight(bool bExpand) { bExpandTopByCellHeight = bExpand; }
-	FORCEINLINE FAreaNavModifier& SetIncludeAgentHeight(bool bInclude) { bIncludeAgentHeight = bInclude; return *this; }
-	FORCEINLINE const TSubclassOf<UNavAreaBase> GetAreaClass() const { return TSubclassOf<UNavAreaBase>(AreaClassOb.Get()); }
-	FORCEINLINE const TSubclassOf<UNavAreaBase> GetAreaClassToReplace() const { return TSubclassOf<UNavAreaBase>(ReplaceAreaClassOb.Get()); }
+	inline const FBox& GetBounds() const { return Bounds; }
+	inline ENavigationShapeType::Type GetShapeType() const { return ShapeType; }
+	inline ENavigationAreaMode::Type GetApplyMode() const { return ApplyMode; }
+	inline bool IsLowAreaModifier() const { return bIsLowAreaModifier; }
+	inline bool ShouldExpandTopByCellHeight() const { return bExpandTopByCellHeight; }
+	inline bool ShouldIncludeAgentHeight() const { return bIncludeAgentHeight; }
+	inline void SetExpandTopByCellHeight(bool bExpand) { bExpandTopByCellHeight = bExpand; }
+	inline FAreaNavModifier& SetIncludeAgentHeight(bool bInclude) { bIncludeAgentHeight = bInclude; return *this; }
+	inline const TSubclassOf<UNavAreaBase> GetAreaClass() const { return TSubclassOf<UNavAreaBase>(AreaClassOb.Get()); }
+	inline const TSubclassOf<UNavAreaBase> GetAreaClassToReplace() const { return TSubclassOf<UNavAreaBase>(ReplaceAreaClassOb.Get()); }
 
 	/** navigation area applied by this modifier */
 	ENGINE_API void SetAreaClass(const TSubclassOf<UNavAreaBase> AreaClass);
@@ -242,7 +242,7 @@ struct FSimpleLinkNavModifier : public FNavigationModifier
 	FSimpleLinkNavModifier& operator=(FSimpleLinkNavModifier&& Other) = default;
 	PRAGMA_ENABLE_DEPRECATION_WARNINGS
 
-	FORCEINLINE bool HasFallDownLinks() const { return !!bHasFallDownLinks; }
+	inline bool HasFallDownLinks() const { return !!bHasFallDownLinks; }
 	ENGINE_API void SetLinks(const TArray<FNavigationLink>& InLinks);
 	ENGINE_API void SetSegmentLinks(const TArray<FNavigationSegmentLink>& InLinks);
 	ENGINE_API void AppendLinks(const TArray<FNavigationLink>& InLinks);
@@ -263,7 +263,7 @@ struct FCustomLinkNavModifier : public FNavigationModifier
 	FTransform LocalToWorld;
 
 	ENGINE_API void Set(TSubclassOf<UNavLinkDefinition> LinkDefinitionClass, const FTransform& InLocalToWorld);
-	FORCEINLINE const TSubclassOf<UNavLinkDefinition> GetNavLinkClass() const { return TSubclassOf<UNavLinkDefinition>(LinkDefinitionClassOb.Get()); }
+	inline const TSubclassOf<UNavLinkDefinition> GetNavLinkClass() const { return TSubclassOf<UNavLinkDefinition>(LinkDefinitionClassOb.Get()); }
 
 protected:
 	TWeakObjectPtr<UClass> LinkDefinitionClassOb;
@@ -288,7 +288,7 @@ struct FCompositeNavModifier : public FNavigationModifier
 	ENGINE_API void Reset();
 	ENGINE_API void Empty();
 
-	FORCEINLINE bool IsEmpty() const 
+	inline bool IsEmpty() const 
 	{ 
 		return (Areas.Num() == 0) && (SimpleLinks.Num() == 0) && (CustomLinks.Num() == 0) && 
 			!bFillCollisionUnderneathForNavmesh && 
@@ -297,7 +297,7 @@ struct FCompositeNavModifier : public FNavigationModifier
 			NavMeshResolution == ENavigationDataResolution::Invalid;
 	}
 
-	FORCEINLINE bool IsDynamic() const 
+	inline bool IsDynamic() const 
 	{
 		// Excluding bFillCollisionUnderneathForNavmesh, bMaskFillCollisionUnderneathForNavmesh and NavMeshResolution since they require full tile rebuild.
 		return !Areas.IsEmpty() || !SimpleLinks.IsEmpty() || !CustomLinks.IsEmpty();
@@ -343,25 +343,25 @@ struct FCompositeNavModifier : public FNavigationModifier
 	ENGINE_API void CreateAreaModifiers(const UPrimitiveComponent* PrimComp, const TSubclassOf<UNavAreaBase> AreaClass, const TSubclassOf<UNavAreaBase> AreaClassToReplace = {});
 	ENGINE_API void CreateAreaModifiers(const FCollisionShape& CollisionShape, const FTransform& LocalToWorld, const TSubclassOf<UNavAreaBase> AreaClass, const bool bIncludeAgentHeight = false);
 
-	FORCEINLINE const TArray<FAreaNavModifier>& GetAreas() const { return Areas; }
-	FORCEINLINE const TArray<FSimpleLinkNavModifier>& GetSimpleLinks() const { return SimpleLinks; }
-	FORCEINLINE const TArray<FCustomLinkNavModifier>& GetCustomLinks() const { return CustomLinks; }
+	inline const TArray<FAreaNavModifier>& GetAreas() const { return Areas; }
+	inline const TArray<FSimpleLinkNavModifier>& GetSimpleLinks() const { return SimpleLinks; }
+	inline const TArray<FCustomLinkNavModifier>& GetCustomLinks() const { return CustomLinks; }
 	
-	FORCEINLINE bool HasLinks() const { return (SimpleLinks.Num() > 0) || (CustomLinks.Num() > 0); }
-	FORCEINLINE bool HasPotentialLinks() const { return bHasPotentialLinks; }
-	FORCEINLINE bool HasAgentHeightAdjust() const { return bAdjustHeight; }
-	FORCEINLINE bool HasAreas() const { return Areas.Num() > 0; }
-	FORCEINLINE bool HasLowAreaModifiers() const { return bHasLowAreaModifiers; }
-	FORCEINLINE bool IsPerInstanceModifier() const { return bIsPerInstanceModifier; }
-	FORCEINLINE bool GetFillCollisionUnderneathForNavmesh() const { return bFillCollisionUnderneathForNavmesh; }
-	FORCEINLINE void SetFillCollisionUnderneathForNavmesh(bool bValue) { bFillCollisionUnderneathForNavmesh = bValue; }
-	FORCEINLINE bool GetMaskFillCollisionUnderneathForNavmesh() const { return bMaskFillCollisionUnderneathForNavmesh; }
-	FORCEINLINE void SetMaskFillCollisionUnderneathForNavmesh(bool bValue) { bMaskFillCollisionUnderneathForNavmesh = bValue; }
-	FORCEINLINE bool GetRasterizeAsFilledConvexVolume() const { return bRasterizeAsFilledConvexVolume; }
-	FORCEINLINE void SetRasterizeAsFilledConvexVolume(bool bValue) { bRasterizeAsFilledConvexVolume = bValue; }
-	FORCEINLINE ENavigationDataResolution GetNavMeshResolution() const { return bEnableNavMeshResolutions ? NavMeshResolution : ENavigationDataResolution::Invalid; }
-	FORCEINLINE void SetNavMeshResolution(ENavigationDataResolution Resolution) { NavMeshResolution = Resolution; }
-	FORCEINLINE void ReserveForAdditionalAreas(int32 AdditionalElementsCount) { Areas.Reserve(Areas.Num() + AdditionalElementsCount); }
+	inline bool HasLinks() const { return (SimpleLinks.Num() > 0) || (CustomLinks.Num() > 0); }
+	inline bool HasPotentialLinks() const { return bHasPotentialLinks; }
+	inline bool HasAgentHeightAdjust() const { return bAdjustHeight; }
+	inline bool HasAreas() const { return Areas.Num() > 0; }
+	inline bool HasLowAreaModifiers() const { return bHasLowAreaModifiers; }
+	inline bool IsPerInstanceModifier() const { return bIsPerInstanceModifier; }
+	inline bool GetFillCollisionUnderneathForNavmesh() const { return bFillCollisionUnderneathForNavmesh; }
+	inline void SetFillCollisionUnderneathForNavmesh(bool bValue) { bFillCollisionUnderneathForNavmesh = bValue; }
+	inline bool GetMaskFillCollisionUnderneathForNavmesh() const { return bMaskFillCollisionUnderneathForNavmesh; }
+	inline void SetMaskFillCollisionUnderneathForNavmesh(bool bValue) { bMaskFillCollisionUnderneathForNavmesh = bValue; }
+	inline bool GetRasterizeAsFilledConvexVolume() const { return bRasterizeAsFilledConvexVolume; }
+	inline void SetRasterizeAsFilledConvexVolume(bool bValue) { bRasterizeAsFilledConvexVolume = bValue; }
+	inline ENavigationDataResolution GetNavMeshResolution() const { return bEnableNavMeshResolutions ? NavMeshResolution : ENavigationDataResolution::Invalid; }
+	inline void SetNavMeshResolution(ENavigationDataResolution Resolution) { NavMeshResolution = Resolution; }
+	inline void ReserveForAdditionalAreas(int32 AdditionalElementsCount) { Areas.Reserve(Areas.Num() + AdditionalElementsCount); }
 
 	void MarkPotentialLinks() { bHasPotentialLinks = true; }
     void MarkAsPerInstanceModifier() { bIsPerInstanceModifier = true; }
diff --git a/Engine/Source/Runtime/NavigationSystem/Public/NavAreas/NavArea.h b/Engine/Source/Runtime/NavigationSystem/Public/NavAreas/NavArea.h
--- a/Engine/Source/Runtime/NavigationSystem/Public/NavAreas/NavArea.h
+++ b/Engine/Source/Runtime/NavigationSystem/Public/NavAreas/NavArea.h
@@ -92,10 +92,10 @@ public:
 	NAVIGATIONSYSTEM_API virtual void PostEditChangeProperty(FPropertyChangedEvent& PropertyChangedEvent);
 #endif // WITH_EDITOR
 
-	FORCEINLINE uint16 GetAreaFlags() const { return AreaFlags; }
-	FORCEINLINE bool HasFlags(uint16 InFlags) const { return (InFlags & AreaFlags) != 0; }
+	inline uint16 GetAreaFlags() const { return AreaFlags; }
+	inline bool HasFlags(uint16 InFlags) const { return (InFlags & AreaFlags) != 0; }
 
-	FORCEINLINE bool IsSupportingAgent(int32 AgentIndex) const { return SupportedAgents.Contains(AgentIndex); }
+	inline bool IsSupportingAgent(int32 AgentIndex) const { return SupportedAgents.Contains(AgentIndex); }
 
 	/** called before adding to navigation system */
 	virtual void InitializeArea() {};
diff --git a/Engine/Source/Runtime/NavigationSystem/Public/NavLinkCustomComponent.h b/Engine/Source/Runtime/NavigationSystem/Public/NavLinkCustomComponent.h
--- a/Engine/Source/Runtime/NavigationSystem/Public/NavLinkCustomComponent.h
+++ b/Engine/Source/Runtime/NavigationSystem/Public/NavLinkCustomComponent.h
@@ -121,23 +121,23 @@ class UNavLinkCustomComponent : public UNavRelevantComponent, public INavLinkCus
 	// helper functions for setting delegates
 
 	template< class UserClass >	
-	FORCEINLINE void SetMoveReachedLink(UserClass* TargetOb, typename FOnMoveReachedLink::TMethodPtr< UserClass > InFunc)
+	inline void SetMoveReachedLink(UserClass* TargetOb, typename FOnMoveReachedLink::TMethodPtr< UserClass > InFunc)
 	{
 		SetMoveReachedLink(FOnMoveReachedLink::CreateUObject(TargetOb, InFunc));
 	}
 	template< class UserClass >	
-	FORCEINLINE void SetMoveReachedLink(UserClass* TargetOb, typename FOnMoveReachedLink::TConstMethodPtr< UserClass > InFunc)
+	inline void SetMoveReachedLink(UserClass* TargetOb, typename FOnMoveReachedLink::TConstMethodPtr< UserClass > InFunc)
 	{
 		SetMoveReachedLink(FOnMoveReachedLink::CreateUObject(TargetOb, InFunc));
 	}
 
 	template< class UserClass >	
-	FORCEINLINE void SetBroadcastFilter(UserClass* TargetOb, typename FBroadcastFilter::TMethodPtr< UserClass > InFunc)
+	inline void SetBroadcastFilter(UserClass* TargetOb, typename FBroadcastFilter::TMethodPtr< UserClass > InFunc)
 	{
 		SetBroadcastFilter(FBroadcastFilter::CreateUObject(TargetOb, InFunc));
 	}
 	template< class UserClass >	
-	FORCEINLINE void SetBroadcastFilter(UserClass* TargetOb, typename FBroadcastFilter::TConstMethodPtr< UserClass > InFunc)
+	inline void SetBroadcastFilter(UserClass* TargetOb, typename FBroadcastFilter::TConstMethodPtr< UserClass > InFunc)
 	{
 		SetBroadcastFilter(FBroadcastFilter::CreateUObject(TargetOb, InFunc));
 	}
diff --git a/Engine/Source/Runtime/NavigationSystem/Public/NavMesh/NavMeshPath.h b/Engine/Source/Runtime/NavigationSystem/Public/NavMesh/NavMeshPath.h
--- a/Engine/Source/Runtime/NavigationSystem/Public/NavMesh/NavMeshPath.h
+++ b/Engine/Source/Runtime/NavigationSystem/Public/NavMesh/NavMeshPath.h
@@ -47,20 +47,20 @@ struct FNavMeshPath : public FNavigationPath
 	NAVIGATIONSYSTEM_API FNavMeshPath& operator=(const FNavMeshPath& Other);
 	NAVIGATIONSYSTEM_API FNavMeshPath& operator=(FNavMeshPath&& Other);
 
-	FORCEINLINE void SetWantsStringPulling(const bool bNewWantsStringPulling) { bWantsStringPulling = bNewWantsStringPulling; }
-	FORCEINLINE bool WantsStringPulling() const { return bWantsStringPulling; }
-	FORCEINLINE bool IsStringPulled() const { return bStringPulled; }
+	inline void SetWantsStringPulling(const bool bNewWantsStringPulling) { bWantsStringPulling = bNewWantsStringPulling; }
+	inline bool WantsStringPulling() const { return bWantsStringPulling; }
+	inline bool IsStringPulled() const { return bStringPulled; }
 
 	/** find string pulled path from PathCorridor */
 	NAVIGATIONSYSTEM_API void PerformStringPulling(const FVector& StartLoc, const FVector& EndLoc);
 
-	FORCEINLINE void SetWantsPathCorridor(const bool bNewWantsPathCorridor) { bWantsPathCorridor = bNewWantsPathCorridor; }
-	FORCEINLINE bool WantsPathCorridor() const { return bWantsPathCorridor; }
+	inline void SetWantsPathCorridor(const bool bNewWantsPathCorridor) { bWantsPathCorridor = bNewWantsPathCorridor; }
+	inline bool WantsPathCorridor() const { return bWantsPathCorridor; }
 
-	FORCEINLINE const TArray<FNavigationPortalEdge>& GetPathCorridorEdges() const { return bCorridorEdgesGenerated ? PathCorridorEdges : GeneratePathCorridorEdges(); }
-	FORCEINLINE void SetPathCorridorEdges(const TArray<FNavigationPortalEdge>& InPathCorridorEdges) { PathCorridorEdges = InPathCorridorEdges; bCorridorEdgesGenerated = true; }
+	inline const TArray<FNavigationPortalEdge>& GetPathCorridorEdges() const { return bCorridorEdgesGenerated ? PathCorridorEdges : GeneratePathCorridorEdges(); }
+	inline void SetPathCorridorEdges(const TArray<FNavigationPortalEdge>& InPathCorridorEdges) { PathCorridorEdges = InPathCorridorEdges; bCorridorEdgesGenerated = true; }
 
-	FORCEINLINE void OnPathCorridorUpdated() { bCorridorEdgesGenerated = false; }
+	inline void OnPathCorridorUpdated() { bCorridorEdgesGenerated = false; }
 
 	NAVIGATIONSYSTEM_API virtual void DebugDraw(const ANavigationData* NavData, const FColor PathColor, UCanvas* Canvas, const bool bPersistent, const float LifeTime, const uint32 NextPathPointIndex = 0) const override;
 
@@ -91,12 +91,12 @@ struct FNavMeshPath : public FNavigationPath
 		return TotalCost;
 	}
 
-	FORCEINLINE_DEBUGGABLE FVector::FReal GetTotalPathLength() const
+	inline FVector::FReal GetTotalPathLength() const
 	{
 		return bStringPulled ? GetStringPulledLength(0) : GetPathCorridorLength(0);
 	}
 
-	FORCEINLINE int32 GetNodeRefIndex(const NavNodeRef NodeRef) const { return PathCorridor.Find(NodeRef); }
+	inline int32 GetNodeRefIndex(const NavNodeRef NodeRef) const { return PathCorridor.Find(NodeRef); }
 
 	/** check if path (all polys in corridor) contains given node */
 	virtual bool ContainsNode(NavNodeRef NodeRef) const override { return PathCorridor.Contains(NodeRef); }
diff --git a/Engine/Source/Runtime/NavigationSystem/Public/NavMesh/NavTestRenderingComponent.h b/Engine/Source/Runtime/NavigationSystem/Public/NavMesh/NavTestRenderingComponent.h
--- a/Engine/Source/Runtime/NavigationSystem/Public/NavMesh/NavTestRenderingComponent.h
+++ b/Engine/Source/Runtime/NavigationSystem/Public/NavMesh/NavTestRenderingComponent.h
@@ -41,11 +41,11 @@ public:
 		uint32 bModified : 1;
 		uint32 bOffMeshLink : 1;
 
-		FORCEINLINE bool operator==(const FNodeDebugData& Other) const
+		inline bool operator==(const FNodeDebugData& Other) const
 		{
 			return PolyRef == Other.PolyRef;
 		}
-		FORCEINLINE friend uint32 GetTypeHash(const FNodeDebugData& Other)
+		inline friend uint32 GetTypeHash(const FNodeDebugData& Other)
 		{
 			return ::GetTypeHash(Other.PolyRef);
 		}
@@ -59,7 +59,7 @@ public:
 
 	void GatherPathStep();
 
-	FORCEINLINE static bool LocationInView(const FVector& Location, const FSceneView* View)
+	inline static bool LocationInView(const FVector& Location, const FSceneView* View)
 	{
 		return View->ViewFrustum.IntersectBox(Location, FVector::ZeroVector);
 	}
diff --git a/Engine/Source/Runtime/NavigationSystem/Public/NavMesh/PImplRecastNavMesh.h b/Engine/Source/Runtime/NavigationSystem/Public/NavMesh/PImplRecastNavMesh.h
--- a/Engine/Source/Runtime/NavigationSystem/Public/NavMesh/PImplRecastNavMesh.h
+++ b/Engine/Source/Runtime/NavigationSystem/Public/NavMesh/PImplRecastNavMesh.h
@@ -96,7 +96,7 @@ public:
 	NAVIGATIONSYSTEM_API void GetNavMeshTilesIn(const TArray<FBox>& InclusionBounds, TArray<FNavTileRef>& OutRefs) const;
 
 	/** Retrieves number of tiles in this navmesh */
-	FORCEINLINE int32 GetNavMeshTilesCount() const { return DetourNavMesh ? DetourNavMesh->getMaxTiles() : 0; }
+	inline int32 GetNavMeshTilesCount() const { return DetourNavMesh ? DetourNavMesh->getMaxTiles() : 0; }
 
 	/** Supported queries */
 
diff --git a/Engine/Source/Runtime/NavigationSystem/Public/NavMesh/RecastNavMesh.h b/Engine/Source/Runtime/NavigationSystem/Public/NavMesh/RecastNavMesh.h
--- a/Engine/Source/Runtime/NavigationSystem/Public/NavMesh/RecastNavMesh.h
+++ b/Engine/Source/Runtime/NavigationSystem/Public/NavMesh/RecastNavMesh.h
@@ -155,11 +155,11 @@ struct FRecastDebugPathfindingNode
 	FRecastDebugPathfindingNode() : PolyRef(0), ParentRef(0), NumVerts(0), bOpenSet(0), bOffMeshLink(0), bModified(0) {}
 	FRecastDebugPathfindingNode(NavNodeRef InPolyRef) : PolyRef(InPolyRef), ParentRef(0), NumVerts(0), bOpenSet(0), bOffMeshLink(0), bModified(0) {}
 
-	FORCEINLINE bool operator==(const NavNodeRef& OtherPolyRef) const { return PolyRef == OtherPolyRef; }
-	FORCEINLINE bool operator==(const FRecastDebugPathfindingNode& Other) const { return PolyRef == Other.PolyRef; }
-	FORCEINLINE friend uint32 GetTypeHash(const FRecastDebugPathfindingNode& Other) { return GetTypeHash(Other.PolyRef); }
+	inline bool operator==(const NavNodeRef& OtherPolyRef) const { return PolyRef == OtherPolyRef; }
+	inline bool operator==(const FRecastDebugPathfindingNode& Other) const { return PolyRef == Other.PolyRef; }
+	inline friend uint32 GetTypeHash(const FRecastDebugPathfindingNode& Other) { return GetTypeHash(Other.PolyRef); }
 
-	FORCEINLINE FVector::FReal GetHeuristicCost() const { return TotalCost - Cost; }
+	inline FVector::FReal GetHeuristicCost() const { return TotalCost - Cost; }
 };
 
 namespace ERecastDebugPathfindingFlags
@@ -285,7 +285,7 @@ struct FNavTileRef
 
 	bool IsValid() const { return TileRef != (uint64)FNavTileRef(); }
 
-	FORCEINLINE friend uint32 GetTypeHash(const FNavTileRef& NavTileRef) { return GetTypeHash(NavTileRef.TileRef); }
+	inline friend uint32 GetTypeHash(const FNavTileRef& NavTileRef) { return GetTypeHash(NavTileRef.TileRef); }
 
 	/** Those 2 functions are used for backward compatibility of the following deprecated functions in FRecastNavMeshGenerator and ARecastNavMesh:
 	*	  RemoveTileLayers
@@ -505,29 +505,29 @@ struct FNavMeshTileData
 	
 	explicit FNavMeshTileData(uint8* RawData, int32 RawDataSize, int32 LayerIdx = 0, FBox LayerBounds = FBox(ForceInit));
 		
-	FORCEINLINE uint8* GetData()
+	inline uint8* GetData()
 	{
 		check(NavData.IsValid());
 		return NavData->GetMutableRawNavData();
 	}
 
-	FORCEINLINE const uint8* GetData() const
+	inline const uint8* GetData() const
 	{
 		check(NavData.IsValid());
 		return NavData->GetRawNavData();
 	}
 
-	FORCEINLINE uint8* GetDataSafe()
+	inline uint8* GetDataSafe()
 	{
 		return NavData.IsValid() ? NavData->GetMutableRawNavData() : NULL;
 	}
 
-	FORCEINLINE bool operator==(const uint8* RawData) const
+	inline bool operator==(const uint8* RawData) const
 	{
 		return GetData() == RawData;
 	}
 
-	FORCEINLINE bool IsValid() const { return NavData.IsValid() && GetData() != nullptr && DataSize > 0; }
+	inline bool IsValid() const { return NavData.IsValid() && GetData() != nullptr && DataSize > 0; }
 
 	uint8* Release();
 
@@ -941,9 +941,9 @@ public:
 			FMemory::Memzero(CorridorCost);
 		}
 
-		FORCEINLINE int32 GetMaxCorridorSize() const { return MAX_PATH_CORRIDOR_POLYS; }
-		FORCEINLINE bool HasHit() const { return HitTime != TNumericLimits<FVector::FReal>::Max(); }
-		FORCEINLINE NavNodeRef GetLastNodeRef() const { return CorridorPolysCount > 0 ? CorridorPolys[CorridorPolysCount - 1] : INVALID_NAVNODEREF; }
+		inline int32 GetMaxCorridorSize() const { return MAX_PATH_CORRIDOR_POLYS; }
+		inline bool HasHit() const { return HitTime != TNumericLimits<FVector::FReal>::Max(); }
+		inline NavNodeRef GetLastNodeRef() const { return CorridorPolysCount > 0 ? CorridorPolys[CorridorPolysCount - 1] : INVALID_NAVNODEREF; }
 	};
 
 	//----------------------------------------------------------------------//
@@ -960,11 +960,11 @@ public:
 	/** broadcast for navmesh updates */
 	FOnNavMeshUpdate OnNavMeshUpdate;
 
-	FORCEINLINE static void SetDrawDistance(FVector::FReal NewDistance) { DrawDistanceSq = NewDistance * NewDistance; }
-	FORCEINLINE static FVector::FReal GetDrawDistanceSq() { return DrawDistanceSq; }
+	inline static void SetDrawDistance(FVector::FReal NewDistance) { DrawDistanceSq = NewDistance * NewDistance; }
+	inline static FVector::FReal GetDrawDistanceSq() { return DrawDistanceSq; }
 
-	FORCEINLINE static void SetMinimumSizeForChaosNavMeshInfluence(float NewSize) { MinimumSizeForChaosNavMeshInfluenceSq = NewSize * NewSize; }
-	FORCEINLINE static float GetMinimumSizeForChaosNavMeshInfluenceSq() { return MinimumSizeForChaosNavMeshInfluenceSq; }
+	inline static void SetMinimumSizeForChaosNavMeshInfluence(float NewSize) { MinimumSizeForChaosNavMeshInfluenceSq = NewSize * NewSize; }
+	inline static float GetMinimumSizeForChaosNavMeshInfluenceSq() { return MinimumSizeForChaosNavMeshInfluenceSq; }
 	
 
 	//////////////////////////////////////////////////////////////////////////
@@ -1183,7 +1183,7 @@ public:
 protected:
 	NAVIGATIONSYSTEM_API virtual void FillConfig(FNavDataConfig& Dest) override;
 
-	FORCEINLINE const FNavigationQueryFilter& GetRightFilterRef(FSharedConstNavQueryFilter Filter) const 
+	inline const FNavigationQueryFilter& GetRightFilterRef(FSharedConstNavQueryFilter Filter) const 
 	{
 		return *(Filter.IsValid() ? Filter.Get() : GetDefaultQueryFilter().Get());
 	}
@@ -1441,7 +1441,7 @@ public:
 	NAVIGATIONSYSTEM_API int32 DebugPathfinding(const FPathFindingQuery& Query, TArray<FRecastDebugPathfindingData>& Steps);
 
 	static NAVIGATIONSYSTEM_API const FRecastQueryFilter* GetNamedFilter(ERecastNamedFilter::Type FilterType);
-	FORCEINLINE static FNavPolyFlags GetNavLinkFlag() { return NavLinkFlag; }
+	inline static FNavPolyFlags GetNavLinkFlag() { return NavLinkFlag; }
 	
 	NAVIGATIONSYSTEM_API virtual bool NeedsRebuild() const override;
 	NAVIGATIONSYSTEM_API virtual bool SupportsRuntimeGeneration() const override;
@@ -1586,13 +1586,13 @@ public:
 };
 
 #if WITH_RECAST
-FORCEINLINE
+inline
 bool ARecastNavMesh::NavMeshRaycast(const ANavigationData* Self, const FVector& RayStart, const FVector& RayEnd, FVector& HitLocation, FSharedConstNavQueryFilter QueryFilter, const UObject* Querier)
 {
 	return NavMeshRaycast(Self, RayStart, RayEnd, HitLocation, nullptr, QueryFilter, Querier);
 }
 
-FORCEINLINE
+inline
 bool ARecastNavMesh::NavMeshRaycast(const ANavigationData* Self, const FVector& RayStart, const FVector& RayEnd, FVector& HitLocation, FNavigationRaycastAdditionalResults* AdditionalResults, FSharedConstNavQueryFilter QueryFilter, const UObject* Querier)
 {
 	FRaycastResult Result;
diff --git a/Engine/Source/Runtime/NavigationSystem/Public/NavMesh/RecastNavMeshGenerator.h b/Engine/Source/Runtime/NavigationSystem/Public/NavMesh/RecastNavMeshGenerator.h
--- a/Engine/Source/Runtime/NavigationSystem/Public/NavMesh/RecastNavMeshGenerator.h
+++ b/Engine/Source/Runtime/NavigationSystem/Public/NavMesh/RecastNavMeshGenerator.h
@@ -257,13 +257,13 @@ struct FRcTileBox
 		YMax = CalcMaxCoordExclusive((RcAreaBounds.Max.Z - RcNavMeshOrigin.Z) / TileSizeInWorldUnits, YMin);
 	}
 
-	FORCEINLINE bool Contains(const FIntPoint& Point) const
+	inline bool Contains(const FIntPoint& Point) const
 	{
 		return Point.X >= XMin && Point.X <= XMax
 			&& Point.Y >= YMin && Point.Y <= YMax;
 	}
 
-	static FORCEINLINE int32 ClampToInt32(const int64 Value)
+	static inline int32 ClampToInt32(const int64 Value)
 	{
 #if !NO_LOGGING
 		UE_CLOG(!IntFitsIn<int32>(Value), LogNavigation, Warning,
@@ -379,14 +379,14 @@ public:
 	/** Does the work involved with regenerating this tile */
 	NAVIGATIONSYSTEM_API bool DoWork();
 
-	FORCEINLINE int32 GetTileX() const { return TileX; }
-	FORCEINLINE int32 GetTileY() const { return TileY; }
-	FORCEINLINE const FBox& GetTileBB() const { return TileBB; }
+	inline int32 GetTileX() const { return TileX; }
+	inline int32 GetTileY() const { return TileY; }
+	inline const FBox& GetTileBB() const { return TileBB; }
 	/** Whether specified layer was updated */
-	FORCEINLINE bool IsLayerChanged(int32 LayerIdx) const { return DirtyLayers[LayerIdx]; }
-	FORCEINLINE const TBitArray<>& GetDirtyLayersMask() const { return DirtyLayers; }
+	inline bool IsLayerChanged(int32 LayerIdx) const { return DirtyLayers[LayerIdx]; }
+	inline const TBitArray<>& GetDirtyLayersMask() const { return DirtyLayers; }
 	/** Whether tile data was fully regenerated */
-	FORCEINLINE bool IsFullyRegenerated() const { return bRegenerateCompressedLayers; }
+	inline bool IsFullyRegenerated() const { return bRegenerateCompressedLayers; }
 	/** Whether tile task has anything to build */
 	NAVIGATIONSYSTEM_API bool HasDataToBuild() const;
 
@@ -655,7 +655,7 @@ struct FRecastTileGeneratorWrapper : public FNonAbandonableTask
 		TileGenerator->DoWork();
 	}
 
-	FORCEINLINE TStatId GetStatId() const
+	inline TStatId GetStatId() const
 	{
 		RETURN_QUICK_DECLARE_CYCLE_STAT(FRecastTileGenerator, STATGROUP_ThreadPoolAsyncTasks);
 	}
@@ -823,7 +823,7 @@ public:
 	/** Checks if a given tile is being build or has just finished building */
 	NAVIGATIONSYSTEM_API bool IsTileChanged(const FNavTileRef InTileRef) const;
 		
-	FORCEINLINE uint32 GetVersion() const { return Version; }
+	inline uint32 GetVersion() const { return Version; }
 
 	const ARecastNavMesh* GetOwner() const { return DestNavMesh; }
 
@@ -832,7 +832,7 @@ public:
 	NAVIGATIONSYSTEM_API void OnAreaRemoved(const UClass* AreaClass);
 		
 	//--- accessors --- //
-	FORCEINLINE class UWorld* GetWorld() const { return DestNavMesh->GetWorld(); }
+	inline class UWorld* GetWorld() const { return DestNavMesh->GetWorld(); }
 
 	const FRecastBuildConfig& GetConfig() const { return Config; }
 
diff --git a/Engine/Source/Runtime/NavigationSystem/Public/NavigationData.h b/Engine/Source/Runtime/NavigationSystem/Public/NavigationData.h
--- a/Engine/Source/Runtime/NavigationSystem/Public/NavigationData.h
+++ b/Engine/Source/Runtime/NavigationSystem/Public/NavigationData.h
@@ -68,11 +68,11 @@ struct FPathFindingResult
 	FPathFindingResult(ENavigationQueryResult::Type InResult = ENavigationQueryResult::Invalid) : Result(InResult)
 	{ }
 
-	FORCEINLINE bool IsSuccessful() const
+	inline bool IsSuccessful() const
 	{
 		return Result == ENavigationQueryResult::Success;
 	}
-	FORCEINLINE bool IsPartial() const;
+	inline bool IsPartial() const;
 };
 
 struct FNavigationRaycastAdditionalResults
@@ -97,23 +97,23 @@ struct FNavigationPath : public TSharedFromThis<FNavigationPath, ESPMode::Thread
 	FNavigationPath& operator=(const FNavigationPath& Other) = default;
 	PRAGMA_ENABLE_DEPRECATION_WARNINGS
 
-	FORCEINLINE bool IsValid() const
+	inline bool IsValid() const
 	{
 		return bIsReady && PathPoints.Num() > 1 && bUpToDate;
 	}
-	FORCEINLINE bool IsUpToDate() const
+	inline bool IsUpToDate() const
 	{
 		return bUpToDate;
 	}
-	FORCEINLINE bool IsReady() const
+	inline bool IsReady() const
 	{
 		return bIsReady;
 	}
-	FORCEINLINE bool IsPartial() const
+	inline bool IsPartial() const
 	{
 		return bIsPartial;
 	}
-	FORCEINLINE bool DidSearchReachedLimit() const
+	inline bool DidSearchReachedLimit() const
 	{
 		return bReachedSearchLimit;
 	}
@@ -131,74 +131,74 @@ struct FNavigationPath : public TSharedFromThis<FNavigationPath, ESPMode::Thread
 	{
 		return bErrorEndLocationNonNavigable;
 	}
-	FORCEINLINE bool IsWaitingForRepath() const
+	inline bool IsWaitingForRepath() const
 	{
 		return bWaitingForRepath;
 	}
-	FORCEINLINE void SetManualRepathWaiting(const bool bInWaitingForRepath)
+	inline void SetManualRepathWaiting(const bool bInWaitingForRepath)
 	{
 		bWaitingForRepath = bInWaitingForRepath;
 	}
-	FORCEINLINE bool ShouldUpdateStartPointOnRepath() const
+	inline bool ShouldUpdateStartPointOnRepath() const
 	{
 		return bUpdateStartPointOnRepath;
 	}
-	FORCEINLINE bool ShouldUpdateEndPointOnRepath() const
+	inline bool ShouldUpdateEndPointOnRepath() const
 	{
 		return bUpdateEndPointOnRepath;
 	}
-	FORCEINLINE FVector GetDestinationLocation() const
+	inline FVector GetDestinationLocation() const
 	{
 		return IsValid() ? PathPoints.Last().Location : INVALID_NAVEXTENT;
 	}
-	FORCEINLINE FPathObserverDelegate& GetObserver()
+	inline FPathObserverDelegate& GetObserver()
 	{
 		return ObserverDelegate;
 	}
-	FORCEINLINE FDelegateHandle AddObserver(FPathObserverDelegate::FDelegate NewObserver)
+	inline FDelegateHandle AddObserver(FPathObserverDelegate::FDelegate NewObserver)
 	{
 		return ObserverDelegate.Add(NewObserver);
 	}
-	FORCEINLINE void RemoveObserver(FDelegateHandle HandleOfObserverToRemove)
+	inline void RemoveObserver(FDelegateHandle HandleOfObserverToRemove)
 	{
 		ObserverDelegate.Remove(HandleOfObserverToRemove);
 	}
 
-	FORCEINLINE void MarkReady()
+	inline void MarkReady()
 	{
 		bIsReady = true;
 	}
 
-	NAVIGATIONSYSTEM_API FORCEINLINE void SetNavigationDataUsed(const ANavigationData* const NewData);
+	NAVIGATIONSYSTEM_API inline void SetNavigationDataUsed(const ANavigationData* const NewData);
 
-	FORCEINLINE ANavigationData* GetNavigationDataUsed() const
+	inline ANavigationData* GetNavigationDataUsed() const
 	{
 		return NavigationDataUsed.Get();
 	}
-	FORCEINLINE void SetQuerier(const UObject* InQuerier)
+	inline void SetQuerier(const UObject* InQuerier)
 	{
 		PathFindingQueryData.Owner = InQuerier;
 	}
-	FORCEINLINE const UObject* GetQuerier() const
+	inline const UObject* GetQuerier() const
 	{
 		return PathFindingQueryData.Owner.Get();
 	}
-	FORCEINLINE void SetQueryData(const FPathFindingQueryData& QueryData)
+	inline void SetQueryData(const FPathFindingQueryData& QueryData)
 	{
 		PathFindingQueryData = QueryData;
 	}
-	FORCEINLINE FPathFindingQueryData GetQueryData() const
+	inline FPathFindingQueryData GetQueryData() const
 	{
 		// return copy of query data
 		return PathFindingQueryData;
 	}
 
 	//FORCEINLINE void SetObserver(const FPathObserverDelegate& Observer) { ObserverDelegate = Observer; }
-	FORCEINLINE void SetIsPartial(const bool bPartial)
+	inline void SetIsPartial(const bool bPartial)
 	{
 		bIsPartial = bPartial;
 	}
-	FORCEINLINE void SetSearchReachedLimit(const bool bLimited)
+	inline void SetSearchReachedLimit(const bool bLimited)
 	{
 		bReachedSearchLimit = bLimited;
 	}
@@ -210,16 +210,16 @@ struct FNavigationPath : public TSharedFromThis<FNavigationPath, ESPMode::Thread
 	{
 		bErrorEndLocationNonNavigable = bErrorNonNavigable;
 	}
-	FORCEINLINE void SetFilter(FSharedConstNavQueryFilter InFilter)
+	inline void SetFilter(FSharedConstNavQueryFilter InFilter)
 	{
 		PathFindingQueryData.QueryFilter = InFilter;
 		Filter = InFilter;
 	}
-	FORCEINLINE FSharedConstNavQueryFilter GetFilter() const
+	inline FSharedConstNavQueryFilter GetFilter() const
 	{
 		return PathFindingQueryData.QueryFilter;
 	}
-	FORCEINLINE AActor* GetBaseActor() const
+	inline AActor* GetBaseActor() const
 	{
 		return Base.Get();
 	}
@@ -233,7 +233,7 @@ struct FNavigationPath : public TSharedFromThis<FNavigationPath, ESPMode::Thread
 		return PathPoints.Num() > 0 ? PathPoints.Last().Location : FNavigationSystem::InvalidLocation;
 	}
 
-	FORCEINLINE void DoneUpdating(ENavPathUpdateType::Type UpdateType)
+	inline void DoneUpdating(ENavPathUpdateType::Type UpdateType)
 	{
 		static const ENavPathEvent::Type PathUpdateTypeToPathEvent[] = {
 			ENavPathEvent::UpdatedDueToGoalMoved // GoalMoved,
@@ -254,8 +254,8 @@ struct FNavigationPath : public TSharedFromThis<FNavigationPath, ESPMode::Thread
 		ObserverDelegate.Broadcast(this, PathUpdateTypeToPathEvent[uint8(UpdateType)]);
 	}
 
-	FORCEINLINE double GetTimeStamp() const { return LastUpdateTimeStamp; }
-	FORCEINLINE void SetTimeStamp(double TimeStamp) { LastUpdateTimeStamp = TimeStamp; }
+	inline double GetTimeStamp() const { return LastUpdateTimeStamp; }
+	inline void SetTimeStamp(double TimeStamp) { LastUpdateTimeStamp = TimeStamp; }
 
 	NAVIGATIONSYSTEM_API void Invalidate();
 	NAVIGATIONSYSTEM_API void RePathFailed();
@@ -298,7 +298,7 @@ struct FNavigationPath : public TSharedFromThis<FNavigationPath, ESPMode::Thread
 		return 0.;
 	}
 
-	FORCEINLINE FVector::FReal GetCost() const
+	inline FVector::FReal GetCost() const
 	{
 		return GetCostFromIndex(0);
 	}
@@ -306,7 +306,7 @@ struct FNavigationPath : public TSharedFromThis<FNavigationPath, ESPMode::Thread
 	/** calculates total length of segments from NextPathPoint to the end of path, plus distance from CurrentPosition to NextPathPoint */
 	NAVIGATIONSYSTEM_API virtual FVector::FReal GetLengthFromPosition(FVector SegmentStart, uint32 NextPathPointIndex) const;
 
-	FORCEINLINE FVector::FReal GetLength() const
+	inline FVector::FReal GetLength() const
 	{
 		return PathPoints.Num() ? GetLengthFromPosition(PathPoints[0].Location, 1) : 0.0f;
 	}
@@ -322,11 +322,11 @@ struct FNavigationPath : public TSharedFromThis<FNavigationPath, ESPMode::Thread
 		return false;
 	}
 
-	FORCEINLINE const TArray<FNavPathPoint>& GetPathPoints() const
+	inline const TArray<FNavPathPoint>& GetPathPoints() const
 	{
 		return PathPoints;
 	}
-	FORCEINLINE TArray<FNavPathPoint>& GetPathPoints()
+	inline TArray<FNavPathPoint>& GetPathPoints()
 	{
 		return PathPoints;
 	}
@@ -354,13 +354,13 @@ public:
 
 	/** type safe casts */
 	template<typename PathClass>
-	FORCEINLINE const PathClass* CastPath() const
+	inline const PathClass* CastPath() const
 	{
 		return PathType.IsA(PathClass::Type) ? static_cast<const PathClass*>(this) : NULL;
 	}
 
 	template<typename PathClass>
-	FORCEINLINE PathClass* CastPath()
+	inline PathClass* CastPath()
 	{
 		return PathType.IsA(PathClass::Type) ? static_cast<PathClass*>(this) : NULL;
 	}
@@ -624,11 +624,11 @@ private:
 public:
 	NAVIGATIONSYSTEM_API virtual void CleanUpAndMarkPendingKill();
 
-	FORCEINLINE bool IsRegistered() const { return bRegistered; }
+	inline bool IsRegistered() const { return bRegistered; }
 	NAVIGATIONSYSTEM_API virtual void OnRegistered();
 	NAVIGATIONSYSTEM_API void OnUnregistered();
 	
-	FORCEINLINE uint16 GetNavDataUniqueID() const { return NavDataUniqueID; }
+	inline uint16 GetNavDataUniqueID() const { return NavDataUniqueID; }
 
 	NAVIGATIONSYSTEM_API virtual void TickActor(float DeltaTime, enum ELevelTick TickType, FActorTickFunction& ThisTickFunction) override;
 
@@ -652,8 +652,8 @@ public:
 	// Generation & data access                                                      
 	//----------------------------------------------------------------------//
 public:
-	FORCEINLINE const FNavDataConfig& GetConfig() const { return NavDataConfig; }
-	FORCEINLINE ERuntimeGenerationType GetRuntimeGenerationMode() const { return RuntimeGeneration; }
+	inline const FNavDataConfig& GetConfig() const { return NavDataConfig; }
+	inline ERuntimeGenerationType GetRuntimeGenerationMode() const { return RuntimeGeneration; }
 	/** Populates NavDataConfig and sets NavAgentProperties with the Src config
 	 *  Should be used when initially configuring the navdata and not when updating
 	 *  values used in NavDataConfig */
@@ -788,7 +788,7 @@ public:
 	//----------------------------------------------------------------------//
 	NAVIGATIONSYSTEM_API void DrawDebugPath(FNavigationPath* Path, const FColor PathColor = FColor::White, class UCanvas* Canvas = nullptr, const bool bPersistent = true, const float LifeTime = -1.f, const uint32 NextPathPointIndex = 0) const;
 
-	FORCEINLINE bool IsDrawingEnabled() const { return bEnableDrawing; }
+	inline bool IsDrawingEnabled() const { return bEnableDrawing; }
 
 	/** @return Total mem counted, including super calls. */
 	NAVIGATIONSYSTEM_API virtual uint32 LogMemUsed() const;
@@ -806,9 +806,9 @@ public:
 	//----------------------------------------------------------------------//
 	// Querying                                                                
 	//----------------------------------------------------------------------//
-	FORCEINLINE FSharedConstNavQueryFilter GetDefaultQueryFilter() const { return DefaultQueryFilter; }
-	FORCEINLINE const class INavigationQueryFilterInterface* GetDefaultQueryFilterImpl() const { return DefaultQueryFilter->GetImplementation(); }	
-	FORCEINLINE FVector GetDefaultQueryExtent() const { return NavDataConfig.DefaultQueryExtent; }
+	inline FSharedConstNavQueryFilter GetDefaultQueryFilter() const { return DefaultQueryFilter; }
+	inline const class INavigationQueryFilterInterface* GetDefaultQueryFilterImpl() const { return DefaultQueryFilter->GetImplementation(); }	
+	inline FVector GetDefaultQueryExtent() const { return NavDataConfig.DefaultQueryExtent; }
 
 	/** 
 	 *	Synchronously looks for a path from @StartLocation to @EndLocation for agent with properties @AgentProperties. NavMesh actor appropriate for specified 
@@ -818,7 +818,7 @@ public:
 	 *
 	 *	@note don't make this function virtual! Look at implementation details and its comments for more info.
 	 */
-	FORCEINLINE FPathFindingResult FindPath(const FNavAgentProperties& AgentProperties, const FPathFindingQuery& Query) const
+	inline FPathFindingResult FindPath(const FNavAgentProperties& AgentProperties, const FPathFindingQuery& Query) const
 	{
 		check(FindPathImplementation);
 		// this awkward implementation avoids virtual call overhead - it's possible this function will be called a lot
@@ -833,7 +833,7 @@ public:
 	 *
 	 *	@note don't make this function virtual! Look at implementation details and its comments for more info.
 	 */
-	FORCEINLINE FPathFindingResult FindHierarchicalPath(const FNavAgentProperties& AgentProperties, const FPathFindingQuery& Query) const
+	inline FPathFindingResult FindHierarchicalPath(const FNavAgentProperties& AgentProperties, const FPathFindingQuery& Query) const
 	{
 		check(FindHierarchicalPathImplementation);
 		// this awkward implementation avoids virtual call overhead - it's possible this function will be called a lot
@@ -847,7 +847,7 @@ public:
 	 *
 	 *	@note don't make this function virtual! Look at implementation details and its comments for more info.
 	 */
-	FORCEINLINE bool TestPath(const FNavAgentProperties& AgentProperties, const FPathFindingQuery& Query, int32* NumVisitedNodes) const
+	inline bool TestPath(const FNavAgentProperties& AgentProperties, const FPathFindingQuery& Query, int32* NumVisitedNodes) const
 	{
 		check(TestPathImplementation);
 		// this awkward implementation avoids virtual call overhead - it's possible this function will be called a lot
@@ -861,7 +861,7 @@ public:
 	 *
 	 *	@note don't make this function virtual! Look at implementation details and its comments for more info.
 	 */
-	FORCEINLINE bool TestHierarchicalPath(const FNavAgentProperties& AgentProperties, const FPathFindingQuery& Query, int32* NumVisitedNodes) const
+	inline bool TestHierarchicalPath(const FNavAgentProperties& AgentProperties, const FPathFindingQuery& Query, int32* NumVisitedNodes) const
 	{
 		check(TestHierarchicalPathImplementation);
 		// this awkward implementation avoids virtual call overhead - it's possible this function will be called a lot
@@ -875,7 +875,7 @@ public:
 	 *
 	 *	@note don't make this function virtual! Look at implementation details and its comments for more info.
 	 */
-	FORCEINLINE bool Raycast(const FVector& RayStart, const FVector& RayEnd, FVector& HitLocation, FSharedConstNavQueryFilter QueryFilter, const UObject* Querier = NULL) const
+	inline bool Raycast(const FVector& RayStart, const FVector& RayEnd, FVector& HitLocation, FSharedConstNavQueryFilter QueryFilter, const UObject* Querier = NULL) const
 	{
 		return Raycast(RayStart, RayEnd, HitLocation, nullptr/*AdditionalResults*/, QueryFilter, Querier);
 	}
@@ -888,7 +888,7 @@ public:
 	 *
 	 *	@note don't make this function virtual! Look at implementation details and its comments for more info.
 	 */
-	FORCEINLINE bool Raycast(const FVector& RayStart, const FVector& RayEnd, FVector& HitLocation, FNavigationRaycastAdditionalResults* AdditionalResults, FSharedConstNavQueryFilter QueryFilter, const UObject* Querier = NULL) const
+	inline bool Raycast(const FVector& RayStart, const FVector& RayEnd, FVector& HitLocation, FNavigationRaycastAdditionalResults* AdditionalResults, FSharedConstNavQueryFilter QueryFilter, const UObject* Querier = NULL) const
 	{
 		check(RaycastImplementationWithAdditionalResults);
 		// this awkward implementation avoids virtual call overhead - it's possible this function will be called a lot
@@ -1155,7 +1155,7 @@ struct FAsyncPathFindingQuery : public FPathFindingQuery
 	FAsyncPathFindingQuery(const FPathFindingQuery& Query, const FNavPathQueryDelegate& Delegate, const EPathFindingMode::Type QueryMode);
 
 protected:
-	FORCEINLINE static uint32 GetUniqueID()
+	inline static uint32 GetUniqueID()
 	{
 		return ++LastPathFindingUniqueID;
 	}
@@ -1163,12 +1163,12 @@ protected:
 	static uint32 LastPathFindingUniqueID;
 };
 
-FORCEINLINE bool FPathFindingResult::IsPartial() const
+inline bool FPathFindingResult::IsPartial() const
 {
 	return (Result != ENavigationQueryResult::Error) && Path.IsValid() && Path->IsPartial();
 }
 
-FORCEINLINE void FNavigationPath::SetNavigationDataUsed(const ANavigationData* const NavData)
+inline void FNavigationPath::SetNavigationDataUsed(const ANavigationData* const NavData)
 {
 	NavigationDataUsed = MakeWeakObjectPtr(const_cast<ANavigationData*>(NavData));
 }
diff --git a/Engine/Source/Runtime/NavigationSystem/Public/NavigationOctreeController.h b/Engine/Source/Runtime/NavigationSystem/Public/NavigationOctreeController.h
--- a/Engine/Source/Runtime/NavigationSystem/Public/NavigationOctreeController.h
+++ b/Engine/Source/Runtime/NavigationSystem/Public/NavigationOctreeController.h
@@ -113,26 +113,26 @@ private:
 //----------------------------------------------------------------------//
 
 // deprecated
-FORCEINLINE uint32 FNavigationOctreeController::HashObject(const UObject& Object)
+inline uint32 FNavigationOctreeController::HashObject(const UObject& Object)
 {
 PRAGMA_DISABLE_DEPRECATION_WARNINGS
 	return FNavigationOctree::HashObject(Object);
 PRAGMA_ENABLE_DEPRECATION_WARNINGS
 }
 
-FORCEINLINE_DEBUGGABLE const FOctreeElementId2* FNavigationOctreeController::GetNavOctreeIdForElement(const FNavigationElementHandle Element) const
+inline const FOctreeElementId2* FNavigationOctreeController::GetNavOctreeIdForElement(const FNavigationElementHandle Element) const
 { 
 	return NavOctree.IsValid()
 		? NavOctree->ElementToOctreeId.Find(Element)
 		: nullptr;
 }
 
-FORCEINLINE_DEBUGGABLE bool FNavigationOctreeController::HasElementNavOctreeId(const FNavigationElementHandle Element) const
+inline bool FNavigationOctreeController::HasElementNavOctreeId(const FNavigationElementHandle Element) const
 {
 	return NavOctree.IsValid() && (NavOctree->ElementToOctreeId.Find(Element) != nullptr);
 }
 
-FORCEINLINE_DEBUGGABLE void FNavigationOctreeController::RemoveNode(const FOctreeElementId2 ElementId, const FNavigationElementHandle ElementHandle)
+inline void FNavigationOctreeController::RemoveNode(const FOctreeElementId2 ElementId, const FNavigationElementHandle ElementHandle)
 { 
 	if (NavOctree.IsValid())
 	{
@@ -141,32 +141,32 @@ FORCEINLINE_DEBUGGABLE void FNavigationOctreeController::RemoveNode(const FOctre
 	}
 }
 
-FORCEINLINE const FNavigationOctree* FNavigationOctreeController::GetOctree() const
+inline const FNavigationOctree* FNavigationOctreeController::GetOctree() const
 { 
 	return NavOctree.Get(); 
 }
 
-FORCEINLINE FNavigationOctree* FNavigationOctreeController::GetMutableOctree()
+inline FNavigationOctree* FNavigationOctreeController::GetMutableOctree()
 { 
 	return NavOctree.Get(); 
 }
 
-FORCEINLINE bool FNavigationOctreeController::IsNavigationOctreeLocked() const
+inline bool FNavigationOctreeController::IsNavigationOctreeLocked() const
 { 
 	return bNavOctreeLock; 
 }
 
-FORCEINLINE void FNavigationOctreeController::SetNavigationOctreeLock(bool bLock) 
+inline void FNavigationOctreeController::SetNavigationOctreeLock(bool bLock) 
 { 
 	bNavOctreeLock = bLock; 
 }
 
-FORCEINLINE bool FNavigationOctreeController::IsValidElement(const FOctreeElementId2* ElementId) const
+inline bool FNavigationOctreeController::IsValidElement(const FOctreeElementId2* ElementId) const
 {
 	return ElementId && IsValidElement(*ElementId);
 }
 
-FORCEINLINE bool FNavigationOctreeController::IsValidElement(const FOctreeElementId2& ElementId) const
+inline bool FNavigationOctreeController::IsValidElement(const FOctreeElementId2& ElementId) const
 {
 	return IsValid() && NavOctree->IsValidElementId(ElementId);
 }
diff --git a/Engine/Source/Runtime/NavigationSystem/Public/NavigationSystem.h b/Engine/Source/Runtime/NavigationSystem/Public/NavigationSystem.h
--- a/Engine/Source/Runtime/NavigationSystem/Public/NavigationSystem.h
+++ b/Engine/Source/Runtime/NavigationSystem/Public/NavigationSystem.h
@@ -568,7 +568,7 @@ public:
 	UFUNCTION(BlueprintCallable, Category = "AI|Navigation", meta=(DisplayName="Replace Area in Octree Data"))
 	NAVIGATIONSYSTEM_API bool K2_ReplaceAreaInOctreeData(const UObject* Object, TSubclassOf<UNavArea> OldArea, TSubclassOf<UNavArea> NewArea);
 
-	FORCEINLINE bool IsActiveTilesGenerationEnabled() const{ return bGenerateNavigationOnlyAroundNavigationInvokers; }
+	inline bool IsActiveTilesGenerationEnabled() const{ return bGenerateNavigationOnlyAroundNavigationInvokers; }
 	
 	/** delegate type for events that dirty the navigation data ( Params: const FBox& DirtyBounds ) */
 	DECLARE_MULTICAST_DELEGATE_OneParam(FOnNavigationDirty, const FBox&);
@@ -784,7 +784,7 @@ public:
 #endif
 	
 	NAVIGATIONSYSTEM_API const FNavDataConfig& GetDefaultSupportedAgentConfig() const;
-	FORCEINLINE const TArray<FNavDataConfig>& GetSupportedAgents() const { return SupportedAgents; }
+	inline const TArray<FNavDataConfig>& GetSupportedAgents() const { return SupportedAgents; }
 	NAVIGATIONSYSTEM_API void OverrideSupportedAgents(const TArray<FNavDataConfig>& NewSupportedAgents);
 	NAVIGATIONSYSTEM_API void SetSupportedAgentsMask(const FNavAgentSelector& InSupportedAgentsMask);
 	FNavAgentSelector GetSupportedAgentsMask() const { return SupportedAgentsMask; }
@@ -797,7 +797,7 @@ public:
 	/** checks if dirty navigation data can rebuild itself */
 	NAVIGATIONSYSTEM_API bool CanRebuildDirtyNavigation() const;
 
-	FORCEINLINE bool SupportsNavigationGeneration() const { return bSupportRebuilding; }
+	inline bool SupportsNavigationGeneration() const { return bSupportRebuilding; }
 
 	static NAVIGATIONSYSTEM_API bool DoesPathIntersectBox(const FNavigationPath* Path, const FBox& Box, uint32 StartingIndex = 0, FVector* AgentExtent = NULL);
 	static NAVIGATIONSYSTEM_API bool DoesPathIntersectBox(const FNavigationPath* Path, const FBox& Box, const FVector& AgentLocation, uint32 StartingIndex = 0, FVector* AgentExtent = NULL);
@@ -982,14 +982,14 @@ public:
 	UE_DEPRECATED(5.5, "This method is no longer used by the navigation system.")
 	NAVIGATIONSYSTEM_API static uint32 HashObject(const UObject& Object);
 
-	FORCEINLINE const FOctreeElementId2* GetNavOctreeIdForElement(FNavigationElementHandle Element) const
+	inline const FOctreeElementId2* GetNavOctreeIdForElement(FNavigationElementHandle Element) const
 	{
 		return DefaultOctreeController.GetNavOctreeIdForElement(Element);
 	}
 	UE_DEPRECATED(5.5, "Use GetNavOctreeIdForElement instead.")
 	NAVIGATIONSYSTEM_API const FOctreeElementId2* GetObjectsNavOctreeId(const UObject& Object) const;
 
-	FORCEINLINE bool HasPendingUpdateForElement(FNavigationElementHandle Element) const
+	inline bool HasPendingUpdateForElement(FNavigationElementHandle Element) const
 	{
 		return DefaultOctreeController.HasPendingUpdateForElement(Element);
 	}
@@ -1102,16 +1102,16 @@ public:
 	NAVIGATIONSYSTEM_API void OnPIEEnd();
 	
 	// @todo document
-	FORCEINLINE bool IsNavigationBuildingLocked(uint8 Flags = ~0) const { return (NavBuildingLockFlags & Flags) != 0; }
+	inline bool IsNavigationBuildingLocked(uint8 Flags = ~0) const { return (NavBuildingLockFlags & Flags) != 0; }
 
 	/** check if building is permanently locked to avoid showing navmesh building notify (due to queued dirty areas) */
-	FORCEINLINE bool IsNavigationBuildingPermanentlyLocked() const
+	inline bool IsNavigationBuildingPermanentlyLocked() const
 	{
 		return (NavBuildingLockFlags & ~ENavigationBuildLock::InitialLock) != 0; 
 	}
 
 	/** check if navigation octree updates are currently ignored */
-	FORCEINLINE bool IsNavigationOctreeLocked() const { return DefaultOctreeController.IsNavigationOctreeLocked(); }
+	inline bool IsNavigationOctreeLocked() const { return DefaultOctreeController.IsNavigationOctreeLocked(); }
 
 	// @todo document
 	UFUNCTION(BlueprintCallable, Category = "AI|Navigation")
@@ -1157,10 +1157,10 @@ public:
 	NAVIGATIONSYSTEM_API virtual void OnWorldInitDone(FNavigationSystemRunMode Mode);
 
 	/** Returns true if world has been initialized. */
-	FORCEINLINE virtual bool IsWorldInitDone() const override { return bWorldInitDone; }
-	FORCEINLINE bool IsInitialized() const { return bWorldInitDone; }
+	inline virtual bool IsWorldInitDone() const override { return bWorldInitDone; }
+	inline bool IsInitialized() const { return bWorldInitDone; }
 
-	FORCEINLINE FNavigationSystemRunMode GetRunMode() const { return OperationMode; }
+	inline FNavigationSystemRunMode GetRunMode() const { return OperationMode; }
 
 	/** adds BSP collisions of currently streamed in levels to octree */
 	NAVIGATIONSYSTEM_API void InitializeLevelCollisions();
@@ -1184,17 +1184,17 @@ public:
 	/** allow editor to toggle whether seamless navigation building is enabled */
 	static NAVIGATIONSYSTEM_API void SetNavigationAutoUpdateEnabled(bool bNewEnable, UNavigationSystemBase* InNavigationSystem);
 
-	FORCEINLINE bool IsNavigationRegisterLocked() const { return NavUpdateLockFlags != 0; }
-	FORCEINLINE bool IsNavigationUnregisterLocked() const { return NavUpdateLockFlags && !(NavUpdateLockFlags & ENavigationLockReason::AllowUnregister); }
-	FORCEINLINE bool IsNavigationUpdateLocked() const { return IsNavigationRegisterLocked(); }
-	FORCEINLINE void AddNavigationUpdateLock(uint8 Flags) { NavUpdateLockFlags |= Flags; }
-	FORCEINLINE void RemoveNavigationUpdateLock(uint8 Flags) { NavUpdateLockFlags &= ~Flags; }
+	inline bool IsNavigationRegisterLocked() const { return NavUpdateLockFlags != 0; }
+	inline bool IsNavigationUnregisterLocked() const { return NavUpdateLockFlags && !(NavUpdateLockFlags & ENavigationLockReason::AllowUnregister); }
+	inline bool IsNavigationUpdateLocked() const { return IsNavigationRegisterLocked(); }
+	inline void AddNavigationUpdateLock(uint8 Flags) { NavUpdateLockFlags |= Flags; }
+	inline void RemoveNavigationUpdateLock(uint8 Flags) { NavUpdateLockFlags &= ~Flags; }
 
 	NAVIGATIONSYSTEM_API void UpdateLevelCollision(ULevel* InLevel);
 
 #endif // WITH_EDITOR
 
-	FORCEINLINE bool IsSetUpForLazyGeometryExporting() const { return bGenerateNavigationOnlyAroundNavigationInvokers; }
+	inline bool IsSetUpForLazyGeometryExporting() const { return bGenerateNavigationOnlyAroundNavigationInvokers; }
 
 	UE_DEPRECATED(5.5, "This creation flow is deprecation. Use FNavigationSystem::AddNavigationSystemToWorld instead.")
 	static NAVIGATIONSYSTEM_API UNavigationSystemV1* CreateNavigationSystem(UWorld* WorldOwner);
@@ -1207,7 +1207,7 @@ public:
 	// Fetch the array of all nav-agent properties.
 	NAVIGATIONSYSTEM_API void GetNavAgentPropertiesArray(TArray<FNavAgentProperties>& OutNavAgentProperties) const;
 
-	static FORCEINLINE bool ShouldUpdateNavOctreeOnComponentChange()
+	static inline bool ShouldUpdateNavOctreeOnComponentChange()
 	{
 		return (bUpdateNavOctreeOnComponentChange && !bStaticRuntimeNavigation)
 #if WITH_EDITOR
@@ -1223,7 +1223,7 @@ public:
 	 */
 	NAVIGATIONSYSTEM_API virtual bool ShouldCreateNavigationSystemInstance(const UWorld* World) const override;
 
-	static FORCEINLINE bool IsNavigationSystemStatic()
+	static inline bool IsNavigationSystemStatic()
 	{
 #if WITH_EDITOR
 		return (bStaticRuntimeNavigation && !(GIsEditor && !GIsPlayInEditorWorld))
