[NavigationSystem] Make sure bounds updates for navigation elements are also propagated to the repository and not only to the registered (or pending registration) octree node.

diff --git a/Engine/Source/Runtime/NavigationSystem/Private/NavigationObjectRepository.cpp b/Engine/Source/Runtime/NavigationSystem/Private/NavigationObjectRepository.cpp
--- a/Engine/Source/Runtime/NavigationSystem/Private/NavigationObjectRepository.cpp
+++ b/Engine/Source/Runtime/NavigationSystem/Private/NavigationObjectRepository.cpp
@@ -70,7 +70,7 @@ void UNavigationObjectRepository::RemoveNavigationElement(const FNavigationEleme
 {
 	UE_MT_SCOPED_WRITE_ACCESS(NavElementAccessDetector);
 
-	TSharedPtr<const FNavigationElement> Element;
+	TSharedPtr<FNavigationElement> Element;
 	if (ensureMsgf(NavRelevantElements.RemoveAndCopyValue(Handle, Element),
 		TEXT("Navigation element can't be removed since it was not registered or already unregistered)")))
 	{
@@ -110,24 +110,26 @@ TSharedPtr<const FNavigationElement> UNavigationObjectRepository::RegisterNavRel
 	// In AActor/UActorComponent code paths it is possible that a component registration is performed more than once
 	// (i.e., Actor registering its component, then individual component registers too)
 	// In such case we update with the latest.
-	if (const TSharedPtr<const FNavigationElement> ElementPtr = GetNavigationElementForUObject(&NavRelevantObject))
+	if (const TSharedPtr<FNavigationElement> ExistingElement = GetMutableNavigationElementForUObject(&NavRelevantObject))
 	{
-		const TSharedRef<const FNavigationElement> NewElement = FNavigationElement::CreateFromNavRelevantInterface(NavRelevantInterface);
+#if !NO_LOGGING
+		const FBox PreviousBounds = ExistingElement->GetBounds();
+#endif
 		{
 			UE_MT_SCOPED_WRITE_ACCESS(NavElementAccessDetector);
-			NavRelevantElements[ElementPtr->GetHandle()] = NewElement;
+			*ExistingElement = FNavigationElement(NavRelevantInterface);
 		}
 
 		if (NotifyOnSuccess == ENotifyOnSuccess::Yes)
 		{
-			(void)OnNavigationElementAddedDelegate.ExecuteIfBound(NewElement);
+			(void)OnNavigationElementAddedDelegate.ExecuteIfBound(ExistingElement.ToSharedRef());
 		}
 
 		UE_LOG(LogNavigation, Verbose, TEXT("%hs [already registered - updating] (%s:%s) Bounds: [%s]->[%s]"), __FUNCTION__,
 			*GetNameSafe(NavRelevantObject.GetOuter()), *GetNameSafe(&NavRelevantObject),
-			*ElementPtr->GetBounds().ToString(), *NewElement->GetBounds().ToString());
+			*PreviousBounds.ToString(), *ExistingElement->GetBounds().ToString());
 
-		return NewElement.ToSharedPtr();
+		return ExistingElement;
 	}
 
 	if (NavRelevantInterface.IsNavigationRelevant())
@@ -176,10 +178,15 @@ void UNavigationObjectRepository::UnregisterNavRelevantObject(const UObject* Nav
 }
 
 TSharedPtr<const FNavigationElement> UNavigationObjectRepository::GetNavigationElementForHandle(const FNavigationElementHandle Handle) const
+{
+	return GetMutableNavigationElementForHandle(Handle);
+}
+
+TSharedPtr<FNavigationElement> UNavigationObjectRepository::GetMutableNavigationElementForHandle(const FNavigationElementHandle Handle) const
 {
 	UE_MT_SCOPED_READ_ACCESS(NavElementAccessDetector);
 
-	if (const TSharedPtr<const FNavigationElement>* Element = NavRelevantElements.Find(Handle))
+	if (const TSharedPtr<FNavigationElement>* Element = NavRelevantElements.Find(Handle))
 	{
 		return *Element;
 	}
@@ -187,9 +194,10 @@ TSharedPtr<const FNavigationElement> UNavigationObjectRepository::GetNavigationE
 	return nullptr;
 }
 
-FNavigationElementHandle UNavigationObjectRepository::GetNavigationElementHandleForUObject(const UObject* NavRelevantObject)
+FNavigationElementHandle UNavigationObjectRepository::GetNavigationElementHandleForUObject(const UObject* NavRelevantObject) const
 {
 	UE_MT_SCOPED_READ_ACCESS(NavElementAccessDetector);
+
 	if (const FNavigationElementHandle* Handle = ObjectsToHandleMap.Find(FObjectKey(Cast<UObject>(NavRelevantObject))))
 	{
 		return *Handle;
@@ -198,13 +206,18 @@ FNavigationElementHandle UNavigationObjectRepository::GetNavigationElementHandle
 	return FNavigationElementHandle::Invalid;
 }
 
-TSharedPtr<const FNavigationElement> UNavigationObjectRepository::GetNavigationElementForUObject(const UObject* NavRelevantObject)
+TSharedPtr<const FNavigationElement> UNavigationObjectRepository::GetNavigationElementForUObject(const UObject* NavRelevantObject) const
+{
+	return GetMutableNavigationElementForUObject(NavRelevantObject);
+}
+
+TSharedPtr<FNavigationElement> UNavigationObjectRepository::GetMutableNavigationElementForUObject(const UObject* NavRelevantObject) const
 {
 	UE_MT_SCOPED_READ_ACCESS(NavElementAccessDetector);
 
 	if (const FNavigationElementHandle* Handle = ObjectsToHandleMap.Find(FObjectKey(NavRelevantObject)))
 	{
-		if (const TSharedPtr<const FNavigationElement>* Element = NavRelevantElements.Find(*Handle))
+		if (const TSharedPtr<FNavigationElement>* Element = NavRelevantElements.Find(*Handle))
 		{
 			return *Element;
 		}
diff --git a/Engine/Source/Runtime/NavigationSystem/Private/NavigationObjectRepository.h b/Engine/Source/Runtime/NavigationSystem/Private/NavigationObjectRepository.h
--- a/Engine/Source/Runtime/NavigationSystem/Private/NavigationObjectRepository.h
+++ b/Engine/Source/Runtime/NavigationSystem/Private/NavigationObjectRepository.h
@@ -83,7 +83,7 @@ public:
 	 * @param NavRelevantObject The UObject for which the handle must be retrieved.
 	 * @return The handle associated with the specified UObject if already registered, an invalid handle otherwise.
 	 */
-	FNavigationElementHandle GetNavigationElementHandleForUObject(const UObject* NavRelevantObject);
+	FNavigationElementHandle GetNavigationElementHandleForUObject(const UObject* NavRelevantObject) const;
 
 	/**
 	 * Returns a shared pointer to the FNavigationElement created for a registered UObject.
@@ -91,7 +91,15 @@ public:
 	 * @return A valid shared pointer if the specified UObject is registered and an associated FNavigationElement was created,
 	 * null pointer otherwise.
 	 */
-	TSharedPtr<const FNavigationElement> GetNavigationElementForUObject(const UObject* NavRelevantObject);
+	TSharedPtr<const FNavigationElement> GetNavigationElementForUObject(const UObject* NavRelevantObject) const;
+
+	/**
+	 * Returns a shared pointer to the mutable FNavigationElement created for a registered UObject.
+	 * @param NavRelevantObject The UObject for which the FNavigationElement must be retrieved.
+	 * @return A valid shared pointer if the specified UObject is registered and an associated FNavigationElement was created,
+	 * null pointer otherwise.
+	 */
+	TSharedPtr<FNavigationElement> GetMutableNavigationElementForUObject(const UObject* NavRelevantObject) const;
 
 	/**
 	 * Returns a shared pointer to the FNavigationElement created for a registered UObject.
@@ -111,6 +119,13 @@ public:
 	 */
 	TSharedPtr<const FNavigationElement> GetNavigationElementForHandle(FNavigationElementHandle Handle) const;
 
+	/**
+	 * Returns a shared pointer to the mutable FNavigationElement associated with the provided handle.
+	 * @param Handle The handle for which the FNavigationElement must be retrieved.
+	 * @return A valid shared pointer if a FNavigationElement associated with the provided handle was registered, null pointer otherwise.
+	 */
+	TSharedPtr<FNavigationElement> GetMutableNavigationElementForHandle(FNavigationElementHandle Handle) const;
+
 	/**
 	 * Adds the provided interface to the list of registered custom navigation links.
 	 * @param CustomNavLinkObject INavLinkCustomInterface of the object to register.
@@ -175,7 +190,7 @@ private:
 	TMap<FObjectKey, FNavigationElementHandle> ObjectsToHandleMap;
 
 	/** List of registered navigation elements. */
-	TMap<FNavigationElementHandle, TSharedPtr<const FNavigationElement>> NavRelevantElements;
+	TMap<FNavigationElementHandle, TSharedPtr<FNavigationElement>> NavRelevantElements;
 
 	/** List of registered custom navigation link objects. */
 	TArray<TWeakInterfacePtr<INavLinkCustomInterface>> CustomLinkObjects;
diff --git a/Engine/Source/Runtime/NavigationSystem/Private/NavigationSystem.cpp b/Engine/Source/Runtime/NavigationSystem/Private/NavigationSystem.cpp
--- a/Engine/Source/Runtime/NavigationSystem/Private/NavigationSystem.cpp
+++ b/Engine/Source/Runtime/NavigationSystem/Private/NavigationSystem.cpp
@@ -3380,6 +3380,8 @@ void UNavigationSystemV1::RegisterComponentToNavOctree(UActorComponent* Comp)
 
 	if (UE::Navigation::Private::ShouldComponentWaitForActorToRegister(Comp))
 	{
+		UE_LOG(LogNavigation, VeryVerbose
+			, TEXT("%hs: %s registration skipped (waiting for actor to register)"), __FUNCTION__, *GetNameSafe(Comp));
 		return;
 	}
 
@@ -4048,6 +4050,11 @@ bool UNavigationSystemV1::UpdateNavOctreeElementBounds(const FNavigationElementH
 		return false;
 	}
 
+	if (const TSharedPtr<FNavigationElement> NavigationElement = Repository ? Repository->GetMutableNavigationElementForHandle(Handle) : nullptr)
+	{
+		NavigationElement->SetBounds(NewBounds);
+	}
+
 	return FNavigationDataHandler(DefaultOctreeController, DefaultDirtyAreasController).UpdateNavOctreeElementBounds(Handle, NewBounds, DirtyAreas);
 }
 
