[Navigation] fix CIS Issue 937899 (can't rely on NO_LOGGING...)

diff --git a/Engine/Source/Runtime/NavigationSystem/Private/NavigationObjectRepository.cpp b/Engine/Source/Runtime/NavigationSystem/Private/NavigationObjectRepository.cpp
--- a/Engine/Source/Runtime/NavigationSystem/Private/NavigationObjectRepository.cpp
+++ b/Engine/Source/Runtime/NavigationSystem/Private/NavigationObjectRepository.cpp
@@ -112,9 +112,7 @@ TSharedPtr<const FNavigationElement> UNavigationObjectRepository::RegisterNavRel
 	// In such case we update with the latest.
 	if (const TSharedPtr<FNavigationElement> ExistingElement = GetMutableNavigationElementForUObject(&NavRelevantObject))
 	{
-#if !NO_LOGGING
 		const FBox PreviousBounds = ExistingElement->GetBounds();
-#endif
 		{
 			UE_MT_SCOPED_WRITE_ACCESS(NavElementAccessDetector);
 			*ExistingElement = FNavigationElement(NavRelevantInterface);
