[Navigation] Update SplineNavModifierComponent to apply the NavAreaToReplace to its modifier.

diff --git a/Engine/Source/Runtime/NavigationSystem/Private/SplineNavModifierComponent.cpp b/Engine/Source/Runtime/NavigationSystem/Private/SplineNavModifierComponent.cpp
--- a/Engine/Source/Runtime/NavigationSystem/Private/SplineNavModifierComponent.cpp
+++ b/Engine/Source/Runtime/NavigationSystem/Private/SplineNavModifierComponent.cpp
@@ -102,7 +102,11 @@ void USplineNavModifierComponent::GetNavigationData(FNavigationRelevantData& Dat
 		}
 
 		// From the tube construct a convex hull whose volume will be used to mark the nav mesh with the selected AreaClass
-		const FAreaNavModifier NavModifier(Tube, ENavigationCoordSystem::Type::Unreal, ComponentTransform, AreaClass);
+		FAreaNavModifier NavModifier(Tube, ENavigationCoordSystem::Type::Unreal, ComponentTransform, AreaClass);
+		if (AreaClassToReplace)
+		{
+			NavModifier.SetAreaClassToReplace(AreaClassToReplace);
+		}
 		Data.Modifiers.Add(NavModifier);
 
 		PrevIndex = SubdivisionIndex;
