[Navlink Generation] Addition of multiple link configurations for generated navlinks.

diff --git a/Engine/Source/Runtime/NavigationSystem/Private/NavMesh/LinkGenerationConfig.cpp b/Engine/Source/Runtime/NavigationSystem/Private/NavMesh/LinkGenerationConfig.cpp
--- a/Engine/Source/Runtime/NavigationSystem/Private/NavMesh/LinkGenerationConfig.cpp
+++ b/Engine/Source/Runtime/NavigationSystem/Private/NavMesh/LinkGenerationConfig.cpp
@@ -8,20 +8,75 @@
 #include "Detour/DetourNavLinkBuilderConfig.h"
 #endif //WITH_RECAST
 
+// Deprecated FNavLinkGenerationJumpDownConfig
+
 // Refer to header for more information
 PRAGMA_DISABLE_DEPRECATION_WARNINGS
+// Deprecated
 FNavLinkGenerationJumpDownConfig::FNavLinkGenerationJumpDownConfig()
 {
 	DownDirectionAreaClass = UNavArea_Default::StaticClass();
 	UpDirectionAreaClass = UNavArea_Default::StaticClass();
 }
-PRAGMA_ENABLE_DEPRECATION_WARNINGS
 
+// Deprecated
 bool FNavLinkGenerationJumpDownConfig::Serialize(FArchive& Ar)
 {
 	UScriptStruct* const Struct = StaticStruct();
 	Struct->SerializeTaggedProperties(Ar, (uint8*)this, Struct, nullptr);
 
+#if WITH_EDITORONLY_DATA
+	if (Ar.IsLoading())
+	{
+		if (AreaClass_DEPRECATED)
+		{
+			DownDirectionAreaClass = AreaClass_DEPRECATED;
+			UpDirectionAreaClass = AreaClass_DEPRECATED;
+			AreaClass_DEPRECATED = nullptr;
+		}
+	}
+#endif //WITH_EDITORONLY_DATA
+
+	return true;
+}
+
+#if WITH_RECAST
+// Deprecated
+void FNavLinkGenerationJumpDownConfig::CopyToDetourConfig(dtNavLinkBuilderJumpDownConfig& OutDetourConfig) const
+{
+	OutDetourConfig.enabled = bEnabled;
+	OutDetourConfig.jumpLength = JumpLength;
+	OutDetourConfig.jumpDistanceFromEdge = JumpDistanceFromEdge;
+	OutDetourConfig.jumpMaxDepth = JumpMaxDepth;
+	OutDetourConfig.jumpHeight = JumpHeight;
+	OutDetourConfig.jumpEndsHeightTolerance	= JumpEndsHeightTolerance;
+	OutDetourConfig.samplingSeparationFactor = SamplingSeparationFactor;
+	OutDetourConfig.filterDistanceThreshold = FilterDistanceThreshold;
+	OutDetourConfig.linkBuilderFlags = LinkBuilderFlags;
+
+	if (LinkProxy)
+	{
+		OutDetourConfig.linkUserId = LinkProxy->GetId().GetId();	
+	}
+}
+#endif //WITH_RECAST
+PRAGMA_ENABLE_DEPRECATION_WARNINGS
+
+
+// Refer to header for more information
+PRAGMA_DISABLE_DEPRECATION_WARNINGS
+FNavLinkGenerationJumpConfig::FNavLinkGenerationJumpConfig()
+{
+	DownDirectionAreaClass = UNavArea_Default::StaticClass();
+	UpDirectionAreaClass = UNavArea_Default::StaticClass();
+}
+PRAGMA_ENABLE_DEPRECATION_WARNINGS
+
+bool FNavLinkGenerationJumpConfig::Serialize(FArchive& Ar)
+{
+	UScriptStruct* const Struct = StaticStruct();
+	Struct->SerializeTaggedProperties(Ar, (uint8*)this, Struct, nullptr);
+
 #if WITH_EDITORONLY_DATA
 	if (Ar.IsLoading())
 	{
@@ -41,7 +96,7 @@ bool FNavLinkGenerationJumpDownConfig::Serialize(FArchive& Ar)
 
 #if WITH_RECAST
 
-void FNavLinkGenerationJumpDownConfig::CopyToDetourConfig(dtNavLinkBuilderJumpDownConfig& OutDetourConfig) const
+void FNavLinkGenerationJumpConfig::CopyToDetourConfig(dtNavLinkBuilderJumpConfig& OutDetourConfig) const
 {
 	OutDetourConfig.enabled = bEnabled;
 	OutDetourConfig.jumpLength = JumpLength;
diff --git a/Engine/Source/Runtime/NavigationSystem/Private/NavMesh/RecastNavMesh.cpp b/Engine/Source/Runtime/NavigationSystem/Private/NavMesh/RecastNavMesh.cpp
--- a/Engine/Source/Runtime/NavigationSystem/Private/NavMesh/RecastNavMesh.cpp
+++ b/Engine/Source/Runtime/NavigationSystem/Private/NavMesh/RecastNavMesh.cpp
@@ -496,10 +496,13 @@ ARecastNavMesh::ARecastNavMesh(const FObjectInitializer& ObjectInitializer)
 	, bAllowNavLinkAsPathEnd(false)
 #if WITH_EDITORONLY_DATA
 	, bAllowWorldPartitionedNavMesh(false)
+PRAGMA_DISABLE_DEPRECATION_WARNINGS
+	, NavLinkJumpDownConfig()
+PRAGMA_ENABLE_DEPRECATION_WARNINGS
 #endif // WITH_EDITORONLY_DATA
 	, TileSetUpdateInterval(1.0f)
 	, NavMeshVersion(NAVMESHVER_LATEST)	
-	, RecastNavMeshImpl(NULL)
+	, RecastNavMeshImpl(nullptr)
 {
 	HeuristicScale = 0.999f;
 	LedgeSlopeFilterMode = ENavigationLedgeSlopeFilterMode::Recast;
@@ -516,6 +519,7 @@ ARecastNavMesh::ARecastNavMesh(const FObjectInitializer& ObjectInitializer)
 	BatchQueryCounter = 0;
 #endif // RECAST_ASYNC_REBUILDING
 
+	NavLinkJumpConfigs.AddDefaulted();
 
 	if (HasAnyFlags(RF_ClassDefaultObject) == false)
 	{
@@ -617,7 +621,7 @@ void ARecastNavMesh::CleanUp()
 
 void ARecastNavMesh::PostLoad()
 {
-	UE_LOG(LogNavigation, Verbose, TEXT("%s %s"), ANSI_TO_TCHAR(__FUNCTION__), *GetFullNameSafe(this));
+	UE_LOG(LogNavigation, Verbose, TEXT("%hs %s"), __FUNCTION__, *GetFullNameSafe(this));
 	
 	Super::PostLoad();
 	
@@ -662,6 +666,27 @@ void ARecastNavMesh::PostLoad()
 			SetAgentMaxStepHeight((ENavigationDataResolution)i, AgentMaxStepHeight);
 		}
 	}
+
+	// If needed, initialize NavLinkJumpConfigs from the deprecated version.
+	if (NavMeshVersion < NAVMESHVER_NAVLINK_JUMP_CONFIGS && NavLinkJumpConfigs.Num() > 0)
+	{
+		FNavLinkGenerationJumpConfig& NavLinkConfig = NavLinkJumpConfigs.Top();
+		NavLinkConfig.bEnabled = NavLinkJumpDownConfig.bEnabled;
+		NavLinkConfig.JumpLength = NavLinkJumpDownConfig.JumpLength;
+		NavLinkConfig.JumpDistanceFromEdge = NavLinkJumpDownConfig.JumpDistanceFromEdge;
+		NavLinkConfig.JumpMaxDepth = NavLinkJumpDownConfig.JumpMaxDepth;
+		NavLinkConfig.JumpHeight = NavLinkJumpDownConfig.JumpHeight;
+		NavLinkConfig.JumpEndsHeightTolerance = NavLinkJumpDownConfig.JumpEndsHeightTolerance;
+		NavLinkConfig.SamplingSeparationFactor = NavLinkJumpDownConfig.SamplingSeparationFactor;
+		NavLinkConfig.FilterDistanceThreshold = NavLinkJumpDownConfig.FilterDistanceThreshold;
+		NavLinkConfig.LinkBuilderFlags = NavLinkJumpDownConfig.LinkBuilderFlags;
+		NavLinkConfig.AreaClass_DEPRECATED = NavLinkJumpDownConfig.AreaClass_DEPRECATED;
+		NavLinkConfig.DownDirectionAreaClass = NavLinkJumpDownConfig.DownDirectionAreaClass;
+		NavLinkConfig.UpDirectionAreaClass = NavLinkJumpDownConfig.UpDirectionAreaClass;
+		NavLinkConfig.LinkProxyClass = NavLinkJumpDownConfig.LinkProxyClass;
+		NavLinkConfig.LinkProxyId = NavLinkJumpDownConfig.LinkProxyId;
+		// LinkProxy and bLinkProxyRegistered are Transient.
+	}
 	PRAGMA_ENABLE_DEPRECATION_WARNINGS
 #endif // WITH_EDITORONLY_DATA
 
@@ -712,17 +737,20 @@ void ARecastNavMesh::CheckToDiscardSubLevelNavData(const UNavigationSystemBase&
 
 void ARecastNavMesh::RegisterGeneratedLinksProxy()
 {
-	if (!NavLinkJumpDownConfig.bLinkProxyRegistered && NavLinkJumpDownConfig.LinkProxy)
+	for (FNavLinkGenerationJumpConfig& NavLinkJumpConfig : NavLinkJumpConfigs)
 	{
-		if (UWorld* World = GetWorld())
+		if (!NavLinkJumpConfig.bLinkProxyRegistered && NavLinkJumpConfig.LinkProxy)
 		{
-			UNavigationSystemV1* NavSys = FNavigationSystem::GetCurrent<UNavigationSystemV1>(World);
-			if (NavSys)
+			if (UWorld* World = GetWorld())
 			{
-				UE_LOG(LogNavLink, Log, TEXT("RegisterLinkProxy id: %llu ptr: 0x%p"), NavLinkJumpDownConfig.LinkProxyId.GetId(), NavLinkJumpDownConfig.LinkProxy.Get());
+				UNavigationSystemV1* NavSys = FNavigationSystem::GetCurrent<UNavigationSystemV1>(World);
+				if (NavSys)
+				{
+					UE_LOG(LogNavLink, Log, TEXT("RegisterLinkProxy id: %llu ptr: 0x%p"), NavLinkJumpConfig.LinkProxyId.GetId(), NavLinkJumpConfig.LinkProxy.Get());
 
-				NavSys->RegisterCustomLink(*NavLinkJumpDownConfig.LinkProxy);
-				NavLinkJumpDownConfig.bLinkProxyRegistered = true;
+					NavSys->RegisterCustomLink(*NavLinkJumpConfig.LinkProxy);
+					NavLinkJumpConfig.bLinkProxyRegistered = true;
+				}
 			}
 		}
 	}
@@ -730,24 +758,30 @@ void ARecastNavMesh::RegisterGeneratedLinksProxy()
 
 void ARecastNavMesh::UnregisterGeneratedLinksProxy()
 {
-	if (NavLinkJumpDownConfig.bLinkProxyRegistered && NavLinkJumpDownConfig.LinkProxy)
+	for (FNavLinkGenerationJumpConfig& NavLinkJumpConfig : NavLinkJumpConfigs)
 	{
-		if (UWorld* World = GetWorld())
+		if (NavLinkJumpConfig.bLinkProxyRegistered && NavLinkJumpConfig.LinkProxy)
 		{
-			UNavigationSystemV1* NavSys = FNavigationSystem::GetCurrent<UNavigationSystemV1>(World);
-			if (NavSys)
+			if (UWorld* World = GetWorld())
 			{
-				UE_LOG(LogNavLink, Log, TEXT("UnregisterLinkProxy id: %llu ptr: 0x%p"), NavLinkJumpDownConfig.LinkProxyId.GetId(), NavLinkJumpDownConfig.LinkProxy.Get());
+				UNavigationSystemV1* NavSys = FNavigationSystem::GetCurrent<UNavigationSystemV1>(World);
+				if (NavSys)
+				{
+					UE_LOG(LogNavLink, Log, TEXT("UnregisterLinkProxy id: %llu ptr: 0x%p"), NavLinkJumpConfig.LinkProxyId.GetId(), NavLinkJumpConfig.LinkProxy.Get());
 
-				NavSys->UnregisterCustomLink(*NavLinkJumpDownConfig.LinkProxy);
-				NavLinkJumpDownConfig.bLinkProxyRegistered = false;
+					NavSys->UnregisterCustomLink(*NavLinkJumpConfig.LinkProxy);
+					NavLinkJumpConfig.bLinkProxyRegistered = false;
+				}
 			}
 		}
 	}
 }
 
+// Deprecated
 void ARecastNavMesh::CreateAndRegisterJumpDownLinksProxy(const FNavLinkId LinkProxyId /*= FNavLinkId::GenerateUniqueId()*/)
 {
+#if WITH_EDITORONLY_DATA
+	PRAGMA_DISABLE_DEPRECATION_WARNINGS	
 	if (!ensureMsgf(NavLinkJumpDownConfig.LinkProxy == nullptr, TEXT("LinkProxy is expected to have been cleanup before creating a new one. A new proxy was not created.")))
 	{
 		return;
@@ -761,6 +795,28 @@ void ARecastNavMesh::CreateAndRegisterJumpDownLinksProxy(const FNavLinkId LinkPr
 		LinkProxy->SetOwner(this);
 		NavLinkJumpDownConfig.LinkProxy = LinkProxy;
 
+		RegisterGeneratedLinksProxy();
+	}
+	PRAGMA_ENABLE_DEPRECATION_WARNINGS
+#endif // WITH_EDITORONLY_DATA	
+}
+
+void ARecastNavMesh::CreateAndRegisterJumpLinksProxy(FNavLinkGenerationJumpConfig& InOutNavLinkJumpConfig, 
+	const FNavLinkId LinkProxyId /*= FNavLinkId::GenerateUniqueId()*/)
+{
+	if (!ensureMsgf(InOutNavLinkJumpConfig.LinkProxy == nullptr, TEXT("LinkProxy is expected to have been cleanup before creating a new one. A new proxy was not created.")))
+	{
+		return;
+	}
+
+	UBaseGeneratedNavLinksProxy* LinkProxy = NewObject<UBaseGeneratedNavLinksProxy>(this, InOutNavLinkJumpConfig.LinkProxyClass);
+	if (LinkProxy)
+	{
+		InOutNavLinkJumpConfig.LinkProxyId = LinkProxyId;
+		LinkProxy->UpdateLinkId(InOutNavLinkJumpConfig.LinkProxyId);
+		LinkProxy->SetOwner(this);
+		InOutNavLinkJumpConfig.LinkProxy = LinkProxy;
+
 		RegisterGeneratedLinksProxy();
 	}
 }
@@ -775,11 +831,14 @@ void ARecastNavMesh::PostRegisterAllComponents()
 	}
 
 	// Create and register link proxy if enabled.
-	if (bGenerateNavLinks && NavLinkJumpDownConfig.bEnabled && NavLinkJumpDownConfig.LinkProxyClass.Get() != nullptr)
+	for (FNavLinkGenerationJumpConfig& NavLinkJumpConfig : NavLinkJumpConfigs)
 	{
-		// If it's the first time, LinkProxyId will have the default invalid value so we need to generate an Id.
-		const FNavLinkId Id = NavLinkJumpDownConfig.LinkProxyId.IsValid() ? NavLinkJumpDownConfig.LinkProxyId :  FNavLinkId::GenerateUniqueId(); 
-		CreateAndRegisterJumpDownLinksProxy(Id);
+		if (bGenerateNavLinks && NavLinkJumpConfig.bEnabled && NavLinkJumpConfig.LinkProxyClass.Get() != nullptr)
+		{
+			// If it's the first time, LinkProxyId will have the default invalid value so we need to generate an Id.
+			const FNavLinkId Id = NavLinkJumpConfig.LinkProxyId.IsValid() ? NavLinkJumpConfig.LinkProxyId :  FNavLinkId::GenerateUniqueId(); 
+			CreateAndRegisterJumpLinksProxy(NavLinkJumpConfig, Id);
+		}
 	}
 }
 
@@ -788,7 +847,11 @@ void ARecastNavMesh::PostUnregisterAllComponents()
 	Super::PostUnregisterAllComponents();
 
 	UnregisterGeneratedLinksProxy();
-	NavLinkJumpDownConfig.LinkProxy = nullptr;
+
+	for (FNavLinkGenerationJumpConfig& NavLinkJumpConfig : NavLinkJumpConfigs)
+	{
+		NavLinkJumpConfig.LinkProxy = nullptr;
+	}
 }
 
 void ARecastNavMesh::PostInitProperties()
@@ -3612,7 +3675,7 @@ void ARecastNavMesh::PostEditChangeChainProperty(FPropertyChangedChainEvent& Pro
 		
 		if (CategoryName == NAME_Generation)
 		{
-			static const FName NAME_NavLinkJumpDownConfig = GET_MEMBER_NAME_CHECKED(ARecastNavMesh, NavLinkJumpDownConfig);
+			static const FName NAME_NavLinkJumpConfigs = GET_MEMBER_NAME_CHECKED(ARecastNavMesh, NavLinkJumpConfigs);
 			
 			const FName PropName = PropertyChangedChainEvent.Property->GetFName();
 			bool bRebuild = false;
@@ -3649,7 +3712,7 @@ void ARecastNavMesh::PostEditChangeChainProperty(FPropertyChangedChainEvent& Pro
 			{
 				bRebuild = true;
 			}
-			else if (MemberProperty && (MemberProperty->GetFName() == NAME_NavLinkJumpDownConfig))
+			else if (MemberProperty && (MemberProperty->GetFName() == NAME_NavLinkJumpConfigs))
 			{
 				bRebuild = true;
 			}
@@ -3809,13 +3872,17 @@ void ARecastNavMesh::PostLoadPreRebuild()
 {
 	// If any, unregister previous generated links proxy.
 	UnregisterGeneratedLinksProxy();
-	NavLinkJumpDownConfig.LinkProxy = nullptr;
 
-	// Register new one if needed.
-	if (NavLinkJumpDownConfig.bEnabled && NavLinkJumpDownConfig.LinkProxyClass.Get() != nullptr)
+	for (FNavLinkGenerationJumpConfig& NavLinkJumpConfig : NavLinkJumpConfigs)
 	{
-		// Since this is on rebuild, create a new proxy id.
-		CreateAndRegisterJumpDownLinksProxy();
+		NavLinkJumpConfig.LinkProxy = nullptr;
+
+		// Register new one if needed.
+		if (NavLinkJumpConfig.bEnabled && NavLinkJumpConfig.LinkProxyClass.Get() != nullptr)
+		{
+			// Since this is on rebuild, create a new proxy id.
+			CreateAndRegisterJumpLinksProxy(NavLinkJumpConfig);
+		}
 	}
 }
 
diff --git a/Engine/Source/Runtime/NavigationSystem/Private/NavMesh/RecastNavMeshGenerator.cpp b/Engine/Source/Runtime/NavigationSystem/Private/NavMesh/RecastNavMeshGenerator.cpp
--- a/Engine/Source/Runtime/NavigationSystem/Private/NavMesh/RecastNavMeshGenerator.cpp
+++ b/Engine/Source/Runtime/NavigationSystem/Private/NavMesh/RecastNavMeshGenerator.cpp
@@ -846,9 +846,6 @@ void ExportRigidBodySphylElements(const FKAggregateGeom& AggGeom, TNavStatArray<
 		const int32 NumSides = 16;
 		const int32 NumRings = (NumSides/2) + 1;
 
-		// The first/last arc are on top of each other.
-		const int32 NumVerts = (NumSides+1) * (NumRings+1);
-
 		ArcVerts.Reset();
 		ArcVerts.AddZeroed(NumRings+1);
 		for (int32 RingIdx=0; RingIdx<NumRings+1; RingIdx++)
@@ -926,9 +923,6 @@ void ExportRigidBodySphereElements(const FKAggregateGeom& AggGeom, TNavStatArray
 		const int32 NumSides = 16;
 		const int32 NumRings = (NumSides/2) + 1;
 
-		// The first/last arc are on top of each other.
-		const int32 NumVerts = (NumSides+1) * (NumRings+1);
-
 		ArcVerts.Reset();
 		ArcVerts.AddZeroed(NumRings+1);
 		for (int32 RingIdx=0; RingIdx<NumRings+1; RingIdx++)
@@ -3827,8 +3821,112 @@ struct FTileGenerationContext
 	TArray<FNavMeshTileData> NavigationData;
 };
 
+static void AddGeneratedLinks(
+	TArray<FGeneratedNavigationLink>& OutGeneratedLinks,
+	const dtLinkBuilderConfig& linkBuilderConfig,
+	const dtNavLinkBuilder& linkBuilder)
+{
+	auto AddLinkSharedLambda = [&OutGeneratedLinks](const dtNavLinkBuilderJumpConfig& config) -> FGeneratedNavigationLink&
+	{
+		FGeneratedNavigationLink& NewLink = OutGeneratedLinks.Emplace_GetRef();
+		NewLink.bIsGenerated = true;
+		NewLink.NavLinkId = FNavLinkId(config.linkUserId);
+		return NewLink;
+	};
+
+	// Used when we're provided with identical up and down areas. In this situation, a single bidirectional link is sufficient for our needs
+	auto AddBidirectionalLinkLambda = [&AddLinkSharedLambda](const dtNavLinkBuilderJumpConfig& config, const dtReal* posA, const dtReal* posB)
+	{
+		FGeneratedNavigationLink& NewLink = AddLinkSharedLambda(config);
+		NewLink.generatedLinkArea = config.downDirArea;
+		NewLink.generatedLinkPolyFlag = config.downDirPolyFlag;
+		NewLink.Left = Recast2UnrealPoint(posA);
+		NewLink.Right = Recast2UnrealPoint(posB);
+	};
+
+	// Used when we're provided with different up and down areas. In this situation, we need to create two different one-way link representations in order to meet our needs
+	auto AddDirectionalLinksLambda = [&AddLinkSharedLambda](const dtNavLinkBuilderJumpConfig& config, const dtReal* posA, const dtReal* posB)
+	{
+		const FVector UnrealPosA = Recast2UnrealPoint(posA);
+		const FVector UnrealPosB = Recast2UnrealPoint(posB);
+
+		const bool bABDown = UnrealPosA.Z > UnrealPosB.Z;
+
+		const FVector DownLinkLeft = bABDown ? UnrealPosA : UnrealPosB;
+		const FVector DownLinkRight = bABDown ? UnrealPosB : UnrealPosA;
+
+		if (config.downDirArea != RECAST_NULL_AREA)
+		{
+			FGeneratedNavigationLink& DownLink = AddLinkSharedLambda(config);
+			DownLink.generatedLinkArea = config.downDirArea;
+			DownLink.generatedLinkPolyFlag = config.downDirPolyFlag;
+			DownLink.Left = DownLinkLeft;
+			DownLink.Right = DownLinkRight;
+			DownLink.Direction = ENavLinkDirection::LeftToRight;
+		}
+
+		if (config.upDirArea != RECAST_NULL_AREA)
+		{
+			// Note that in order for this link to be recognized by the navigation system, it needs to have unique left/right points, which is why the left and right points are flipped and we're marking this link as left-to-right
+			FGeneratedNavigationLink& UpLink = AddLinkSharedLambda(config);
+			UpLink.generatedLinkArea = config.upDirArea;
+			UpLink.generatedLinkPolyFlag = config.upDirPolyFlag;
+			UpLink.Left = DownLinkRight;
+			UpLink.Right = DownLinkLeft;
+			UpLink.Direction = ENavLinkDirection::LeftToRight;
+		}
+	};
+
+	// Make FGeneratedNavigationLinks
+	for (const dtNavLinkBuilder::JumpLink& link : linkBuilder.m_links)
+	{
+		if (link.flags == dtNavLinkBuilder::FILTERED)
+			continue;
+
+		const dtNavLinkBuilderJumpConfig& jumpConfig = linkBuilderConfig.jumpConfigs[link.configIndex];
+		if (!ensure(jumpConfig.enabled))
+			continue;
+
+		const bool bBidirectionalJumpDownLinks = jumpConfig.downDirArea == jumpConfig.upDirArea;
+		TFunctionRef<void(const dtNavLinkBuilderJumpConfig&, const dtReal*, const dtReal*)> AddLinkLambda =
+			bBidirectionalJumpDownLinks ?
+				TFunctionRef<void(const dtNavLinkBuilderJumpConfig&, const dtReal*, const dtReal*)>(AddBidirectionalLinkLambda) :
+				TFunctionRef<void(const dtNavLinkBuilderJumpConfig&, const dtReal*, const dtReal*)>(AddDirectionalLinksLambda);
+		
+		if (jumpConfig.linkBuilderFlags & DT_NAVLINK_CREATE_CENTER_POINT_LINK)
+		{
+			// Make a link using the center of the range.
+			dtReal midA[3];
+			dtVlerp(midA, &link.spine0[0], &link.spine1[0], 0.5);
+			dtReal midB[3];
+			dtVlerp(midB, &link.spine0[(link.nspine-1)*3], &link.spine1[(link.nspine-1)*3], 0.5);
+			// Since trajectory validation starts at agentClimb height to ignore small bumps, remove the offset for the actual link height.
+			midA[1] -= linkBuilderConfig.agentClimb;
+			midB[1] -= linkBuilderConfig.agentClimb;
+			AddLinkLambda(jumpConfig, midA, midB);
+		}
+
+		if (jumpConfig.linkBuilderFlags & DT_NAVLINK_CREATE_EXTREMITY_LINKS)
+		{
+			dtReal posA[3];
+			dtReal posB[3];
+			dtVcopy(posA, &link.spine0[0]);
+			dtVcopy(posB, &link.spine0[(link.nspine-1)*3]);
+			posA[1] -= linkBuilderConfig.agentClimb;
+			posB[1] -= linkBuilderConfig.agentClimb;
+			AddLinkLambda(jumpConfig, posA, posB);
+
+			dtVcopy(posA, &link.spine1[0]);
+			dtVcopy(posB, &link.spine1[(link.nspine-1)*3]);
+			posA[1] -= linkBuilderConfig.agentClimb;
+			posB[1] -= linkBuilderConfig.agentClimb;
+			AddLinkLambda(jumpConfig, posA, posB);
+		}
+	}
+}
+
 dtStatus FRecastTileGenerator::BuildTileCacheLinks(FNavMeshBuildContext& BuildContext, dtTileCacheAlloc* alloc, const dtTileCacheLayer& layer,
-	const dtTileCacheContourSet& lcset, TArray<FGeneratedNavigationLink>& OutGeneratedLinks) const
+                                                   const dtTileCacheContourSet& lcset, TArray<FGeneratedNavigationLink>& OutGeneratedLinks) const
 {
 	TRACE_CPUPROFILER_EVENT_SCOPE(FRecastTileGenerator::BuildTileCacheLinks);
 	BuildContext.log(RC_LOG_PROGRESS, "Building Links:");
@@ -3866,14 +3964,16 @@ dtStatus FRecastTileGenerator::BuildTileCacheLinks(FNavMeshBuildContext& BuildCo
 	const dtReal* orig = layer.header->bmin;
 
 	dtLinkBuilderConfig linkBuilderConfig;
-	linkBuilderConfig.jumpDownConfig = TileConfig.JumpDownConfig;
-	if (linkBuilderConfig.jumpDownConfig.enabled)
+
+	linkBuilderConfig.jumpConfigs = TileConfig.JumpConfigs;
+	for (dtNavLinkBuilderJumpConfig& jumpConfig : linkBuilderConfig.jumpConfigs)
 	{
-		linkBuilderConfig.jumpDownConfig.init();
+		if (jumpConfig.enabled)
+		{
+			jumpConfig.init();
+		}
 	}
 	
-	linkBuilderConfig.jumpOverConfig = TileConfig.JumpOverConfig;
-	
 	linkBuilderConfig.agentRadius = TileConfig.walkableRadius * TileConfig.cs;
 	linkBuilderConfig.agentHeight = TileConfig.walkableHeight * TileConfig.ch;
 	linkBuilderConfig.agentClimb = TileConfig.walkableClimb * TileConfig.ch;
@@ -3902,16 +4002,13 @@ dtStatus FRecastTileGenerator::BuildTileCacheLinks(FNavMeshBuildContext& BuildCo
 	if (DebugEdge == -1)
 	{
 		rcContext& context = BuildContext;
-		if (linkBuilderConfig.jumpDownConfig.enabled)
+		for (const dtNavLinkBuilderJumpConfig& jumpConfig : linkBuilderConfig.jumpConfigs)
 		{
-			TRACE_CPUPROFILER_EVENT_SCOPE(RecastBuildLinks_JumpDown);
-			linkBuilder.buildForAllEdges(context, linkBuilderConfig, DT_LINK_ACTION_JUMP_DOWN);
-		}
-
-		if (linkBuilderConfig.jumpOverConfig.enabled)
-		{
-			TRACE_CPUPROFILER_EVENT_SCOPE(RecastBuildLinks_JumpOver);
-			linkBuilder.buildForAllEdges(context, linkBuilderConfig, DT_LINK_ACTION_JUMP_OVER);
+			if (jumpConfig.enabled)
+			{
+				TRACE_CPUPROFILER_EVENT_SCOPE(RecastBuildLinks_JumpConfig);
+				linkBuilder.buildForAllEdges(context, linkBuilderConfig, jumpConfig);
+			}
 		}
 
 #if RECAST_INTERNAL_DEBUG_DATA
@@ -3919,123 +4016,20 @@ dtStatus FRecastTileGenerator::BuildTileCacheLinks(FNavMeshBuildContext& BuildCo
 	}
 	else
 	{
-		if (linkBuilderConfig.jumpDownConfig.enabled)
+		for (const dtNavLinkBuilderJumpConfig& jumpConfig : linkBuilderConfig.jumpConfigs)
 		{
-			dtNavLinkBuilder::EdgeSampler sampler1;
-			linkBuilder.debugBuildEdge(linkBuilderConfig, DT_LINK_ACTION_JUMP_DOWN, DebugEdge, sampler1);
-			duDebugDrawNavLinkBuilder(dd, linkBuilder, DebugFlags, &sampler1);
-		}
-		
-		if (linkBuilderConfig.jumpOverConfig.enabled)
-		{
-			dtNavLinkBuilder::EdgeSampler sampler2;
-			linkBuilder.debugBuildEdge(linkBuilderConfig, DT_LINK_ACTION_JUMP_OVER, DebugEdge, sampler2);
-			duDebugDrawNavLinkBuilder(dd, linkBuilder, DebugFlags, &sampler2);
-		}
-#endif // RECAST_INTERNAL_DEBUG_DATA		
-	}
-
-	if (linkBuilderConfig.jumpDownConfig.enabled)
-	{
-		auto AddLinkSharedLambda = [&OutGeneratedLinks](const dtNavLinkBuilderJumpDownConfig& config) -> FGeneratedNavigationLink&
-		{
-			FGeneratedNavigationLink& NewLink = OutGeneratedLinks.Emplace_GetRef();
-			NewLink.bIsGenerated = true;
-			NewLink.NavLinkId = FNavLinkId(config.linkUserId);
-			return NewLink;
-		};
-
-		// Used when we're provided with identical up and down areas. In this situation, a single bidirectional link is sufficient for our needs
-		auto AddBidirectionalLinkLambda = [&AddLinkSharedLambda](const dtNavLinkBuilderJumpDownConfig& config, const dtReal* posA, const dtReal* posB)
-		{
-			FGeneratedNavigationLink& NewLink = AddLinkSharedLambda(config);
-			NewLink.generatedLinkArea = config.downDirArea;
-			NewLink.generatedLinkPolyFlag = config.downDirPolyFlag;
-			NewLink.Left = Recast2UnrealPoint(posA);
-			NewLink.Right = Recast2UnrealPoint(posB);
-		};
-
-		// Used when we're provided with different up and down areas. In this situation, we need to create two different one-way link representations in order to meet our needs
-		auto AddDirectionalLinksLambda = [&AddLinkSharedLambda](const dtNavLinkBuilderJumpDownConfig& config, const dtReal* posA, const dtReal* posB)
-		{
-			const FVector UnrealPosA = Recast2UnrealPoint(posA);
-			const FVector UnrealPosB = Recast2UnrealPoint(posB);
-
-			const bool bABDown = UnrealPosA.Z > UnrealPosB.Z;
-
-			const FVector DownLinkLeft = bABDown ? UnrealPosA : UnrealPosB;
-			const FVector DownLinkRight = bABDown ? UnrealPosB : UnrealPosA;
-
-
-			if (config.downDirArea != RECAST_NULL_AREA)
+			if (jumpConfig.enabled)
 			{
-				FGeneratedNavigationLink& DownLink = AddLinkSharedLambda(config);
-				DownLink.generatedLinkArea = config.downDirArea;
-				DownLink.generatedLinkPolyFlag = config.downDirPolyFlag;
-				DownLink.Left = DownLinkLeft;
-				DownLink.Right = DownLinkRight;
-				DownLink.Direction = ENavLinkDirection::LeftToRight;
-			}
-
-			if (config.upDirArea != RECAST_NULL_AREA)
-			{
-				// Note that in order for this link to be recognized by the navigation system, it needs to have unique left/right points, which is why the left and right points are flipped and we're marking this link as left-to-right
-				FGeneratedNavigationLink& UpLink = AddLinkSharedLambda(config);
-				UpLink.generatedLinkArea = config.upDirArea;
-				UpLink.generatedLinkPolyFlag = config.upDirPolyFlag;
-				UpLink.Left = DownLinkRight;
-				UpLink.Right = DownLinkLeft;
-				UpLink.Direction = ENavLinkDirection::LeftToRight;
-			}
-		};
-
-		const bool bBidirectionalJumpDownLinks = linkBuilderConfig.jumpDownConfig.downDirArea == linkBuilderConfig.jumpDownConfig.upDirArea;
-		TFunctionRef<void(const dtNavLinkBuilderJumpDownConfig&, const dtReal*, const dtReal*)> AddLinkLambda = bBidirectionalJumpDownLinks ? TFunctionRef<void(const dtNavLinkBuilderJumpDownConfig&, const dtReal*, const dtReal*)>(AddBidirectionalLinkLambda) : TFunctionRef<void(const dtNavLinkBuilderJumpDownConfig&, const dtReal*, const dtReal*)>(AddDirectionalLinksLambda);
-
-		// Make FGeneratedNavigationLinks
-		for (const dtNavLinkBuilder::JumpLink& link : linkBuilder.m_links)
-		{
-			if (link.flags == dtNavLinkBuilder::FILTERED)
-				continue;
-
-			// Only "JumpDownConfig are expected for now
-			if (link.action != DT_LINK_ACTION_JUMP_DOWN)
-			{
-				continue;
-			}
-
-			if (TileConfig.JumpDownConfig.linkBuilderFlags & DT_NAVLINK_CREATE_CENTER_POINT_LINK)
-			{
-				// Make a link using the center of the range.
-				dtReal midA[3];
-				dtVlerp(midA, &link.spine0[0], &link.spine1[0], 0.5);
-				dtReal midB[3];
-				dtVlerp(midB, &link.spine0[(link.nspine-1)*3], &link.spine1[(link.nspine-1)*3], 0.5);
-				// Since trajectory validation starts at agentClimb height to ignore small bumps, remove the offset for the actual link height.
-				midA[1] -= linkBuilderConfig.agentClimb;
-				midB[1] -= linkBuilderConfig.agentClimb;
-				AddLinkLambda(linkBuilderConfig.jumpDownConfig, midA, midB);
-			}
-
-			if (TileConfig.JumpDownConfig.linkBuilderFlags & DT_NAVLINK_CREATE_EXTREMITY_LINKS)
-			{
-				dtReal posA[3];
-				dtReal posB[3];
-				dtVcopy(posA, &link.spine0[0]);
-				dtVcopy(posB, &link.spine0[(link.nspine-1)*3]);
-				posA[1] -= linkBuilderConfig.agentClimb;
-				posB[1] -= linkBuilderConfig.agentClimb;
-				AddLinkLambda(linkBuilderConfig.jumpDownConfig, posA, posB);
-
-				dtVcopy(posA, &link.spine1[0]);
-				dtVcopy(posB, &link.spine1[(link.nspine-1)*3]);
-				posA[1] -= linkBuilderConfig.agentClimb;
-				posB[1] -= linkBuilderConfig.agentClimb;
-				AddLinkLambda(linkBuilderConfig.jumpDownConfig, posA, posB);
+				dtNavLinkBuilder::EdgeSampler sampler1;
+				linkBuilder.debugBuildEdge(linkBuilderConfig, DebugEdge, sampler1, jumpConfig);
+				duDebugDrawNavLinkBuilder(dd, linkBuilder, DebugFlags, &sampler1);
 			}
 		}
+#endif // RECAST_INTERNAL_DEBUG_DATA		
 	}
 
+	AddGeneratedLinks(OutGeneratedLinks, linkBuilderConfig, linkBuilder);
+
 	LogOnExit();
 	return DT_SUCCESS;
 }
@@ -5156,12 +5150,18 @@ void FRecastNavMeshGenerator::ConfigureBuildProperties(FRecastBuildConfig& OutCo
 	if (IsGeneratingLinks())
 	{
 		// NavLink builder configuration
-		const FNavLinkGenerationJumpDownConfig& JumpDown = DestNavMesh->NavLinkJumpDownConfig;
-		JumpDown.CopyToDetourConfig(OutConfig.JumpDownConfig);
+		check(OutConfig.JumpConfigs.IsEmpty());
+		for (uint8 Index = 0; Index < DestNavMesh->NavLinkJumpConfigs.Num() && Index < TNumericLimits<uint8>::Max(); Index++)
+		{
+			const FNavLinkGenerationJumpConfig& JumpConfig = DestNavMesh->NavLinkJumpConfigs[Index];
+
+			dtNavLinkBuilderJumpConfig& DetourJumpConfig = OutConfig.JumpConfigs.AddDefaulted_GetRef();
+			JumpConfig.CopyToDetourConfig(DetourJumpConfig);
+			DetourJumpConfig.ConfigIndex = Index;
 
-		const float JumpDownSpillDistance = JumpDown.bEnabled ? JumpDown.JumpLength - JumpDown.JumpDistanceFromEdge : 0.f;
-		constexpr float JumpOverSpillDistance = 0.f; //JumpOver.bEnabled ? JumpOver.JumpDistanceFromGapCenter : 0.f;		// @todo: jump over config is not exposed for now
-		OutConfig.LinkSpillDistance = FMath::Max(JumpDownSpillDistance, JumpOverSpillDistance);
+			const float JumpSpillDistance = JumpConfig.bEnabled ? JumpConfig.JumpLength - JumpConfig.JumpDistanceFromEdge : 0.f;
+			OutConfig.LinkSpillDistance = FMath::Max(JumpSpillDistance, OutConfig.LinkSpillDistance);
+		}
 	}
 	
 	// store original sizes
@@ -5506,14 +5506,22 @@ void FRecastNavMeshGenerator::ResolveGeneratedLinkAreas(FRecastBuildConfig& OutC
 			}
 		};
 
-		const FNavLinkGenerationJumpDownConfig& JumpDown = DestNavMesh->NavLinkJumpDownConfig;
-		GetAreaIDAndPolyFlags(JumpDown.DownDirectionAreaClass, OutConfig.JumpDownConfig.downDirArea, OutConfig.JumpDownConfig.downDirPolyFlag);
-		GetAreaIDAndPolyFlags(JumpDown.UpDirectionAreaClass, OutConfig.JumpDownConfig.upDirArea, OutConfig.JumpDownConfig.upDirPolyFlag);
-
-		if (OutConfig.JumpDownConfig.downDirArea == RECAST_NULL_AREA && OutConfig.JumpDownConfig.upDirArea == RECAST_NULL_AREA)
+		if (ensureMsgf(DestNavMesh->NavLinkJumpConfigs.Num() == OutConfig.JumpConfigs.Num(), TEXT("Link configuration sizes are expected to match.")))
 		{
-			UE_LOG(LogNavigation, Warning, TEXT("FRecastTileGenerator: no area class has been defined for the jump down links. Disabling them"));
-			OutConfig.JumpDownConfig.enabled = false;
+			for (int32 Index = 0; Index < DestNavMesh->NavLinkJumpConfigs.Num(); Index++)
+			{
+				const FNavLinkGenerationJumpConfig& JumpConfig = DestNavMesh->NavLinkJumpConfigs[Index];
+				dtNavLinkBuilderJumpConfig& DetourJumpConfig = OutConfig.JumpConfigs[Index];
+
+				GetAreaIDAndPolyFlags(JumpConfig.DownDirectionAreaClass, DetourJumpConfig.downDirArea, DetourJumpConfig.downDirPolyFlag);
+				GetAreaIDAndPolyFlags(JumpConfig.UpDirectionAreaClass, DetourJumpConfig.upDirArea, DetourJumpConfig.upDirPolyFlag);
+
+				if (DetourJumpConfig.downDirArea == RECAST_NULL_AREA && DetourJumpConfig.upDirArea == RECAST_NULL_AREA)
+				{
+					UE_LOG(LogNavigation, Warning, TEXT("FRecastTileGenerator: no area class has been defined for the jump down links. Disabling them"));
+					DetourJumpConfig.enabled = false;
+				}
+			}
 		}
 	}
 }
diff --git a/Engine/Source/Runtime/NavigationSystem/Public/NavMesh/LinkGenerationConfig.h b/Engine/Source/Runtime/NavigationSystem/Public/NavMesh/LinkGenerationConfig.h
--- a/Engine/Source/Runtime/NavigationSystem/Public/NavMesh/LinkGenerationConfig.h
+++ b/Engine/Source/Runtime/NavigationSystem/Public/NavMesh/LinkGenerationConfig.h
@@ -10,8 +10,7 @@ class UBaseGeneratedNavLinksProxy;
 class UNavAreaBase;
 
 #if WITH_RECAST
-struct dtNavLinkBuilderJumpDownConfig;
-struct dtNavLinkBuilderJumpOverConfig;
+struct dtNavLinkBuilderJumpConfig;
 #endif //WITH_RECAST
 
 UENUM(meta = (Bitflags, UseEnumValuesAsMaskValuesInEditor = "true"))
@@ -24,20 +23,144 @@ ENUM_CLASS_FLAGS(ENavLinkBuilderFlags);
 
 /** Experimental configuration to generate vertical links. */
 USTRUCT()
-struct FNavLinkGenerationJumpDownConfig
+struct FNavLinkGenerationJumpConfig
 {
 	GENERATED_BODY()
 
-	FNavLinkGenerationJumpDownConfig();
+	NAVIGATIONSYSTEM_API FNavLinkGenerationJumpConfig();
 
 	// Note: We need to explicitly disable warnings on these constructors/operators for clang to be happy with deprecated variables
 	PRAGMA_DISABLE_DEPRECATION_WARNINGS
+	~FNavLinkGenerationJumpConfig() = default;
+	FNavLinkGenerationJumpConfig(const FNavLinkGenerationJumpConfig&) = default;
+	FNavLinkGenerationJumpConfig(FNavLinkGenerationJumpConfig&&) = default;
+	FNavLinkGenerationJumpConfig& operator=(const FNavLinkGenerationJumpConfig&) = default;
+	FNavLinkGenerationJumpConfig& operator=(FNavLinkGenerationJumpConfig&&) = default;
+	PRAGMA_ENABLE_DEPRECATION_WARNINGS
+
+	/** Should this config be used to generate links. */
+	UPROPERTY(Config)		// @todo Hidden for now. Since there is currently only one config, it has no real usage.
+	bool bEnabled = true;
+	
+	/** Horizontal length of the jump.
+	 * How far from the starting point we will look for ground. */ 
+	UPROPERTY(EditAnywhere, Config, Category = Settings, meta=(UIMin=0, ClampMin=0))
+	float JumpLength = 150.f; 
+
+	/** How far from the navmesh edge is the jump started. */
+	UPROPERTY(EditAnywhere, Config, Category = Settings, meta=(UIMin=0, ClampMin=0))
+	float JumpDistanceFromEdge = 10.f; 
+
+	/** How far below the starting height we want to look for landing ground. */
+	UPROPERTY(EditAnywhere, Config, Category = Settings, meta=(UIMin=0, ClampMin=0))
+	float JumpMaxDepth = 150.f;
+
+	/** Peak height relative to the height of the starting point. */
+	UPROPERTY(EditAnywhere, Config, Category = Settings, meta=(UIMin=0, ClampMin=0))
+	float JumpHeight = 50.f;
+
+	/** Tolerance at both ends of the jump to find ground. */
+	UPROPERTY(EditAnywhere, Config, Category = Settings, meta=(UIMin=0, ClampMin=0))
+	float JumpEndsHeightTolerance = 80.f;
+
+	/** Value multiplied by CellSize to find the distance between sampling trajectories. Default is 1.
+     *  Larger values improve generation speed but might introduce sampling errors.  */
+	UPROPERTY(EditAnywhere, Config, Category = Settings, meta=(UIMin=1, ClampMin=1))
+	float SamplingSeparationFactor = 1.f;
+	
+	/** When filtering similar links, it's the distance used to compare between segment endpoints to match similar links.
+	 * Use greater distance for more filtering (0 to deactivate filtering). */
+	UPROPERTY(EditAnywhere, Config, Category = Settings, meta=(UIMin=0, ClampMin=0))
+	float FilterDistanceThreshold = 80.f;
+
+	/** Flags used to indicate how links will be added. */
+	UPROPERTY(EditAnywhere, Config, Category = Settings, meta = (Bitmask, BitmaskEnum = "/Script/NavigationSystem.ENavLinkBuilderFlags"))
+	uint16 LinkBuilderFlags = (uint16)ENavLinkBuilderFlags::CreateCenterPointLink;
+
+#if WITH_EDITORONLY_DATA
+	/** Area class for links generated by this configuration. */
+	UE_DEPRECATED(5.6, "AreaClass is deprecated")
+	UPROPERTY(Config, meta = (DeprecatedProperty, DeprecationMessage = "Use DownDirectionAreaClass and UpDirectionAreaClass instead."))
+	TSubclassOf<UNavAreaBase> AreaClass_DEPRECATED;
+#endif //WITH_EDITORONLY_DATA
+
+	/**
+	 * Area class for downward traversal of links generated by this configuration.
+	 * @note If the value matches UpDirectionAreaClass, a single bidirectional link will be used to represent links generated by this configuration.
+	 * If the value is null, no link will be generated in this direction
+	 */
+	UPROPERTY(EditAnywhere, Config, Category = Settings)
+	TSubclassOf<UNavAreaBase> DownDirectionAreaClass;
+
+	/**
+	 * Area class for upward traversal of links generated by this configuration.
+	 * @note If the value matches DownDirectionAreaClass, a single bidirectional link will be used to represent links generated by this configuration.
+	 * If the value is null, no link will be generated in this direction
+	 */
+	UPROPERTY(EditAnywhere, Config, Category = Settings)
+	TSubclassOf<UNavAreaBase> UpDirectionAreaClass;
+
+	/** Class used to handle links made with this configuration.
+	 * Using this allows to implement custom behaviors when using navlinks, for example during the pathfollow.
+	 * Note that having a proxy is not required for successful navlink pathfinding,
+	 * but it does allow for custom behavior at the start and the end of a given navlink.
+	 * This implies that using LinkProxyClass is optional, and it can remain empty (the default value).
+	 * @see INavLinkCustomInterface 
+	 * @see UGeneratedNavLinksProxy
+	 */
+	UPROPERTY(EditAnywhere, Category= Settings)
+	TSubclassOf<UBaseGeneratedNavLinksProxy> LinkProxyClass;
+
+	/** Identifier used identify the current proxy handler. All links generated through this config will use the same handler. */
+	UPROPERTY()
+	FNavLinkId LinkProxyId;
+
+	/** Current proxy. The proxy instance is build from the LinkProxyClass (provided it's not null).
+	 * A proxy will be created if a @see LinkProxyClass is used.
+	 */
+	UPROPERTY(Transient)
+	TObjectPtr<UBaseGeneratedNavLinksProxy> LinkProxy = nullptr;
+
+	/** Is the link proxy registered to the navigation system CustomNavLinksMap.
+	 * Registration occurs on PostRegisterAllComponents or on PostLoadPreRebuild if a new proxy was created. */
+	UPROPERTY(Transient)
+	bool bLinkProxyRegistered = false;
+
+	/** Implemented for deprecated property cleanup purposes. */
+	bool Serialize(FArchive& Ar);
+
+#if WITH_RECAST	
+	/** Copy configuration to dtNavLinkBuilderJumpConfig. */
+	void CopyToDetourConfig(dtNavLinkBuilderJumpConfig& OutDetourConfig) const;
+#endif //WITH_RECAST
+};
+
+template<> struct TStructOpsTypeTraits<FNavLinkGenerationJumpConfig> : public TStructOpsTypeTraitsBase2<FNavLinkGenerationJumpConfig>
+{
+	enum
+	{
+		WithSerializer = true
+	};
+};
+
+
+PRAGMA_DISABLE_DEPRECATION_WARNINGS
+struct dtNavLinkBuilderJumpDownConfig;
+
+/** Experimental configuration to generate vertical links. */
+USTRUCT()
+struct UE_DEPRECATED(5.7, "Use FNavLinkGenerationJumpConfig instead.") FNavLinkGenerationJumpDownConfig
+{
+	GENERATED_BODY()
+
+	FNavLinkGenerationJumpDownConfig();
+
+	// Note: We need to explicitly disable warnings on these constructors/operators for clang to be happy with deprecated variables
 	~FNavLinkGenerationJumpDownConfig() = default;
 	FNavLinkGenerationJumpDownConfig(const FNavLinkGenerationJumpDownConfig&) = default;
 	FNavLinkGenerationJumpDownConfig(FNavLinkGenerationJumpDownConfig&&) = default;
 	FNavLinkGenerationJumpDownConfig& operator=(const FNavLinkGenerationJumpDownConfig&) = default;
 	FNavLinkGenerationJumpDownConfig& operator=(FNavLinkGenerationJumpDownConfig&&) = default;
-	PRAGMA_ENABLE_DEPRECATION_WARNINGS
 
 	/** Should this config be used to generate links. */
 	UPROPERTY(Config)		// @todo Hidden for now. Since there is currently only one config, it has no real usage.
@@ -143,4 +266,4 @@ template<> struct TStructOpsTypeTraits<FNavLinkGenerationJumpDownConfig> : publi
 		WithSerializer = true
 	};
 };
-
+PRAGMA_ENABLE_DEPRECATION_WARNINGS
\ No newline at end of file
diff --git a/Engine/Source/Runtime/NavigationSystem/Public/NavMesh/RecastNavMesh.h b/Engine/Source/Runtime/NavigationSystem/Public/NavMesh/RecastNavMesh.h
--- a/Engine/Source/Runtime/NavigationSystem/Public/NavMesh/RecastNavMesh.h
+++ b/Engine/Source/Runtime/NavigationSystem/Public/NavMesh/RecastNavMesh.h
@@ -814,7 +814,7 @@ class ARecastNavMesh : public ANavigationData
 	uint32 bIsWorldPartitioned : 1;
 
 	/** Experimental: if set, navlinks will be automatically generated.
-	 * @see FNavLinkGenerationJumpDownConfig */ 
+	 * @see FNavLinkGenerationJumpConfig */ 
 	UPROPERTY(EditAnywhere, Category=Generation, config)
 	uint32 bGenerateNavLinks : 1;
 	
@@ -890,10 +890,19 @@ protected:
 	uint32 bAllowWorldPartitionedNavMesh : 1;
 #endif // WITH_EDITORONLY_DATA
 
+#if WITH_EDITORONLY_DATA
+PRAGMA_DISABLE_DEPRECATION_WARNINGS
 	/** Experimental configuration to generate vertical links. */
-	UPROPERTY(EditAnywhere, Category=Generation, config)
+	UE_DEPRECATED(5.7, "Use the NavLinkJumpConfigs array instead.")
+	UPROPERTY(config)
 	FNavLinkGenerationJumpDownConfig NavLinkJumpDownConfig;
-	
+PRAGMA_ENABLE_DEPRECATION_WARNINGS
+#endif // WITH_EDITORONLY_DATA
+
+	/** Experimental configurations to generate jump links. */
+	UPROPERTY(EditAnywhere, Category=Generation, config)
+	TArray<FNavLinkGenerationJumpConfig> NavLinkJumpConfigs;
+
 private:
 	/** @returns true if there were no tiles when the navmesh was loaded. */
 	bool bHasNoTileData : 1 = false;
@@ -1472,7 +1481,14 @@ public:
 	NAVIGATIONSYSTEM_API const TMap<FIntPoint, struct FRecastInternalDebugData>* GetDebugDataMap() const;
 #endif
 
+#if WITH_EDITORONLY_DATA
+	PRAGMA_DISABLE_DEPRECATION_WARNINGS
+	UE_DEPRECATED(5.7, "Use GetNavLinkJumpConfigs instead.")
 	const FNavLinkGenerationJumpDownConfig& GetNavLinkJumpDownConfig() const { return NavLinkJumpDownConfig; }
+	PRAGMA_ENABLE_DEPRECATION_WARNINGS
+#endif // WITH_EDITORONLY_DATA
+
+	const TArray<FNavLinkGenerationJumpConfig>& GetNavLinkJumpConfigs() const { return NavLinkJumpConfigs; }
 
 protected:
 
@@ -1491,8 +1507,12 @@ protected:
 	void UnregisterGeneratedLinksProxy();
 
 	/* Create and register links proxy. It's expected to be called on load or when the navmesh is rebuilt. */
+	UE_DEPRECATED(5.7, "Use the override using an FNavLinkGenerationJumpConfig instead.")
 	void CreateAndRegisterJumpDownLinksProxy(const FNavLinkId LinkProxyId = FNavLinkId::GenerateUniqueId());
 
+	/* Create and register links proxy. It's expected to be called on load or when the navmesh is rebuilt. */
+	void CreateAndRegisterJumpLinksProxy(FNavLinkGenerationJumpConfig& InOutNavLinkJumpConfig, const FNavLinkId LinkProxyId = FNavLinkId::GenerateUniqueId());
+	
 private:
 	friend struct FRecastGraphWrapper;
 	friend FRecastNavMeshGenerator;
diff --git a/Engine/Source/Runtime/NavigationSystem/Public/NavMesh/RecastNavMeshGenerator.h b/Engine/Source/Runtime/NavigationSystem/Public/NavMesh/RecastNavMeshGenerator.h
--- a/Engine/Source/Runtime/NavigationSystem/Public/NavMesh/RecastNavMeshGenerator.h
+++ b/Engine/Source/Runtime/NavigationSystem/Public/NavMesh/RecastNavMeshGenerator.h
@@ -52,7 +52,6 @@ class UNavigationSystemV1;
 
 #define MAX_VERTS_PER_POLY	6
 
-PRAGMA_DISABLE_DEPRECATION_WARNINGS
 struct FRecastBuildConfig : public rcConfig
 {
 	/** controls whether voxel filtering will be applied (via FRecastTileGenerator::ApplyVoxelFilter) */
@@ -83,8 +82,7 @@ struct FRecastBuildConfig : public rcConfig
 	int32 MaxPolysPerTile;
 
 	/** NavLink building configuration */
-	dtNavLinkBuilderJumpDownConfig JumpDownConfig;
-	dtNavLinkBuilderJumpOverConfig JumpOverConfig;
+	TArray<dtNavLinkBuilderJumpConfig> JumpConfigs;
 	
 	/** Actual agent height (in uu)*/
 	float AgentHeight;
@@ -119,10 +117,8 @@ struct FRecastBuildConfig : public rcConfig
 		bFilterLowSpanSequences = false;
 		bFilterLowSpanFromTileCache = false;
 		bGenerateLinks = false;
-		// Still initializing, even though the property is deprecated, to avoid static analysis warnings
 		MaxPolysPerTile = -1;
-		JumpDownConfig = dtNavLinkBuilderJumpDownConfig();
-		JumpOverConfig = dtNavLinkBuilderJumpOverConfig();
+		JumpConfigs.Empty();
 		AgentIndex = 0;
 		TileResolution = ENavigationDataResolution::Default;
 		LedgeSlopeFilterMode = ENavigationLedgeSlopeFilterMode::Recast;
@@ -136,8 +132,6 @@ struct FRecastBuildConfig : public rcConfig
 	bool IsUsingCoarseCellSize() const;
 };
 
-PRAGMA_ENABLE_DEPRECATION_WARNINGS
-
 struct FRecastVoxelCache
 {
 	struct FTileInfo
diff --git a/Engine/Source/Runtime/NavigationSystem/Public/NavMesh/RecastVersion.h b/Engine/Source/Runtime/NavigationSystem/Public/NavMesh/RecastVersion.h
--- a/Engine/Source/Runtime/NavigationSystem/Public/NavMesh/RecastVersion.h
+++ b/Engine/Source/Runtime/NavigationSystem/Public/NavMesh/RecastVersion.h
@@ -29,6 +29,7 @@
 #define NAVMESHVER_TILE_RESOLUTIONS_CELLHEIGHT	24 // Addition of CellHeight in the resolution params, deprecating the original CellHeight.
 #define NAVMESHVER_1_VOXEL_AGENT_STEEP_SLOPE_FILTER_FIX	25 // Fix, remove steep slope filtering during heightfield ledge filtering when the agent radius is included into a single voxel
 #define NAVMESHVER_TILE_RESOLUTIONS_AGENTMAXSTEPHEIGHT 26	// Addition of AgentMaxStepHeight in the resolution params, deprecating the original AgentMaxStepHeight.
+#define NAVMESHVER_NAVLINK_JUMP_CONFIGS			27 // Addition of NavLinkJumpConfigs, deprecating the original NavLinkJumpDownConfig
 
-#define NAVMESHVER_LATEST				NAVMESHVER_TILE_RESOLUTIONS_AGENTMAXSTEPHEIGHT
+#define NAVMESHVER_LATEST				NAVMESHVER_NAVLINK_JUMP_CONFIGS
 #define NAVMESHVER_MIN_COMPATIBLE		NAVMESHVER_LWCOORDS_OPTIMIZATION
diff --git a/Engine/Source/Runtime/Navmesh/Private/DebugUtils/DetourNavLinkDebugDraw.cpp b/Engine/Source/Runtime/Navmesh/Private/DebugUtils/DetourNavLinkDebugDraw.cpp
--- a/Engine/Source/Runtime/Navmesh/Private/DebugUtils/DetourNavLinkDebugDraw.cpp
+++ b/Engine/Source/Runtime/Navmesh/Private/DebugUtils/DetourNavLinkDebugDraw.cpp
@@ -314,18 +314,24 @@ void duDebugDrawNavLinkBuilder(duDebugDraw* dd, const dtNavLinkBuilder& linkBuil
 		unsigned int jumpOverCol0 = duLerpCol(duColor::lightGrey, duColor::white, 200);
 		unsigned int jumpOverCol1 = duColor::lightGrey;
 
-		auto selectColors = [&](dtNavLinkBuilder::JumpLinkFlag flag, dtNavLinkAction action, unsigned int& col0, unsigned int& col1)
+		auto selectColors = [&](dtNavLinkBuilder::JumpLinkFlag flag, const uint8 configIndex, unsigned int& col0, unsigned int& col1)
 		{
-			if (action == DT_LINK_ACTION_JUMP_DOWN)
+			// Predefined "nice" colors for the first two configs. Random colors for the others.
+			if (configIndex == 0)
 			{
 				col0 = jumpDownCol0;
 				col1 = jumpDownCol1;
 			}
-			else if (action == DT_LINK_ACTION_JUMP_OVER)
+			else if (configIndex == 1)
 			{
 				col0 = jumpOverCol0;
 				col1 = jumpOverCol1;
 			}
+			else
+			{
+				col1 = duIntToCol(configIndex, 255);
+				col0 = duLerpCol(col1, duColor::white, 200);
+			}
 
 			if (flag == dtNavLinkBuilder::FILTERED)
 			{
@@ -346,7 +352,7 @@ void duDebugDrawNavLinkBuilder(duDebugDraw* dd, const dtNavLinkBuilder& linkBuil
 				if (!drawFilteredLinks && link.flags == dtNavLinkBuilder::FILTERED)
 					continue;
 
-				selectColors(link.flags, link.action, col0, col1);
+				selectColors(link.flags, link.configIndex, col0, col1);
 				
 				for (int j = 0; j < link.nspine-1; ++j)
 				{
@@ -367,7 +373,7 @@ void duDebugDrawNavLinkBuilder(duDebugDraw* dd, const dtNavLinkBuilder& linkBuil
 				if (!drawFilteredLinks && link.flags == dtNavLinkBuilder::FILTERED)
 					continue;
 
-				selectColors(link.flags, link.action, col0, col1);
+				selectColors(link.flags, link.configIndex, col0, col1);
 				
 				for (int j = 0; j < link.nspine-1; ++j)
 				{
@@ -393,7 +399,7 @@ void duDebugDrawNavLinkBuilder(duDebugDraw* dd, const dtNavLinkBuilder& linkBuil
 				if (link.flags == dtNavLinkBuilder::FILTERED)
 					continue;
 
-				selectColors(link.flags, link.action, col0, col1);
+				selectColors(link.flags, link.configIndex, col0, col1);
 				
 				dd->vertex(&link.spine0[0], duDarkenCol(col1));
 				dd->vertex(&link.spine1[0], duDarkenCol(col1));
@@ -419,7 +425,7 @@ void duDebugDrawNavLinkBuilder(duDebugDraw* dd, const dtNavLinkBuilder& linkBuil
 
 	if (drawFlags & DRAW_SELECTED_EDGE)
 	{
-		if (es && es->action != DT_LINK_ACTION_UNSET)
+		if (es && es->configIndex != dtNavLinkBuilder::INVALID_CONFIG_INDEX)
 		{
 			dd->begin(DU_DRAW_LINES, 2.0f);
 
diff --git a/Engine/Source/Runtime/Navmesh/Private/Detour/DetourNavLinkBuilder.cpp b/Engine/Source/Runtime/Navmesh/Private/Detour/DetourNavLinkBuilder.cpp
--- a/Engine/Source/Runtime/Navmesh/Private/Detour/DetourNavLinkBuilder.cpp
+++ b/Engine/Source/Runtime/Navmesh/Private/Detour/DetourNavLinkBuilder.cpp
@@ -118,21 +118,6 @@ namespace UE::Detour::NavLink::Private
 		dst[1] = ax[1]*pt[0] + ay[1]*pt[1];
 		dst[2] = ax[2]*pt[0] + ay[2]*pt[1];
 	}
-
-	static float getDistanceThreshold(const dtLinkBuilderConfig& config, const dtNavLinkAction action)
-	{
-		switch(action)
-		{
-		case DT_LINK_ACTION_JUMP_DOWN:
-			return config.jumpDownConfig.filterDistanceThreshold;
-		case DT_LINK_ACTION_JUMP_OVER:
-			return config.jumpOverConfig.filterDistanceThreshold;
-		default:
-			const bool dtval = false;
-			dtAssert(dtval);
-			return 100.f;
-		}
-	}
 }
 
 dtNavLinkBuilder::GroundSegment::~GroundSegment()
@@ -309,7 +294,7 @@ void dtNavLinkBuilder::addEdgeLinks(const dtLinkBuilderConfig& builderConfig, co
 #if !(UE_BUILD_SHIPPING || UE_BUILD_TEST)
 					link.debugSourceEdge = (short)edgeIndex;
 #endif
-					link.action = es->action;
+					link.configIndex = es->configIndex;
 					link.flags = VALID;
 					link.nspine = es->trajectory.nspine;
 
@@ -402,12 +387,12 @@ void dtNavLinkBuilder::filterOverlappingLinks(const float edgeDistanceThreshold)
 	}
 }
 
-void dtNavLinkBuilder::buildForAllEdges(rcContext& ctx, const dtLinkBuilderConfig& builderConfig, dtNavLinkAction action)
+void dtNavLinkBuilder::buildForAllEdges(rcContext& ctx, const dtLinkBuilderConfig& builderConfig, const dtNavLinkBuilderJumpConfig& jumpConfig)
 {
 	for (int i = 0; i < m_edges.Num(); ++i)
 	{
 		EdgeSampler sampler;
-		const bool success = sampleEdge(builderConfig, action, m_edges[i].sp, m_edges[i].sq, &sampler);
+		const bool success = sampleEdge(jumpConfig, m_edges[i].sp, m_edges[i].sq, &sampler);
 		if (success)
 		{
 			addEdgeLinks(builderConfig, &sampler, i);
@@ -416,11 +401,10 @@ void dtNavLinkBuilder::buildForAllEdges(rcContext& ctx, const dtLinkBuilderConfi
 
 	ctx.log(RC_LOG_PROGRESS, "   %i links added.", m_links.Num());
 
-	const float distanceThreshold = UE::Detour::NavLink::Private::getDistanceThreshold(builderConfig, action);
-	filterOverlappingLinks(distanceThreshold);
+	filterOverlappingLinks(jumpConfig.filterDistanceThreshold);
 }
 
-void dtNavLinkBuilder::debugBuildEdge(const dtLinkBuilderConfig& builderConfig, dtNavLinkAction action, int edgeIndex, EdgeSampler& sampler)
+void dtNavLinkBuilder::debugBuildEdge(const dtLinkBuilderConfig& builderConfig, int edgeIndex, EdgeSampler& sampler, const dtNavLinkBuilderJumpConfig& jumpConfig)
 {
 	if (edgeIndex >= m_edges.Num())
 	{
@@ -429,14 +413,13 @@ void dtNavLinkBuilder::debugBuildEdge(const dtLinkBuilderConfig& builderConfig,
 	
 	m_debugSelectedEdge = edgeIndex;
 
-	const bool success = sampleEdge(builderConfig, action, m_edges[edgeIndex].sp, m_edges[edgeIndex].sq, &sampler);
+	const bool success = sampleEdge(jumpConfig, m_edges[edgeIndex].sp, m_edges[edgeIndex].sq, &sampler);
 	if (success)
 	{
 		addEdgeLinks(builderConfig, &sampler, edgeIndex);
 	}
 
-	const float distanceThreshold = UE::Detour::NavLink::Private::getDistanceThreshold(builderConfig, action);
-	filterOverlappingLinks(distanceThreshold);
+	filterOverlappingLinks(jumpConfig.filterDistanceThreshold);
 }
 
 bool dtNavLinkBuilder::getCompactHeightfieldHeight(const dtReal* pt, const dtReal hrange, dtReal* height) const
@@ -604,7 +587,7 @@ void dtNavLinkBuilder::updateTrajectorySamples(EdgeSampler* es) const
 		if ((ssmp.flags & HAS_GROUND) == 0 || (esmp.flags & HAS_GROUND) == 0)
 			continue;
 
-		// When we sample ground segments, in sampleEdges, we have add least 2 samples.
+		// When we sampled ground segments, in sampleEdges, we added least 2 samples.
 		check(nsamples >= 2);																
 		const dtReal u = (dtReal)i/(dtReal)(nsamples-1);
 		dtReal spt[3], ept[3];
@@ -903,11 +886,10 @@ int dtNavLinkBuilder::findPotentialJumpOverEdges(const dtReal* sp, const dtReal*
 }
 	
 
-void dtNavLinkBuilder::initJumpDownRig(EdgeSampler* es, const dtReal* sp, const dtReal* sq,
-									   const dtNavLinkBuilderJumpDownConfig& config) const
+void dtNavLinkBuilder::initJumpRig(EdgeSampler* es, const dtReal* sp, const dtReal* sq, const dtNavLinkBuilderJumpConfig& config) const
 {
-	es->action = DT_LINK_ACTION_JUMP_DOWN;
-
+	es->configIndex = config.ConfigIndex;
+	
 	// Set axes
 	dtVsub(es->ax, sq, sp);
 	dtVnormalize(es->ax);
@@ -957,82 +939,14 @@ void dtNavLinkBuilder::initJumpDownRig(EdgeSampler* es, const dtReal* sp, const
 	es->groundRange = config.jumpEndsHeightTolerance;
 }
 
-void dtNavLinkBuilder::initJumpOverRig(EdgeSampler* es, const dtReal* sp, const dtReal* sq,
-										const float jumpStartDist, const float jumpEndDist,
-										const float jumpHeight, const float groundRange)
-{
-	es->action = DT_LINK_ACTION_JUMP_OVER;
-
-	// Set edge
-	dtVcopy(es->rigp, sp);
-	dtVcopy(es->rigq, sq);
-
-	// Set axes	
-	dtVsub(es->ax, sq, sp);
-	dtVnormalize(es->ax);
-	dtVset(es->az, es->ax[2], 0, -es->ax[0]);
-	dtVnormalize(es->az);
-	dtVset(es->ay, 0, 1, 0);
-	
-	// Build action sampling spine.
-	es->trajectory.nspine = MAX_SPINE;
-	for (int i = 0; i < MAX_SPINE; ++i)
-	{
-		float* pt = &es->trajectory.spine[i*2];
-		const float u = (float)i/(float)(MAX_SPINE-1);
-		pt[0] = jumpStartDist + u * (jumpEndDist - jumpStartDist);
-		pt[1] = (1-dtSqr(u*2-1)) * jumpHeight;
-	}
-	
-	es->groundRange = groundRange;
-}
-
-bool dtNavLinkBuilder::sampleEdge(const dtLinkBuilderConfig& builderConfig, dtNavLinkAction desiredAction, const dtReal* sp, const dtReal* sq, dtNavLinkBuilder::EdgeSampler* es) const
+bool dtNavLinkBuilder::sampleEdge(const dtNavLinkBuilderJumpConfig& jumpConfig, const dtReal* sp, const dtReal* sq, dtNavLinkBuilder::EdgeSampler* es) const
 {
 	TRACE_CPUPROFILER_EVENT_SCOPE(dtNavLinkBuilder::sampleEdge);
 	
 	using namespace UE::Detour::NavLink::Private;
 
-	float samplingSeparationFactor = 1.f;
-	
-	if (desiredAction == DT_LINK_ACTION_JUMP_DOWN)
-	{
-		const dtNavLinkBuilderJumpDownConfig& config = builderConfig.jumpDownConfig;
-		samplingSeparationFactor = config.samplingSeparationFactor;
-		initJumpDownRig(es, sp, sq, config);
-	}
-	else if (desiredAction == DT_LINK_ACTION_JUMP_OVER)
-	{
-		const dtNavLinkBuilderJumpOverConfig& config = builderConfig.jumpOverConfig;
-		const float jumpDist = config.jumpGapWidth;
-		const float heightRange = config.jumpGapHeightTolerance;
-		static constexpr int NSEGS = 8;
-		dtReal segs[NSEGS*6];
-		int nsegs = findPotentialJumpOverEdges(sp, sq, jumpDist, heightRange, segs, NSEGS);
-
-		int ibest = -1;
-		float dbest = 0;
-		for (int i = 0; i < nsegs; ++i)
-		{
-			const dtReal* seg = &segs[i*6];
-			const float d = dtVdistSqr(seg,seg+3);
-			if (d > dbest)
-			{
-				dbest = d;
-				ibest = i;
-			}
-		}
-		if (ibest == -1)
-		{
-			return false;
-		}
-
-		const float jumpStartDist = config.jumpDistanceFromGapCenter; 
-		const float jumpHeight = config.jumpHeight;
-		const float groundRange = config.jumpEndsHeightTolerance;
-		samplingSeparationFactor = config.samplingSeparationFactor;
-		initJumpOverRig(es, &segs[ibest*6+0], &segs[ibest*6+3], -jumpStartDist, jumpStartDist, jumpHeight, groundRange);
-	}
+	const float samplingSeparationFactor = jumpConfig.samplingSeparationFactor;
+	initJumpRig(es, sp, sq, jumpConfig);
 
 	initTrajectorySamples(es->groundRange, &es->trajectory);
 	
diff --git a/Engine/Source/Runtime/Navmesh/Private/Detour/DetourNavLinkBuilderConfig.cpp b/Engine/Source/Runtime/Navmesh/Private/Detour/DetourNavLinkBuilderConfig.cpp
--- a/Engine/Source/Runtime/Navmesh/Private/Detour/DetourNavLinkBuilderConfig.cpp
+++ b/Engine/Source/Runtime/Navmesh/Private/Detour/DetourNavLinkBuilderConfig.cpp
@@ -3,7 +3,7 @@
 
 #include "Detour/DetourCommon.h"
 
-void dtNavLinkBuilderJumpDownConfig::init()
+void dtNavLinkBuilderJumpConfig::init()
 {
 	check(jumpHeight >= 0.f);
 	check(jumpMaxDepth >= 0.f);
diff --git a/Engine/Source/Runtime/Navmesh/Public/Detour/DetourNavLinkBuilder.h b/Engine/Source/Runtime/Navmesh/Public/Detour/DetourNavLinkBuilder.h
--- a/Engine/Source/Runtime/Navmesh/Public/Detour/DetourNavLinkBuilder.h
+++ b/Engine/Source/Runtime/Navmesh/Public/Detour/DetourNavLinkBuilder.h
@@ -7,17 +7,10 @@
 #include "Detour/DetourNavLinkBuilderConfig.h"
 
 //@UE BEGIN
-enum dtNavLinkAction
-{
-	DT_LINK_ACTION_UNSET,
-	DT_LINK_ACTION_JUMP_DOWN,
-	DT_LINK_ACTION_JUMP_OVER,
-};
 
 struct dtLinkBuilderConfig
 {
-	dtNavLinkBuilderJumpDownConfig jumpDownConfig;
-	dtNavLinkBuilderJumpOverConfig jumpOverConfig;
+	TArray<dtNavLinkBuilderJumpConfig> jumpConfigs;
 	
 	dtReal agentRadius = 0;
 	dtReal agentHeight = 0;
@@ -113,6 +106,8 @@ private:
 	};
 	
 public:
+	static constexpr uint8 INVALID_CONFIG_INDEX = 0xFF;
+	
 	struct EdgeSampler
 	{
 		Trajectory2D trajectory;
@@ -124,8 +119,7 @@ public:
 		dtReal ax[3], ay[3], az[3];		// axis along edge
 
 		float groundRange;
-		
-		dtNavLinkAction action = DT_LINK_ACTION_UNSET;
+		uint8 configIndex = INVALID_CONFIG_INDEX;
 	};
 
 	enum JumpLinkFlag : unsigned char
@@ -140,7 +134,7 @@ public:
 		dtReal spine1[MAX_SPINE*3];
 		int nspine = 0;
 		JumpLinkFlag flags = VALID;
-		dtNavLinkAction action = DT_LINK_ACTION_UNSET;
+		uint8 configIndex = INVALID_CONFIG_INDEX;
 #if !(UE_BUILD_SHIPPING || UE_BUILD_TEST)		
 		short debugSourceEdge = -1;
 #endif
@@ -162,9 +156,9 @@ public:
 							   const rcHeightfield* solidHF, const rcCompactHeightfield* compactHF);
 
 	// For all edges, sample edges (sampleEdge) and add links to m_links
-	NAVMESH_API void buildForAllEdges(rcContext& ctx, const dtLinkBuilderConfig& acfg, dtNavLinkAction action);
+	NAVMESH_API void buildForAllEdges(rcContext& ctx, const dtLinkBuilderConfig& acfg, const dtNavLinkBuilderJumpConfig& jumpConfig);
 
-	NAVMESH_API void debugBuildEdge(const dtLinkBuilderConfig& acfg, dtNavLinkAction action, int edgeIndex, EdgeSampler& sampler);
+	NAVMESH_API void debugBuildEdge(const dtLinkBuilderConfig& acfg, int edgeIndex, EdgeSampler& sampler, const dtNavLinkBuilderJumpConfig& jumpConfig);
 	
 private:
 	void initTrajectorySamples(const float groundRange, Trajectory2D* trajectory) const;
@@ -174,12 +168,8 @@ private:
 								   const float depthRange, const float heightRange,
 								   dtReal* outSegs, const int maxOutSegs) const;
 	
-	void initJumpDownRig(EdgeSampler* es, const dtReal* sp, const dtReal* sq,
-						 const dtNavLinkBuilderJumpDownConfig& config) const;
-
-	static void initJumpOverRig(EdgeSampler* es, const dtReal* sp, const dtReal* sq,
-						 const float jumpStartDist, const float jumpEndDist,
-						 const float jumpHeight, const float groundRange);
+	void initJumpRig(EdgeSampler* es, const dtReal* sp, const dtReal* sq,
+						 const dtNavLinkBuilderJumpConfig& config) const;
 
 	bool getCompactHeightfieldHeight(const dtReal* pt, const dtReal hrange, dtReal* height) const;
 	bool checkHeightfieldCollision(const dtReal x, const dtReal ymin, const dtReal ymax, const dtReal z) const;
@@ -194,7 +184,7 @@ private:
 	
 	void filterOverlappingLinks(const float edgeDistanceThreshold);
 
-	bool sampleEdge(const dtLinkBuilderConfig& builderConfig, dtNavLinkAction desiredAction, const dtReal* sp, const dtReal* sq, dtNavLinkBuilder::EdgeSampler* sampler) const;
+	bool sampleEdge(const dtNavLinkBuilderJumpConfig& jumpConfig, const dtReal* sp, const dtReal* sq, dtNavLinkBuilder::EdgeSampler* sampler) const;
 	void addEdgeLinks(const dtLinkBuilderConfig& builderConfig, const EdgeSampler* es, const int edgeIndex);
 };
 //@UE END
diff --git a/Engine/Source/Runtime/Navmesh/Public/Detour/DetourNavLinkBuilderConfig.h b/Engine/Source/Runtime/Navmesh/Public/Detour/DetourNavLinkBuilderConfig.h
--- a/Engine/Source/Runtime/Navmesh/Public/Detour/DetourNavLinkBuilderConfig.h
+++ b/Engine/Source/Runtime/Navmesh/Public/Detour/DetourNavLinkBuilderConfig.h
@@ -9,8 +9,8 @@ enum dtNavLinkBuilderFlags : unsigned short
 	DT_NAVLINK_CREATE_EXTREMITY_LINKS	= 1 << 1,
 };
 
-/** Configuration for generated jump down links. */
-struct dtNavLinkBuilderJumpDownConfig
+/** Configuration for generated jump links. */
+struct dtNavLinkBuilderJumpConfig
 {
 	/// Should this config be used to generate links.
 	bool enabled = true;
@@ -62,27 +62,32 @@ struct dtNavLinkBuilderJumpDownConfig
 	/// User defined area id assigned to upwards traversal of the off-mesh connections
 	unsigned char upDirArea = 0;
 
+	/// Unique index in TArray<dtNavLinkBuilderJumpConfig> identifying this config
+	static constexpr uint8 INVALID_CONFIG_INDEX = 0xFF;
+	uint8 ConfigIndex = INVALID_CONFIG_INDEX;
+
 	/// Initialize the configuration by computing cached values.
 	NAVMESH_API void init();
 };
 
-/** Configuration for generated jump over links. */
-struct dtNavLinkBuilderJumpOverConfig
+
+/** Configuration for generated jump down links. */
+struct UE_DEPRECATED(5.7, "Use dtNavLinkBuilderJumpConfig instead.") dtNavLinkBuilderJumpDownConfig
 {
 	/// Should this config be used to generate links.
-	bool enabled = false;	// @todo: Disabled for now since this configuration not currently exposed.
+	bool enabled = true;
 
-	/// Maximum jumpable gap size used when matching edges to jump over. [Limit: > 0] [Units: wu]
-	float jumpGapWidth = 200.f;
+	/// Horizontal length of the jump. How far from the starting point we will look for ground. [Limit: > 0] [Units: wu]
+	float jumpLength = 150.f;
 
-	/// Vertical tolerance used when matching edges to jump over. [Limit: > 0] [Units: wu]
-	float jumpGapHeightTolerance = 100.f; 
+	/// How far from the edge is the jump started. [Limit: > 0] [Units: wu]
+	float jumpDistanceFromEdge = 10.f;
 
-	/// How far from the center of the gap is the jump started. [Limit: > 0] [Units: wu]
-	float jumpDistanceFromGapCenter = 100.f; 
+	/// How far below the starting height we want to look for landing ground. [Limit: > 0] [Units: wu]
+	float jumpMaxDepth = 150.f;
 
-	/// Height at the top of the jump trajectory. [Limit: > 0] [Units: wu]
-	float jumpHeight = 100.f; 
+	/// Peak height relative to the height of the starting point. [Limit: >= 0] [Units: wu]
+	float jumpHeight = 50.f;
 
 	/// Tolerance at both ends of the jump to find ground. [Limit: > 0] [Units: wu]
 	float jumpEndsHeightTolerance = 50.f;
@@ -94,4 +99,31 @@ struct dtNavLinkBuilderJumpOverConfig
 	/// When filtering similar links, distance used to compare between segment endpoints to match similar links.
 	/// Use greater distance for more filtering (0 to deactivate filtering). [Limit: > 0] [Units: wu]
 	float filterDistanceThreshold = 80.f;
-};
+
+	/// Flags used to indicate how links will be added.
+	unsigned short linkBuilderFlags = DT_NAVLINK_CREATE_CENTER_POINT_LINK;
+
+	/// Cached parabola constant fitting the configuration parameters. 
+	float cachedParabolaConstant = 0;
+
+	/// Cached value used when computing jump trajectory.
+	float cachedDownRatio = 0;
+	
+	/// User id used to handle links made from this configuration.
+	unsigned long long linkUserId = 0;
+
+	/// User defined flags assigned to downwards traversal of the off-mesh connections
+	unsigned short downDirPolyFlag = 0;
+
+	/// User defined flags assigned to upwards traversal of the off-mesh connections
+	unsigned short upDirPolyFlag = 0;
+
+	/// User defined area id assigned to downards traversal of the off-mesh connections
+	unsigned char downDirArea = 0;
+
+	/// User defined area id assigned to upwards traversal of the off-mesh connections
+	unsigned char upDirArea = 0;
+
+	/// Initialize the configuration by computing cached values.
+	NAVMESH_API void init();
+};
\ No newline at end of file
