Remove deprecated 5.4 navigation code

diff --git a/Engine/Source/Runtime/NavigationSystem/Private/NavMesh/RecastNavMesh.cpp b/Engine/Source/Runtime/NavigationSystem/Private/NavMesh/RecastNavMesh.cpp
--- a/Engine/Source/Runtime/NavigationSystem/Private/NavMesh/RecastNavMesh.cpp
+++ b/Engine/Source/Runtime/NavigationSystem/Private/NavMesh/RecastNavMesh.cpp
@@ -1111,28 +1111,6 @@ void ARecastNavMesh::SortAreasForGenerator(TArray<FRecastAreaNavModifierElement>
 	Modifiers.Sort(FNavAreaSortPredicate());
 }
 
-#if WITH_EDITORONLY_DATA
-// Deprecated
-const TArray<FIntPoint>& ARecastNavMesh::GetActiveTiles() const
-{
-	const FRecastNavMeshGenerator* MyGenerator = static_cast<const FRecastNavMeshGenerator*>(GetGenerator());
-	check(MyGenerator);
-PRAGMA_DISABLE_DEPRECATION_WARNINGS
-	return MyGenerator->ActiveTiles;
-PRAGMA_ENABLE_DEPRECATION_WARNINGS
-}
-
-// Deprecated
-TArray<FIntPoint>& ARecastNavMesh::GetActiveTiles()
-{
-	FRecastNavMeshGenerator* MyGenerator = static_cast<FRecastNavMeshGenerator*>(GetGenerator());
-	check(MyGenerator);
-PRAGMA_DISABLE_DEPRECATION_WARNINGS	
-	return MyGenerator->ActiveTiles;
-PRAGMA_ENABLE_DEPRECATION_WARNINGS
-}
-#endif // WITH_EDITORONLY_DATA
-
 const TSet<FIntPoint>& ARecastNavMesh::GetActiveTileSet() const
 {
 	const FRecastNavMeshGenerator* MyGenerator = static_cast<const FRecastNavMeshGenerator*>(GetGenerator());
diff --git a/Engine/Source/Runtime/NavigationSystem/Private/NavMesh/RecastNavMeshGenerator.cpp b/Engine/Source/Runtime/NavigationSystem/Private/NavMesh/RecastNavMeshGenerator.cpp
--- a/Engine/Source/Runtime/NavigationSystem/Private/NavMesh/RecastNavMeshGenerator.cpp
+++ b/Engine/Source/Runtime/NavigationSystem/Private/NavMesh/RecastNavMeshGenerator.cpp
@@ -6161,11 +6161,6 @@ bool FRecastNavMeshGenerator::IsAllowedToAddTileLayers(const FIntPoint Tile) con
 }
 
 #if !UE_BUILD_SHIPPING
-// Deprecated
-void FRecastNavMeshGenerator::LogDirtyAreas(const TMap<FPendingTileElement, TArray<FNavigationDirtyAreaPerTileDebugInformation>>& DirtyAreasDebuggingInformation) const
-{
-}
-
 void FRecastNavMeshGenerator::LogDirtyAreas(const UObject& OwnerNav,
 	const TMap<FPendingTileElement, TArray<FNavigationDirtyAreaPerTileDebugInformation>>& DirtyAreasDebuggingInformation) const
 {
@@ -7423,16 +7418,6 @@ void FRecastNavMeshGenerator::GetDebugGeometry(const FNavigationRelevantData& En
 }
 #endif // !UE_BUILD_SHIPPING
 
-// Deprecated
-void FRecastNavMeshGenerator::ExportComponentGeometry(UActorComponent* InOutComponent, FNavigationRelevantData& OutData)
-{
-	if (INavRelevantInterface* NavRelevantInterface = Cast<INavRelevantInterface>(InOutComponent))
-	{
-		const TSharedRef<const FNavigationElement> TmpElement = FNavigationElement::CreateFromNavRelevantInterface(*NavRelevantInterface);
-		FRecastGeometryExport::ExportElementGeometry(TmpElement.Get(), OutData);
-	}
-}
-
 // Deprecated
 void FRecastNavMeshGenerator::ExportNavRelevantObjectGeometry(INavRelevantInterface& InOutNavRelevantInterface, FNavigationRelevantData& OutData)
 {
@@ -7471,16 +7456,6 @@ void FRecastGeometryExport::ExportVertexSoupGeometry(const TArray<FVector>& InVe
 	RecastGeometryExport::StoreCollisionCache(GeomExport);
 }
 
-// Deprecated
-void FRecastNavMeshGenerator::ExportRigidBodyGeometry(
-	UBodySetup& InOutBodySetup,
-	TNavStatArray<FVector>& OutVertexBuffer,
-	TNavStatArray<int32>& OutIndexBuffer,
-	const FTransform& LocalToWorld)
-{
-	FBox TempBounds;
-	FRecastGeometryExport::ExportRigidBodyGeometry(InOutBodySetup, OutVertexBuffer, OutIndexBuffer, TempBounds, LocalToWorld);
-}
 
 // Deprecated
 void FRecastNavMeshGenerator::ExportRigidBodyGeometry(
@@ -7505,27 +7480,6 @@ void FRecastGeometryExport::ExportRigidBodyGeometry(UBodySetup& InOutBodySetup,
 	}
 }
 
-// Deprecated
-void FRecastNavMeshGenerator::ExportRigidBodyGeometry(
-	UBodySetup& InOutBodySetup,
-	TNavStatArray<FVector>& OutTriMeshVertexBuffer,
-	TNavStatArray<int32>& OutTriMeshIndexBuffer,
-	TNavStatArray<FVector>& OutConvexVertexBuffer,
-	TNavStatArray<int32>& OutConvexIndexBuffer,
-	TNavStatArray<int32>& OutShapeBuffer,
-	const FTransform& LocalToWorld)
-{
-	FBox TempBounds;
-	FRecastGeometryExport::ExportRigidBodyGeometry(
-		InOutBodySetup,
-		OutTriMeshVertexBuffer,
-		OutTriMeshIndexBuffer,
-		OutConvexVertexBuffer,
-		OutConvexIndexBuffer,
-		OutShapeBuffer,
-		TempBounds,
-		LocalToWorld);
-}
 
 // Deprecated
 void FRecastNavMeshGenerator::ExportRigidBodyGeometry(
@@ -7576,18 +7530,6 @@ void FRecastGeometryExport::ExportRigidBodyGeometry(UBodySetup& InOutBodySetup,
 	}
 }
 
-// Deprecated
-void FRecastNavMeshGenerator::ExportAggregatedGeometry(
-	const FKAggregateGeom& AggGeom,
-	TNavStatArray<FVector>& OutConvexVertexBuffer,
-	TNavStatArray<int32>& OutConvexIndexBuffer,
-	TNavStatArray<int32>& OutShapeBuffer,
-	const FTransform& LocalToWorld)
-{
-	FBox TempBounds;
-	FRecastGeometryExport::ExportAggregatedGeometry(AggGeom, OutConvexVertexBuffer, OutConvexIndexBuffer, OutShapeBuffer, TempBounds, LocalToWorld);
-}
-
 // Deprecated
 void FRecastNavMeshGenerator::ExportAggregatedGeometry(
 	const FKAggregateGeom& AggGeom,
diff --git a/Engine/Source/Runtime/NavigationSystem/Private/NavigationDataHandler.cpp b/Engine/Source/Runtime/NavigationSystem/Private/NavigationDataHandler.cpp
--- a/Engine/Source/Runtime/NavigationSystem/Private/NavigationDataHandler.cpp
+++ b/Engine/Source/Runtime/NavigationSystem/Private/NavigationDataHandler.cpp
@@ -498,12 +498,6 @@ bool FNavigationDataHandler::UpdateNavOctreeElementBounds(const FNavigationEleme
 	return false;
 }
 
-// Deprecated
-bool FNavigationDataHandler::UpdateNavOctreeElementBounds(UObject& Object, const FBox& NewBounds, const FBox& DirtyArea)
-{
-	return UpdateNavOctreeElementBounds(FNavigationElementHandle(&Object), NewBounds, TConstArrayView<FBox>{DirtyArea});
-}
-
 void FNavigationDataHandler::FindElementsInNavOctree(const FBox& QueryBox, const FNavigationOctreeFilter& Filter, TArray<FNavigationOctreeElement>& Elements)
 {
 	if (OctreeController.IsValid() == false)
diff --git a/Engine/Source/Runtime/NavigationSystem/Private/NavigationSystem.cpp b/Engine/Source/Runtime/NavigationSystem/Private/NavigationSystem.cpp
--- a/Engine/Source/Runtime/NavigationSystem/Private/NavigationSystem.cpp
+++ b/Engine/Source/Runtime/NavigationSystem/Private/NavigationSystem.cpp
@@ -716,11 +716,6 @@ void FNavRegenTimeSliceManager::LogTileStatistics(const TArray<TObjectPtr<ANavig
 //----------------------------------------------------------------------//
 bool UNavigationSystemV1::bNavigationAutoUpdateEnabled = true;
 
-PRAGMA_DISABLE_DEPRECATION_WARNINGS
-TMap<INavLinkCustomInterface*, FWeakObjectPtr> UNavigationSystemV1::PendingCustomLinkRegistration;
-FCriticalSection UNavigationSystemV1::CustomLinkRegistrationSection;
-PRAGMA_ENABLE_DEPRECATION_WARNINGS
-
 #if !(UE_BUILD_SHIPPING || UE_BUILD_TEST)
 FNavigationSystemExec UNavigationSystemV1::ExecHandler;
 #endif // !UE_BUILD_SHIPPING
@@ -4009,23 +4004,6 @@ void UNavigationSystemV1::UpdateNavOctreeParentChain(UObject* ElementOwner, bool
 	PRAGMA_ENABLE_DEPRECATION_WARNINGS
 }
 
-// Deprecated
-bool UNavigationSystemV1::UpdateNavOctreeElementBounds(UActorComponent* Comp, const FBox& NewBounds, const FBox& DirtyArea)
-{
-	if (Repository == nullptr)
-	{
-		return false;
-	}
-
-	const FNavigationElementHandle Handle = Repository->GetNavigationElementHandleForUObject(Comp);
-	if (Handle.IsValid())
-	{
-		return UpdateNavOctreeElementBounds(Handle, NewBounds, { DirtyArea });
-	}
-
-	return false;
-}
-
 // Deprecated
 bool UNavigationSystemV1::UpdateNavOctreeElementBounds(UObject& Object, const FBox& NewBounds, TConstArrayView<FBox> DirtyAreas)
 {
@@ -4477,18 +4455,6 @@ void UNavigationSystemV1::GatherNavigationBounds()
 	}
 }
 
-// Deprecated
-void UNavigationSystemV1::GetInvokerSeedLocations(const UWorld& InWorld, TArray<FVector2D, TInlineAllocator<32>>& OutSeedLocations)
-{
-	TArray<FVector, TInlineAllocator<32>> Locations;
-	GetInvokerSeedLocations(InWorld, Locations);
-
-	for (const FVector Location : Locations)
-	{
-		OutSeedLocations.Add(FVector2D(Location));
-	}
-}
-
 void UNavigationSystemV1::GetInvokerSeedLocations(const UWorld& InWorld, TArray<FVector, TInlineAllocator<32>>& OutSeedLocations)
 {
 	for (FConstPlayerControllerIterator PlayerIt = InWorld.GetPlayerControllerIterator(); PlayerIt; ++PlayerIt)
diff --git a/Engine/Source/Runtime/NavigationSystem/Public/NavMesh/RecastNavMesh.h b/Engine/Source/Runtime/NavigationSystem/Public/NavMesh/RecastNavMesh.h
--- a/Engine/Source/Runtime/NavigationSystem/Public/NavMesh/RecastNavMesh.h
+++ b/Engine/Source/Runtime/NavigationSystem/Public/NavMesh/RecastNavMesh.h
@@ -1060,11 +1060,6 @@ public:
 	NAVIGATIONSYSTEM_API virtual void AttachNavMeshDataChunk(URecastNavMeshDataChunk& NavDataChunk);
 	NAVIGATIONSYSTEM_API virtual void DetachNavMeshDataChunk(URecastNavMeshDataChunk& NavDataChunk);
 
-	UE_DEPRECATED(5.4, "Use GetActiveTileSet instead.")
-	NAVIGATIONSYSTEM_API const TArray<FIntPoint>& GetActiveTiles() const;
-	UE_DEPRECATED(5.4, "Use GetActiveTileSet instead.")
-	NAVIGATIONSYSTEM_API TArray<FIntPoint>& GetActiveTiles(); 
-
 	NAVIGATIONSYSTEM_API const TSet<FIntPoint>& GetActiveTileSet() const;
 	NAVIGATIONSYSTEM_API TSet<FIntPoint>& GetActiveTileSet(); 
 	
diff --git a/Engine/Source/Runtime/NavigationSystem/Public/NavMesh/RecastNavMeshGenerator.h b/Engine/Source/Runtime/NavigationSystem/Public/NavMesh/RecastNavMeshGenerator.h
--- a/Engine/Source/Runtime/NavigationSystem/Public/NavMesh/RecastNavMeshGenerator.h
+++ b/Engine/Source/Runtime/NavigationSystem/Public/NavMesh/RecastNavMeshGenerator.h
@@ -862,34 +862,6 @@ public:
 	NAVIGATIONSYSTEM_API virtual void GrabDebugSnapshot(struct FVisualLogEntry* Snapshot, const FBox& BoundingBox, const FName& CategoryName, ELogVerbosity::Type Verbosity) const override;
 #endif
 
-	UE_DEPRECATED(5.4, "Use FRecastGeometryExport::ExportElementGeometry")
-	static NAVIGATIONSYSTEM_API void ExportComponentGeometry(UActorComponent* InOutComponent, FNavigationRelevantData& OutData);
-
-	UE_DEPRECATED(5.4, "Use FRecastGeometryExport::ExportRigidBodyGeometry.")
-	static NAVIGATIONSYSTEM_API void ExportRigidBodyGeometry(
-		UBodySetup& InOutBodySetup,
-		TNavStatArray<FVector>& OutVertexBuffer,
-		TNavStatArray<int32>& OutIndexBuffer,
-		const FTransform& LocalToWorld = FTransform::Identity);
-
-	UE_DEPRECATED(5.4, "Use FRecastGeometryExport::ExportRigidBodyGeometry.")
-	static NAVIGATIONSYSTEM_API void ExportRigidBodyGeometry(
-		UBodySetup& InOutBodySetup,
-		TNavStatArray<FVector>& OutTriMeshVertexBuffer,
-		TNavStatArray<int32>& OutTriMeshIndexBuffer,
-		TNavStatArray<FVector>& OutConvexVertexBuffer,
-		TNavStatArray<int32>& OutConvexIndexBuffer,
-		TNavStatArray<int32>& OutShapeBuffer,
-		const FTransform& LocalToWorld = FTransform::Identity);
-
-	UE_DEPRECATED(5.4, "Use FRecastGeometryExport::ExportAggregatedGeometry.")
-	static NAVIGATIONSYSTEM_API void ExportAggregatedGeometry(
-		const FKAggregateGeom& AggGeom,
-		TNavStatArray<FVector>& OutConvexVertexBuffer,
-		TNavStatArray<int32>& OutConvexIndexBuffer,
-		TNavStatArray<int32>& OutShapeBuffer,
-		const FTransform& LocalToWorld = FTransform::Identity);
-
 	UE_DEPRECATED(5.5, "Use FRecastGeometryExport::ExportElementGeometry.")
 	static NAVIGATIONSYSTEM_API void ExportNavRelevantObjectGeometry(INavRelevantInterface& InOutNavRelevantInterface, FNavigationRelevantData& OutData);
 	UE_DEPRECATED(5.5, "Use FRecastGeometryExport::ExportVertexSoupGeometry.")
@@ -1080,9 +1052,6 @@ PRAGMA_ENABLE_DEPRECATION_WARNINGS
 		FNavigationDirtyArea DirtyArea;
 		bool bTileWasAlreadyAdded = false;
 	};
-
-	UE_DEPRECATED(5.4, "Use new version with ARecastNavMesh owner instead.")
-	NAVIGATIONSYSTEM_API void LogDirtyAreas(const TMap<FPendingTileElement, TArray<FNavigationDirtyAreaPerTileDebugInformation>>& DirtyAreasDebuggingInformation) const; 
 	
 	/** Used internally, when LogNavigationDirtyArea is VeryVerbose, to log the number of tiles a dirty area is requesting. */
 	NAVIGATIONSYSTEM_API void LogDirtyAreas(const UObject& OwnerNav,
@@ -1142,11 +1111,6 @@ protected:
 	/** List of tiles that were recently regenerated */
 	TNavStatArray<FTileTimestamp> RecentlyBuiltTiles;
 #endif// WITH_EDITOR
-
-#if WITH_EDITORONLY_DATA	
-	UE_DEPRECATED(5.4, "Use ActiveTileSet instead.")
-	TArray<FIntPoint> ActiveTiles;
-#endif // WITH_EDITORONLY_DATA
 	
 	TSet<FIntPoint> ActiveTileSet;
 
diff --git a/Engine/Source/Runtime/NavigationSystem/Public/NavigationDataHandler.h b/Engine/Source/Runtime/NavigationSystem/Public/NavigationDataHandler.h
--- a/Engine/Source/Runtime/NavigationSystem/Public/NavigationDataHandler.h
+++ b/Engine/Source/Runtime/NavigationSystem/Public/NavigationDataHandler.h
@@ -73,9 +73,6 @@ struct FNavigationDataHandler
 
 	NAVIGATIONSYSTEM_API void DemandLazyDataGathering(FNavigationRelevantData& ElementData);
 
-	UE_DEPRECATED(5.4, "Use the overloaded version taking a list of areas as parameter instead.")
-	NAVIGATIONSYSTEM_API bool UpdateNavOctreeElementBounds(UObject& Object, const FBox& NewBounds, const FBox& DirtyArea);
-
 private:
 	void UpdateNavOctreeParentChain(const TSharedRef<const FNavigationElement>& Element, bool bSkipElementOwnerUpdate = false);
 };
diff --git a/Engine/Source/Runtime/NavigationSystem/Public/NavigationOctree.h b/Engine/Source/Runtime/NavigationSystem/Public/NavigationOctree.h
--- a/Engine/Source/Runtime/NavigationSystem/Public/NavigationOctree.h
+++ b/Engine/Source/Runtime/NavigationSystem/Public/NavigationOctree.h
@@ -172,13 +172,6 @@ struct FNavigationOctreeSemantics
 class FNavigationOctree : public TOctree2<FNavigationOctreeElement, FNavigationOctreeSemantics>, public TSharedFromThis<FNavigationOctree, ESPMode::ThreadSafe>
 {
 public:
-	UE_DEPRECATED(5.4, "Use FNavRelevantGeometryExportDelegate.")
-	DECLARE_DELEGATE_TwoParams(FNavigableGeometryComponentExportDelegate, UActorComponent*, FNavigationRelevantData&);
-
-PRAGMA_DISABLE_DEPRECATION_WARNINGS
-	FNavigableGeometryComponentExportDelegate ComponentExportDelegate;
-PRAGMA_ENABLE_DEPRECATION_WARNINGS
-
 	UE_DEPRECATED(5.5, "Use FGeometryExportDelegate.")
 	DECLARE_DELEGATE_TwoParams(FNavRelevantGeometryExportDelegate, INavRelevantInterface&, FNavigationRelevantData&);
 PRAGMA_DISABLE_DEPRECATION_WARNINGS
diff --git a/Engine/Source/Runtime/NavigationSystem/Public/NavigationSystem.h b/Engine/Source/Runtime/NavigationSystem/Public/NavigationSystem.h
--- a/Engine/Source/Runtime/NavigationSystem/Public/NavigationSystem.h
+++ b/Engine/Source/Runtime/NavigationSystem/Public/NavigationSystem.h
@@ -856,12 +856,6 @@ protected:
 	/** Registers custom navigation links awaiting registration in the navigation object repository */
 	NAVIGATIONSYSTEM_API void ProcessCustomLinkPendingRegistration();
 
-	UE_DEPRECATED(5.4, "Obsolete since navigation related objects register to the NavigationObjectRepository world subsytem.")
-	static NAVIGATIONSYSTEM_API FCriticalSection CustomLinkRegistrationSection;
-
-	UE_DEPRECATED(5.4, "Obsolete since navigation related objects register to the NavigationObjectRepository world subsytem.")
-	static NAVIGATIONSYSTEM_API TMap<INavLinkCustomInterface*, FWeakObjectPtr> PendingCustomLinkRegistration;
-
 	/** used to apply updates of nav volumes in navigation system's tick */
 	NAVIGATIONSYSTEM_API virtual void PerformNavigationBoundsUpdate(const TArray<FNavigationBoundsUpdateRequest>& UpdateRequests);
 	
@@ -875,9 +869,6 @@ protected:
 	 *	By default these are the player pawn locations. If a player controller has no pawn assigned, the player's camera location will be used instead. */ 
 	NAVIGATIONSYSTEM_API virtual void GetInvokerSeedLocations(const UWorld& InWorld, TArray<FVector, TInlineAllocator<32>>& OutSeedLocations);
 
-	UE_DEPRECATED(5.4, "This function is deprecated. Use GetInvokerSeedLocations using TArray of FVector.")
-	NAVIGATIONSYSTEM_API virtual void GetInvokerSeedLocations(const UWorld& InWorld, TArray<FVector2D, TInlineAllocator<32>>& OutSeedLocations);
-
 	/** @return pointer to ANavigationData instance of given ID, or NULL if it was not found. Note it looks only through registered navigation data */
 	NAVIGATIONSYSTEM_API ANavigationData* GetNavDataWithID(const uint16 NavDataID) const;
 
@@ -1028,9 +1019,7 @@ public:
 
 	/** update element's bounds in navigation octree and mark only specified area as dirty, doesn't re-export geometry */
 	NAVIGATIONSYSTEM_API bool UpdateNavOctreeElementBounds(FNavigationElementHandle Handle, const FBox& NewBounds, TConstArrayView<FBox> DirtyAreas);
-	UE_DEPRECATED(5.4, "Use the overloaded version with object reference and list of dirty areas as parameters instead.")
-	/** update component bounds in navigation octree and mark only specified area as dirty, doesn't re-export component geometry */
-	NAVIGATIONSYSTEM_API bool UpdateNavOctreeElementBounds(UActorComponent* Comp, const FBox& NewBounds, const FBox& DirtyArea);
+
 	UE_DEPRECATED(5.5, "Use the version taking FNavigationElement instead.")
 	NAVIGATIONSYSTEM_API bool UpdateNavOctreeElementBounds(UObject& Object, const FBox& NewBounds, TConstArrayView<FBox> DirtyAreas);
 
