[NavLink Generation] Added a name property to generation configs
Allowed JumpMaxDepth to be negative
Increased debug color saturation when displaying links

diff --git a/Engine/Source/Runtime/NavigationSystem/Public/NavMesh/LinkGenerationConfig.h b/Engine/Source/Runtime/NavigationSystem/Public/NavMesh/LinkGenerationConfig.h
--- a/Engine/Source/Runtime/NavigationSystem/Public/NavMesh/LinkGenerationConfig.h
+++ b/Engine/Source/Runtime/NavigationSystem/Public/NavMesh/LinkGenerationConfig.h
@@ -41,6 +41,10 @@ struct FNavLinkGenerationJumpConfig
 	/** Should this config be used to generate links. */
 	UPROPERTY(Config)		// @todo Hidden for now. Since there is currently only one config, it has no real usage.
 	bool bEnabled = true;
+
+	/** A name for this config */
+	UPROPERTY(EditAnywhere, Config, Category = Settings)
+	FName Name;
 	
 	/** Horizontal length of the jump.
 	 * How far from the starting point we will look for ground. */ 
@@ -51,8 +55,11 @@ struct FNavLinkGenerationJumpConfig
 	UPROPERTY(EditAnywhere, Config, Category = Settings, meta=(UIMin=0, ClampMin=0))
 	float JumpDistanceFromEdge = 10.f; 
 
-	/** How far below the starting height we want to look for landing ground. */
-	UPROPERTY(EditAnywhere, Config, Category = Settings, meta=(UIMin=0, ClampMin=0))
+	/** How far below the starting height we want to look for landing ground.
+	 * A negative value can be used to generate a trajectory landing above the starting height
+	 * (if negative, make sure to use a JumpHeight value big enough).
+	 */
+	UPROPERTY(EditAnywhere, Config, Category = Settings)
 	float JumpMaxDepth = 150.f;
 
 	/** Peak height relative to the height of the starting point. */
diff --git a/Engine/Source/Runtime/Navmesh/Private/DebugUtils/DetourNavLinkDebugDraw.cpp b/Engine/Source/Runtime/Navmesh/Private/DebugUtils/DetourNavLinkDebugDraw.cpp
--- a/Engine/Source/Runtime/Navmesh/Private/DebugUtils/DetourNavLinkDebugDraw.cpp
+++ b/Engine/Source/Runtime/Navmesh/Private/DebugUtils/DetourNavLinkDebugDraw.cpp
@@ -309,9 +309,10 @@ void duDebugDrawNavLinkBuilder(duDebugDraw* dd, const dtNavLinkBuilder& linkBuil
 	if (drawFlags & DRAW_LINKS)
 	{
 		const bool drawFilteredLinks = drawFlags & DRAW_FILTERED_LINKS;
-		unsigned int jumpDownCol0 = duLerpCol(duColor::blue, duColor::white, 200);
+		constexpr int saturation = 50;
+		unsigned int jumpDownCol0 = duLerpCol(duColor::blue, duColor::white, saturation);
 		unsigned int jumpDownCol1 = duColor::blue;
-		unsigned int jumpOverCol0 = duLerpCol(duColor::lightGrey, duColor::white, 200);
+		unsigned int jumpOverCol0 = duLerpCol(duColor::lightGrey, duColor::white, saturation);
 		unsigned int jumpOverCol1 = duColor::lightGrey;
 
 		auto selectColors = [&](dtNavLinkBuilder::JumpLinkFlag flag, const uint8 configIndex, unsigned int& col0, unsigned int& col1)
@@ -330,7 +331,7 @@ void duDebugDrawNavLinkBuilder(duDebugDraw* dd, const dtNavLinkBuilder& linkBuil
 			else
 			{
 				col1 = duIntToCol(configIndex, 255);
-				col0 = duLerpCol(col1, duColor::white, 200);
+				col0 = duLerpCol(col1, duColor::white, saturation);
 			}
 
 			if (flag == dtNavLinkBuilder::FILTERED)
diff --git a/Engine/Source/Runtime/Navmesh/Private/Detour/DetourNavLinkBuilderConfig.cpp b/Engine/Source/Runtime/Navmesh/Private/Detour/DetourNavLinkBuilderConfig.cpp
--- a/Engine/Source/Runtime/Navmesh/Private/Detour/DetourNavLinkBuilderConfig.cpp
+++ b/Engine/Source/Runtime/Navmesh/Private/Detour/DetourNavLinkBuilderConfig.cpp
@@ -6,7 +6,6 @@
 void dtNavLinkBuilderJumpConfig::init()
 {
 	check(jumpHeight >= 0.f);
-	check(jumpMaxDepth >= 0.f);
 	
 	// Parabolic equation y(x) = ax^2 + (-d/l - al)x
 	// Where 'a' is constant
@@ -15,7 +14,7 @@ void dtNavLinkBuilderJumpConfig::init()
 
 	// Solving a for the jumpHeight (h) from the starting point gives:
 	// a(h) = -(1/l^2) * (d + 2h + 2*srt(h(h+d)))
-	cachedParabolaConstant = -(1/dtSqr(jumpLength)) * (jumpMaxDepth + 2*jumpHeight + 2*dtSqrt(jumpHeight*(jumpHeight+jumpMaxDepth)));
+	cachedParabolaConstant = -(1/dtSqr(jumpLength)) * (jumpMaxDepth + 2*jumpHeight + 2*dtSqrt(jumpHeight * dtMax(0.f, jumpHeight+jumpMaxDepth)));
 
 	// Caching this as well as it remains constant for all links
 	cachedDownRatio = -jumpMaxDepth/jumpLength;
