[Pathfind] Change warning to a log when InitPathfinding can't find a polygon at the start or the end location.

diff --git a/Engine/Source/Runtime/NavigationSystem/Private/NavMesh/PImplRecastNavMesh.cpp b/Engine/Source/Runtime/NavigationSystem/Private/NavMesh/PImplRecastNavMesh.cpp
--- a/Engine/Source/Runtime/NavigationSystem/Private/NavMesh/PImplRecastNavMesh.cpp
+++ b/Engine/Source/Runtime/NavigationSystem/Private/NavMesh/PImplRecastNavMesh.cpp
@@ -1350,10 +1350,10 @@ bool FPImplRecastNavMesh::InitPathfinding(const FVector& UnrealStart, const FVec
 	Query.findNearestPoly(&RecastStartToProject.X, Extent, Filter, &StartPoly, &RecastStart.X);
 	if (StartPoly == INVALID_NAVNODEREF)
 	{
-		UE_VLOG_UELOG(NavMeshOwner, LogNavigation, Warning, TEXT("FPImplRecastNavMesh::InitPathfinding start point not on navmesh (%s)"), *UnrealStart.ToString());
-		UE_VLOG_SEGMENT(NavMeshOwner, LogNavigation, Warning, UnrealStart, UnrealEnd, FColor::Red, TEXT("Failed path"));
-		UE_VLOG_LOCATION(NavMeshOwner, LogNavigation, Warning, UnrealStart, 15, FColor::Red, TEXT("Start failed"));
-		UE_VLOG_BOX(NavMeshOwner, LogNavigation, Warning, FBox(UnrealStart - NavExtent, UnrealStart + NavExtent), FColor::Red, TEXT_EMPTY);
+		UE_VLOG_UELOG(NavMeshOwner, LogNavigation, Log, TEXT("FPImplRecastNavMesh::InitPathfinding start point not on navmesh (%s)"), *UnrealStart.ToString());
+		UE_VLOG_SEGMENT(NavMeshOwner, LogNavigation, Log, UnrealStart, UnrealEnd, FColor::Red, TEXT("Failed path"));
+		UE_VLOG_LOCATION(NavMeshOwner, LogNavigation, Log, UnrealStart, 15, FColor::Red, TEXT("Start failed"));
+		UE_VLOG_BOX(NavMeshOwner, LogNavigation, Log, FBox(UnrealStart - NavExtent, UnrealStart + NavExtent), FColor::Red, TEXT_EMPTY);
 
 		if (OutPath)
 		{
@@ -1368,10 +1368,10 @@ bool FPImplRecastNavMesh::InitPathfinding(const FVector& UnrealStart, const FVec
 	{
 		if (Query.isRequiringNavigableEndLocation())
 		{
-			UE_VLOG_UELOG(NavMeshOwner, LogNavigation, Warning, TEXT("FPImplRecastNavMesh::InitPathfinding end point not on navmesh (%s)"), *UnrealEnd.ToString());
-			UE_VLOG_SEGMENT(NavMeshOwner, LogNavigation, Warning, UnrealEnd, UnrealEnd, FColor::Red, TEXT("Failed path"));
-			UE_VLOG_LOCATION(NavMeshOwner, LogNavigation, Warning, UnrealEnd, 15, FColor::Red, TEXT("End failed"));
-			UE_VLOG_BOX(NavMeshOwner, LogNavigation, Warning, FBox(UnrealEnd - NavExtent, UnrealEnd + NavExtent), FColor::Red, TEXT_EMPTY);
+			UE_VLOG_UELOG(NavMeshOwner, LogNavigation, Log, TEXT("FPImplRecastNavMesh::InitPathfinding end point not on navmesh (%s)"), *UnrealEnd.ToString());
+			UE_VLOG_SEGMENT(NavMeshOwner, LogNavigation, Log, UnrealEnd, UnrealEnd, FColor::Red, TEXT("Failed path"));
+			UE_VLOG_LOCATION(NavMeshOwner, LogNavigation, Log, UnrealEnd, 15, FColor::Red, TEXT("End failed"));
+			UE_VLOG_BOX(NavMeshOwner, LogNavigation, Log, FBox(UnrealEnd - NavExtent, UnrealEnd + NavExtent), FColor::Red, TEXT_EMPTY);
 
 			if (OutPath)
 			{
