Add way to visualize nav area of polygons

diff --git a/Engine/Source/Runtime/NavigationSystem/Private/NavMesh/NavMeshRenderingComponent.cpp b/Engine/Source/Runtime/NavigationSystem/Private/NavMesh/NavMeshRenderingComponent.cpp
--- a/Engine/Source/Runtime/NavigationSystem/Private/NavMesh/NavMeshRenderingComponent.cpp
+++ b/Engine/Source/Runtime/NavigationSystem/Private/NavMesh/NavMeshRenderingComponent.cpp
@@ -300,6 +300,7 @@ namespace FNavMeshRenderingHelpers
 			(NavMesh->bDrawPolygonLabels ? (1 << static_cast<int32>(ENavMeshDetailFlags::PolygonLabels)) : 0) |
 			(NavMesh->bDrawDefaultPolygonCost ? (1 << static_cast<int32>(ENavMeshDetailFlags::PolygonCost)) : 0) |
 			(NavMesh->bDrawPolygonFlags ? (1 << static_cast<int32>(ENavMeshDetailFlags::PolygonFlags)) : 0) |
+			(NavMesh->bDrawPolygonAreaIDs ? (1 << static_cast<int32>(ENavMeshDetailFlags::PolygonAreaIDs)) : 0) |
 			(NavMesh->bDrawLabelsOnPathNodes ? (1 << static_cast<int32>(ENavMeshDetailFlags::PathLabels)) : 0) |
 			(NavMesh->bDrawNavLinks ? (1 << static_cast<int32>(ENavMeshDetailFlags::NavLinks)) : 0) |
 			(NavMesh->bDrawFailedNavLinks ? (1 << static_cast<int32>(ENavMeshDetailFlags::FailedNavLinks)) : 0) |
@@ -813,11 +814,12 @@ void FNavMeshSceneProxyData::GatherData(const ARecastNavMesh* NavMesh, int32 InN
 		const bool bGatherPolygonLabels = FNavMeshRenderingHelpers::HasFlag(NavDetailFlags, ENavMeshDetailFlags::PolygonLabels);
 		const bool bGatherPolygonCost = FNavMeshRenderingHelpers::HasFlag(NavDetailFlags, ENavMeshDetailFlags::PolygonCost);
 		const bool bGatherPolygonFlags = FNavMeshRenderingHelpers::HasFlag(NavDetailFlags, ENavMeshDetailFlags::PolygonFlags);
+		const bool bGatherPolyAreaIDs = FNavMeshRenderingHelpers::HasFlag(NavDetailFlags, ENavMeshDetailFlags::PolygonAreaIDs);
 
-		if (bGatherTileLabels || bGatherTileBounds || bGatherTileResolutions || bGatherPolygonLabels || bGatherPolygonCost || bGatherPolygonFlags || bGatherTileBuildTimes)
+		if (bGatherTileLabels || bGatherTileBounds || bGatherTileResolutions || bGatherPolygonLabels || bGatherPolygonCost || bGatherPolygonFlags || bGatherPolyAreaIDs || bGatherTileBuildTimes)
 		{
 			QUICK_SCOPE_CYCLE_COUNTER(STAT_NavMesh_GatherDebugDrawing_TileIterations);
-			
+
 			TArray<FNavTileRef> UseTileRefs;
 			if (TileSet.Num() > 0)
 			{
@@ -829,6 +831,7 @@ void FNavMeshSceneProxyData::GatherData(const ARecastNavMesh* NavMesh, int32 InN
 			}
 
 			TMap<FIntPoint, FVector> TileBuildTimeLabelLocations;
+			TSet<int32> TilesAreaIDs;
 			
 			// calculate appropriate points for displaying debug labels
 			DebugLabels.Reserve(UseTileRefs.Num());
@@ -868,7 +871,7 @@ void FNavMeshSceneProxyData::GatherData(const ARecastNavMesh* NavMesh, int32 InN
 						}
 					}	
 					
-					if (bGatherPolygonLabels || bGatherPolygonCost || bGatherPolygonFlags)
+					if (bGatherPolygonLabels || bGatherPolygonCost || bGatherPolygonFlags || bGatherPolyAreaIDs)
 					{
 						TArray<FNavPoly> Polys;
 						NavMesh->GetPolysInTile(TileRef, Polys);
@@ -922,6 +925,18 @@ void FNavMeshSceneProxyData::GatherData(const ARecastNavMesh* NavMesh, int32 InN
 								StringBuilder.Appendf(TEXT("Flags Poly/Area: [0x%X : 0x%X]"), PolyFlags, AreaFlags);
 							}
 
+							if (bGatherPolyAreaIDs)
+							{
+								const uint32 PolyAreaID = NavMesh->GetPolyAreaID(Poly.Ref);
+								TilesAreaIDs.Add(PolyAreaID);
+
+								if (StringBuilder.Len() > 0)
+								{
+									StringBuilder.Append(TEXT("\n"));
+								}
+								StringBuilder.Appendf(TEXT("AreaID: %d"), PolyAreaID);
+							}
+
 							DebugLabels.Add(FDebugText(Poly.Center + NavMeshDrawOffset, StringBuilder.ToString()));
 						}
 					}
@@ -1014,6 +1029,21 @@ void FNavMeshSceneProxyData::GatherData(const ARecastNavMesh* NavMesh, int32 InN
 				}
 			}
 #endif // RECAST_INTERNAL_DEBUG_DATA
+
+			if (bGatherPolyAreaIDs && !TilesAreaIDs.IsEmpty())
+			{
+				auto SortHandle = [](const int32 A, const int32 B)
+				{
+					return A < B;
+				};
+				TilesAreaIDs.StableSort(SortHandle);
+				DebugLabels.Add(FDebugText(TEXT("NavAreaClasses:")));
+				for (const int32 AreaID : TilesAreaIDs)
+				{
+					const UClass* AreaClass = NavMesh->GetAreaClass(AreaID);
+					DebugLabels.Add(FDebugText(FString::Printf(TEXT("  %d: %s"), AreaID, *GetNameSafe(AreaClass))));
+				}
+			}
 		}
 
 #if RECAST_INTERNAL_DEBUG_DATA
diff --git a/Engine/Source/Runtime/NavigationSystem/Public/NavMesh/NavMeshRenderingComponent.h b/Engine/Source/Runtime/NavigationSystem/Public/NavMesh/NavMeshRenderingComponent.h
--- a/Engine/Source/Runtime/NavigationSystem/Public/NavMesh/NavMeshRenderingComponent.h
+++ b/Engine/Source/Runtime/NavigationSystem/Public/NavMesh/NavMeshRenderingComponent.h
@@ -36,6 +36,7 @@ enum class ENavMeshDetailFlags : uint8
 	PolygonLabels,
 	PolygonCost,
 	PolygonFlags,
+	PolygonAreaIDs,
 	PathLabels,
 	NavLinks,
 	FailedNavLinks,
diff --git a/Engine/Source/Runtime/NavigationSystem/Public/NavMesh/RecastNavMesh.h b/Engine/Source/Runtime/NavigationSystem/Public/NavMesh/RecastNavMesh.h
--- a/Engine/Source/Runtime/NavigationSystem/Public/NavMesh/RecastNavMesh.h
+++ b/Engine/Source/Runtime/NavigationSystem/Public/NavMesh/RecastNavMesh.h
@@ -631,6 +631,10 @@ class ARecastNavMesh : public ANavigationData
 	UPROPERTY(EditAnywhere, Category=Display)
 	uint32 bDrawPolygonFlags:1;
 
+	/** Draw a label for every poly that indicates its area id and the list of all NavAreaClass used in the displayed tiles. */
+	UPROPERTY(EditAnywhere, Category=Display)
+	uint32 bDrawPolygonAreaIDs:1;
+
 	UPROPERTY(EditAnywhere, Category=Display)
 	uint32 bDrawLabelsOnPathNodes:1;
 
