Add way to debug the nav modifiers added for specific tile

diff --git a/Engine/Source/Runtime/Engine/Private/AI/Navigation/NavigationModifier.cpp b/Engine/Source/Runtime/Engine/Private/AI/Navigation/NavigationModifier.cpp
--- a/Engine/Source/Runtime/Engine/Private/AI/Navigation/NavigationModifier.cpp
+++ b/Engine/Source/Runtime/Engine/Private/AI/Navigation/NavigationModifier.cpp
@@ -293,6 +293,35 @@ bool UNavLinkDefinition::HasAdjustableLinks() const
 // FAreaNavModifier
 //----------------------------------------------------------------------//
 
+FString LexToString(ENavigationShapeType::Type Type)
+{
+	switch (Type)
+	{
+	case ENavigationShapeType::Unknown: return TEXT("Unknown");
+	case ENavigationShapeType::Cylinder: return TEXT("Cylinder");
+	case ENavigationShapeType::Box: return TEXT("Box");
+	case ENavigationShapeType::Convex: return TEXT("Convex");
+	case ENavigationShapeType::InstancedConvex: return TEXT("InstancedConvex");
+	default:
+		checkNoEntry();
+		return FString();
+	}
+}
+
+FString LexToString(ENavigationAreaMode::Type Type)
+{
+	switch (Type)
+	{
+	case ENavigationAreaMode::Apply: return TEXT("Apply");
+	case ENavigationAreaMode::Replace: return TEXT("Replace");
+	case ENavigationAreaMode::ApplyInLowPass: return TEXT("ApplyInLowPass");
+	case ENavigationAreaMode::ReplaceInLowPass: return TEXT("ReplaceInLowPass");
+	default:
+		checkNoEntry();
+		return FString();
+	}
+}
+
 FAreaNavModifier::FAreaNavModifier() 
 : Cost(0.0f)
 , FixedCost(0.0f)
diff --git a/Engine/Source/Runtime/Engine/Public/AI/NavigationModifier.h b/Engine/Source/Runtime/Engine/Public/AI/NavigationModifier.h
--- a/Engine/Source/Runtime/Engine/Public/AI/NavigationModifier.h
+++ b/Engine/Source/Runtime/Engine/Public/AI/NavigationModifier.h
@@ -39,6 +39,8 @@ namespace ENavigationShapeType
 	};
 }
 
+ENGINE_API FString LexToString(ENavigationShapeType::Type Type);
+
 namespace ENavigationAreaMode
 {
 	enum Type
@@ -59,6 +61,8 @@ namespace ENavigationAreaMode
 	};
 }
 
+ENGINE_API FString LexToString(ENavigationAreaMode::Type Type);
+
 namespace ENavigationCoordSystem
 {
 	enum Type
diff --git a/Engine/Source/Runtime/NavigationSystem/Private/NavMesh/RecastNavMeshGenerator.cpp b/Engine/Source/Runtime/NavigationSystem/Private/NavMesh/RecastNavMeshGenerator.cpp
--- a/Engine/Source/Runtime/NavigationSystem/Private/NavMesh/RecastNavMeshGenerator.cpp
+++ b/Engine/Source/Runtime/NavigationSystem/Private/NavMesh/RecastNavMeshGenerator.cpp
@@ -2512,11 +2512,36 @@ void FRecastTileGenerator::AppendModifier(const FCompositeNavModifier& Modifier,
 			return;
 		}
 	}
-		
+
 	ModifierElement.Areas = Modifier.GetAreas();
 	ModifierElement.bMaskFillCollisionUnderneathForNavmesh = Modifier.GetMaskFillCollisionUnderneathForNavmesh();
 	ModifierElement.NavMeshResolution = Modifier.GetNavMeshResolution();
 
+#if RECAST_INTERNAL_DEBUG_DATA
+	if (IsTileDebugActive())
+	{
+		FStringBuilderBase StrBuilder;
+		for (const FTransform& Transform : ModifierElement.PerInstanceTransform)
+		{
+			StrBuilder += FString::Printf(TEXT(" %s"), *Transform.ToString());
+		}
+		UE_LOG(LogNavigation, Log, TEXT("    Modifier transforms [%d]:%s"), ModifierElement.PerInstanceTransform.Num(), StrBuilder.ToString());
+
+		for (const FAreaNavModifier& AreaNavModifier : ModifierElement.Areas)
+		{
+			UE_LOG(LogNavigation, Log, TEXT("    Adding area modifier:\n"
+										"                           Type: %s\n"
+										"                           ApplyMode: %s\n"
+										"                           AreaClass: %s\n"
+										"                           Bounds: %s"),
+			*LexToString(AreaNavModifier.GetShapeType()),
+			*LexToString(AreaNavModifier.GetApplyMode()),
+			*AreaNavModifier.GetAreaClass()->GetName(),
+			*AreaNavModifier.GetBounds().ToString());
+		}
+	}
+#endif
+
 	Modifiers.Add(MoveTemp(ModifierElement));
 }
 
