Fix IsNavigationRelevant if a TestActor is supplied

#rnx

diff --git a/Engine/Source/Runtime/NavigationSystem/Private/NavigationSystem.cpp b/Engine/Source/Runtime/NavigationSystem/Private/NavigationSystem.cpp
--- a/Engine/Source/Runtime/NavigationSystem/Private/NavigationSystem.cpp
+++ b/Engine/Source/Runtime/NavigationSystem/Private/NavigationSystem.cpp
@@ -2561,6 +2561,7 @@ bool UNavigationSystemV1::IsNavigationRelevant(const AActor* TestActor) const
 	if (TestActor)
 	{
 		TInlineComponentArray<UActorComponent*> Components;
+		TestActor->GetComponents(Components);
 		for (int32 Idx = 0; Idx < Components.Num(); Idx++)
 		{
 			NavInterface = Cast<const INavRelevantInterface>(Components[Idx]);
