Fix editor crash when a large number is typed in the Agent Radius value.

diff --git a/Engine/Source/Runtime/NavigationSystem/Public/NavMesh/RecastNavMesh.h b/Engine/Source/Runtime/NavigationSystem/Public/NavMesh/RecastNavMesh.h
--- a/Engine/Source/Runtime/NavigationSystem/Public/NavMesh/RecastNavMesh.h
+++ b/Engine/Source/Runtime/NavigationSystem/Public/NavMesh/RecastNavMesh.h
@@ -707,17 +707,17 @@ class ARecastNavMesh : public ANavigationData
 	float AgentMaxStepHeight;
 
 	/** Resolution params 
-	 * If using multiple resolutions, it's recommended to chose the highest resolution first and 
+	 * If using multiple resolutions, it's recommended to choose the highest resolution first and 
 	 * set it according to the highest desired precision and then the other resolutions. */
 	UPROPERTY(EditAnywhere, Category = Generation, config)
 	FNavMeshResolutionParam NavMeshResolutionParams[(uint8)ENavigationDataResolution::MAX];
 
 	/** Radius of smallest agent to traverse this navmesh */
-	UPROPERTY(EditAnywhere, Category = Generation, config, meta = (ClampMin = "0.0"))
+	UPROPERTY(EditAnywhere, Category = Generation, config, meta = (ClampMin = "0.0", ClampMax = "100000.0", UIMin = "0.0", UIMax = "100000.0"))
 	float AgentRadius;
 
 	/** Size of the tallest agent that will path with this navmesh. */
-	UPROPERTY(EditAnywhere, Category = Generation, config, meta = (ClampMin = "0.0"))
+	UPROPERTY(EditAnywhere, Category = Generation, config, meta = (ClampMin = "0.0", ClampMax = "100000.0", UIMin = "0.0", UIMax = "100000.0"))
 	float AgentHeight;
 
 	/* The maximum slope (angle) that the agent can move on. */ 
