[Navlink] Fix automatically generated navlinks being created without the NavLinkFlag.

diff --git a/Engine/Source/Runtime/NavigationSystem/Private/NavMesh/RecastNavMeshGenerator.cpp b/Engine/Source/Runtime/NavigationSystem/Private/NavMesh/RecastNavMeshGenerator.cpp
--- a/Engine/Source/Runtime/NavigationSystem/Private/NavMesh/RecastNavMeshGenerator.cpp
+++ b/Engine/Source/Runtime/NavigationSystem/Private/NavMesh/RecastNavMeshGenerator.cpp
@@ -5534,7 +5534,10 @@ void FRecastNavMeshGenerator::ResolveGeneratedLinkAreas(FRecastBuildConfig& OutC
 				if (AreaIDPtr != nullptr)
 				{
 					OutAreaID = IntCastChecked<unsigned char>(*AreaIDPtr);
-					OutPolyFlags = AdditionalCachedData.FlagsPerArea[*AreaIDPtr];
+
+					// Using FlagsPerOffMeshLinkArea rather than FlagsPerArea because this function, ResolveGeneratedLinkAreas, 
+					// is concerned only with navlinks (a.k.a. off-mesh links).
+					OutPolyFlags = AdditionalCachedData.FlagsPerOffMeshLinkArea[*AreaIDPtr];
 				}
 				else
 				{
