Remove old comment (there is no lock anymore it was removed in CL 1790185).

#rnx

diff --git a/Engine/Source/Runtime/NavigationSystem/Private/NavMesh/RecastNavMesh.cpp b/Engine/Source/Runtime/NavigationSystem/Private/NavMesh/RecastNavMesh.cpp
--- a/Engine/Source/Runtime/NavigationSystem/Private/NavMesh/RecastNavMesh.cpp
+++ b/Engine/Source/Runtime/NavigationSystem/Private/NavMesh/RecastNavMesh.cpp
@@ -1312,7 +1312,6 @@ void ARecastNavMesh::FillConfig(FNavDataConfig& Dest)
 void ARecastNavMesh::BeginBatchQuery() const
 {
 #if RECAST_ASYNC_REBUILDING
-	// lock critical section when no other batch queries are active
 	if (BatchQueryCounter <= 0)
 	{
 		BatchQueryCounter = 0;
