Fix assertion occuring in UE::NavMesh::Private::CalculateMaxTilesCount()

diff --git a/Engine/Source/Runtime/NavigationSystem/Private/NavMesh/RecastNavMeshGenerator.cpp b/Engine/Source/Runtime/NavigationSystem/Private/NavMesh/RecastNavMeshGenerator.cpp
--- a/Engine/Source/Runtime/NavigationSystem/Private/NavMesh/RecastNavMeshGenerator.cpp
+++ b/Engine/Source/Runtime/NavigationSystem/Private/NavMesh/RecastNavMeshGenerator.cpp
@@ -5063,8 +5063,32 @@ int32 CalculateMaxTilesCount(const TNavStatArray<FBox>& NavigableAreas, FVector:
 			GridCellsCount += (XSize*YSize);
 		}
 	}
+
+	const int64 MaxTilesCount = FMath::CeilToInt(static_cast<FVector::FReal>(GridCellsCount) * AvgLayersPerGridCell);
+
+	// Output warning with details when overflowing
+	UE_SUPPRESS(LogNavigationDataBuild, Warning,
+	{
+		if (!IntFitsIn<int32>(MaxTilesCount))
+		{
+			TStringBuilder<2048> DebugInformation;
+			
+			DebugInformation.Appendf(TEXT("%hs Computed maximum tile count is over limit. "
+				"Check inputs for a abnormally high values.\n"
+				"  MaxTilesCount: %llu, GridCellsCount: %llu, AvgLayersPerGridCell: %f, TileSizeInWorldUnits: %f \n"),
+				__FUNCTION__, MaxTilesCount, GridCellsCount, AvgLayersPerGridCell, TileSizeInWorldUnits);
+
+			DebugInformation.Appendf(TEXT("  Navigable areas:\n"));
+			for (const FBox& AreaBounds : NavigableAreas)
+			{
+				DebugInformation.Appendf(TEXT("    AreaBounds: %s"), *AreaBounds.ToString());
+			}
+			
+			UE_LOG(LogNavigationDataBuild, Warning, TEXT("%s"), DebugInformation.ToString());
+		}
+	});
 	
-	return IntCastChecked<int32>(FMath::CeilToInt(static_cast<FVector::FReal>(GridCellsCount) * AvgLayersPerGridCell));
+	return IntCastChecked<int32>(FMath::Clamp(MaxTilesCount, 0, TNumericLimits<int32>::Max()));
 }
 } // UE::NavMesh::Private
 
