[Navigation] Fix crash occuring when building navmesh on PIE ending.

diff --git a/Engine/Source/Runtime/NavigationSystem/Private/NavigationSystem.cpp b/Engine/Source/Runtime/NavigationSystem/Private/NavigationSystem.cpp
--- a/Engine/Source/Runtime/NavigationSystem/Private/NavigationSystem.cpp
+++ b/Engine/Source/Runtime/NavigationSystem/Private/NavigationSystem.cpp
@@ -1352,7 +1352,7 @@ void UNavigationSystemV1::OnWorldInitDone(FNavigationSystemRunMode Mode)
 
 		if (FNavigationSystem::IsEditorRunMode(OperationMode))
 		{
-			RemoveNavigationBuildLock(InitialNavBuildingLockFlags, ELockRemovalRebuildAction::RebuildIfNotInEditor);
+			RemoveNavigationBuildLock(InitialNavBuildingLockFlags, ELockRemovalRebuildAction::NoRebuild);
 		}
 	}
 	else
@@ -1414,7 +1414,7 @@ void UNavigationSystemV1::OnWorldInitDone(FNavigationSystemRunMode Mode)
 		if (FNavigationSystem::IsEditorRunMode(OperationMode))
 		{
 			// don't lock navigation building in editor
-			RemoveNavigationBuildLock(InitialNavBuildingLockFlags, ELockRemovalRebuildAction::RebuildIfNotInEditor);
+			RemoveNavigationBuildLock(InitialNavBuildingLockFlags, ELockRemovalRebuildAction::NoRebuild);
 		}
 
 		// See if any of registered navigation data needs NavOctree
@@ -4789,8 +4789,9 @@ void UNavigationSystemV1::OnPIEEnd()
 	if (MyWorld && !MyWorld->IsGameWorld())
 	{
 		bAsyncBuildPaused = false;
-		// there's no need to request while navigation rebuilding just because PIE has ended
-		RemoveNavigationBuildLock(ENavigationBuildLock::NoUpdateInPIE, ELockRemovalRebuildAction::RebuildIfNotInEditor);
+
+		// There's no need to request navigation rebuilding when PIE ended.
+		RemoveNavigationBuildLock(ENavigationBuildLock::NoUpdateInPIE, ELockRemovalRebuildAction::NoRebuild);
 	}
 }
 
