Fix crash in NavigationSystem, access null component
#rnx

diff --git a/Engine/Source/Runtime/NavigationSystem/Private/NavigationSystem.cpp b/Engine/Source/Runtime/NavigationSystem/Private/NavigationSystem.cpp
--- a/Engine/Source/Runtime/NavigationSystem/Private/NavigationSystem.cpp
+++ b/Engine/Source/Runtime/NavigationSystem/Private/NavigationSystem.cpp
@@ -4107,7 +4107,7 @@ void UNavigationSystemV1::OnActorRegistered(AActor* Actor)
 		// Note that we do so using the delegate since it is possible for derived systems to override them.
 		for (UActorComponent* Component : Actor->GetComponents())
 		{
-			if (Component->IsRegistered())
+			if (Component && Component->IsRegistered())
 			{
 				FNavigationSystem::OnComponentRegistered(*Component);
 			}
