[Navmesh] Navmesh building optimization when using FCompositeNavModifier with lots of data

diff --git a/Engine/Source/Runtime/NavigationSystem/Private/NavMesh/RecastNavMeshGenerator.cpp b/Engine/Source/Runtime/NavigationSystem/Private/NavMesh/RecastNavMeshGenerator.cpp
--- a/Engine/Source/Runtime/NavigationSystem/Private/NavMesh/RecastNavMeshGenerator.cpp
+++ b/Engine/Source/Runtime/NavigationSystem/Private/NavMesh/RecastNavMeshGenerator.cpp
@@ -2178,11 +2178,20 @@ void FRecastTileGenerator::GatherNavigationDataGeometry(const TSharedRef<FNaviga
 		bDumpGeometryData = true;
 	}
 
-	// Temporary change to help narrow down a rare crash:
-	// Was: const FCompositeNavModifier ModifierInstance = ElementData->GetModifierForAgent(&OwnerNavDataConfig);
-	const FCompositeNavModifier& ModifierInstance = ElementData.Modifiers.HasMetaAreas()
-		? ElementData.Modifiers.GetInstantiatedMetaModifier(&OwnerNavDataConfig, ElementData.SourceElement->GetWeakUObject())
-		: ElementData.Modifiers;
+	// Not using the ternary operator to avoid a lvalue-to-rvalue conversion which causes copies of FCompositeNavModifier even when no meta area is used.
+	// This is to avoid performance hit when using FCompositeNavModifier with a lot of internal content.
+	TOptional<FCompositeNavModifier> Storage;
+	const FCompositeNavModifier* ModifierInstancePtr;
+	if (ElementData.Modifiers.HasMetaAreas())
+	{
+		Storage = ElementData.Modifiers.GetInstantiatedMetaModifier(&OwnerNavDataConfig, ElementData.SourceElement->GetWeakUObject());
+		ModifierInstancePtr = &Storage.GetValue();
+	}
+	else
+	{
+		ModifierInstancePtr = &ElementData.Modifiers;
+	}
+	const FCompositeNavModifier& ModifierInstance = *ModifierInstancePtr;
 
 	const bool bExportGeometry = bGeometryChanged && ElementData.HasGeometry();
 	if (bExportGeometry)
