Addition of navigation system logs (on: AddNavigationSystemToWorld, OnWorldInitDone, GetWorldBounds, PerformNavigationBoundsUpdate, OnLevelAddedToWorld and ANavigationData bound change).

diff --git a/Engine/Source/Runtime/Engine/Private/AI/NavigationSystemBase.cpp b/Engine/Source/Runtime/Engine/Private/AI/NavigationSystemBase.cpp
--- a/Engine/Source/Runtime/Engine/Private/AI/NavigationSystemBase.cpp
+++ b/Engine/Source/Runtime/Engine/Private/AI/NavigationSystemBase.cpp
@@ -78,7 +78,9 @@ namespace FNavigationSystem
 	void AddNavigationSystemToWorld(UWorld& WorldOwner, const FNavigationSystemRunMode RunMode, UNavigationSystemConfig* NavigationSystemConfig, const bool bInitializeForWorld, const bool bOverridePreviousNavSys)
 	{
 		TRACE_CPUPROFILER_EVENT_SCOPE(FNavigationSystem::AddNavigationSystemToWorld);
-		UE_LOG(LogNavigation, Verbose, TEXT("%hs bOverridePreviousNavSys=%i (WorldOwner: %s)"), __FUNCTION__, bOverridePreviousNavSys, *WorldOwner.GetOuter()->GetName());
+
+		UE_LOG(LogNavigation, Verbose, TEXT("%hs bOverridePreviousNavSys=%i (NavigationSystemConfig: %s) (WorldOwner: %s)"),
+			__FUNCTION__, bOverridePreviousNavSys, *GetNameSafe(NavigationSystemConfig), *WorldOwner.GetOuter()->GetName());
 
 		const FNavigationSystemRunMode ResolvedRunMode = (RunMode == FNavigationSystemRunMode::InferFromWorldMode) ? FindRunModeFromWorldType(WorldOwner) : RunMode;
 
diff --git a/Engine/Source/Runtime/NavigationSystem/Private/NavigationData.cpp b/Engine/Source/Runtime/NavigationSystem/Private/NavigationData.cpp
--- a/Engine/Source/Runtime/NavigationSystem/Private/NavigationData.cpp
+++ b/Engine/Source/Runtime/NavigationSystem/Private/NavigationData.cpp
@@ -628,6 +628,8 @@ void ANavigationData::CancelBuild()
 
 void ANavigationData::OnNavigationBoundsChanged()
 {
+	UE_LOG(LogNavigation, Log, TEXT("%hs %s"), __FUNCTION__, *GetNameSafe(this));
+	
 	// Create generator if it wasn't yet
 	if (NavDataGenerator.Get() == nullptr)
 	{
diff --git a/Engine/Source/Runtime/NavigationSystem/Private/NavigationSystem.cpp b/Engine/Source/Runtime/NavigationSystem/Private/NavigationSystem.cpp
--- a/Engine/Source/Runtime/NavigationSystem/Private/NavigationSystem.cpp
+++ b/Engine/Source/Runtime/NavigationSystem/Private/NavigationSystem.cpp
@@ -1320,6 +1320,8 @@ void UNavigationSystemV1::OnBeginTearingDown(UWorld* World)
 
 void UNavigationSystemV1::OnWorldInitDone(FNavigationSystemRunMode Mode)
 {
+	UE_LOG(LogNavigation, Verbose, TEXT("%hs, Mode: %s"), __FUNCTION__, *UEnum::GetValueAsString(Mode));
+	
 	UNavigationSystemBase::OnNavigationInitStartStaticDelegate().Broadcast(*this);
 
 	OperationMode = Mode;
@@ -1800,6 +1802,8 @@ void UNavigationSystemV1::AddReferencedObjects(UObject* InThis, FReferenceCollec
 #if WITH_EDITOR
 void UNavigationSystemV1::SetNavigationAutoUpdateEnabled(bool bNewEnable, UNavigationSystemBase* InNavigationSystemBase)
 {
+	UE_LOG(LogNavigation, Verbose, TEXT("%hs bNewEnable: %i"), __FUNCTION__, bNewEnable);
+	
 	if (bNewEnable != bNavigationAutoUpdateEnabled)
 	{
 		bNavigationAutoUpdateEnabled = bNewEnable;
@@ -2597,12 +2601,15 @@ FBox UNavigationSystemV1::GetWorldBounds() const
 			{
 				if (IsNavigationRelevant(*It))
 				{
+					UE_LOG(LogNavigation, Verbose, TEXT("%hs %s is NavRelevant with bounds %s (%s)."),
+						__FUNCTION__, *It->GetName(), *(*It)->GetComponentsBoundingBox().ToCompactString(), *It->GetFullName());
 					NavigableWorldBounds += (*It)->GetComponentsBoundingBox();
 				}
 			}
 		}
 	}
 
+	UE_LOG(LogNavigation, Log, TEXT("%hs NavigableWorldBounds: %s."), __FUNCTION__, *NavigableWorldBounds.ToCompactString());
 	return NavigableWorldBounds;
 }
 
@@ -4357,6 +4364,7 @@ void UNavigationSystemV1::PerformNavigationBoundsUpdate(const TArray<FNavigation
 			{
 				if (ExistingElementId.IsValidId())
 				{
+					UE_LOG(LogNavigation, Log, TEXT("%hs Remove bounds %s"), __FUNCTION__, *RegisteredNavBounds[ExistingElementId].AreaBox.ToCompactString());
 					UpdatedAreas.Add(RegisteredNavBounds[ExistingElementId].AreaBox);
 					RegisteredNavBounds.Remove(ExistingElementId);
 				}
@@ -4368,6 +4376,7 @@ void UNavigationSystemV1::PerformNavigationBoundsUpdate(const TArray<FNavigation
 			{
 				if (ExistingElementId.IsValidId())
 				{
+					UE_LOG(LogNavigation, Log, TEXT("%hs Update bounds %s"), __FUNCTION__, *RegisteredNavBounds[ExistingElementId].AreaBox.ToCompactString());
 					const FBox ExistingBox = RegisteredNavBounds[ExistingElementId].AreaBox;
 					const bool bSameArea = (Request.NavBounds.AreaBox == ExistingBox);
 					if (!bSameArea)
@@ -4380,6 +4389,7 @@ void UNavigationSystemV1::PerformNavigationBoundsUpdate(const TArray<FNavigation
 				}
 				else
 				{
+					UE_LOG(LogNavigation, Log, TEXT("%hs Add bounds %s"), __FUNCTION__, *Request.NavBounds.AreaBox.ToCompactString());
 					AddNavigationBounds(Request.NavBounds);
 				}
 
@@ -5001,6 +5011,7 @@ void UNavigationSystemV1::OnLevelAddedToWorld(ULevel* InLevel, UWorld* InWorld)
 	}
 	else if (NavDataRegistrationQueue.Num() > 0)
 	{
+		UE_LOG(LogNavigation, Log, TEXT("%hs Calling ProcessRegistrationCandidates (level: %s, world: %s)."), __FUNCTION__, *GetFullNameSafe(InLevel), *GetFullNameSafe(InWorld));
 		ProcessRegistrationCandidates();
 	}
 }
