NavigationSystem: fix FName in header

diff --git a/Engine/Source/Runtime/NavigationSystem/Private/NavigationSystem.cpp b/Engine/Source/Runtime/NavigationSystem/Private/NavigationSystem.cpp
--- a/Engine/Source/Runtime/NavigationSystem/Private/NavigationSystem.cpp
+++ b/Engine/Source/Runtime/NavigationSystem/Private/NavigationSystem.cpp
@@ -174,6 +174,13 @@ CSV_DEFINE_CATEGORY(NavTasksDelays, true);
 CSV_DEFINE_CATEGORY(NavTasks, true);
 CSV_DEFINE_CATEGORY(NavInvokers, true);
 
+namespace FNavigationSystem
+{	
+#if ALLOW_TIME_SLICE_DEBUG
+	const FName DebugTimeSliceDefaultSectionName(TEXT("DefaultSection"));
+#endif // ALLOW_TIME_SLICE_DEBUG
+}
+
 namespace UE::Navigation::Private
 {
 
diff --git a/Engine/Source/Runtime/NavigationSystem/Public/NavigationSystem.h b/Engine/Source/Runtime/NavigationSystem/Public/NavigationSystem.h
--- a/Engine/Source/Runtime/NavigationSystem/Public/NavigationSystem.h
+++ b/Engine/Source/Runtime/NavigationSystem/Public/NavigationSystem.h
@@ -92,7 +92,7 @@ namespace FNavigationSystem
 	void NAVIGATIONSYSTEM_API MakeAllComponentsNeverAffectNav(AActor& Actor);
 
 #if ALLOW_TIME_SLICE_DEBUG
-	const FName DebugTimeSliceDefaultSectionName = FName(TEXT("DefaultSection"));
+	extern NAVIGATIONSYSTEM_API const FName DebugTimeSliceDefaultSectionName;
 #endif // ALLOW_TIME_SLICE_DEBUG
 }
 
