Fix UNavAreaMeta_SwitchByAgent UClass being not deterministic during cook.

The function UpdateAgentConfig is called during the post load of objects of type UNavAreaMeta_SwitchByAgent. This function changes flags in the UClass. The UClass is then serialized in blueprint dependencies. During a second cook, when evaluating dependencies, no package is loaded so UpdateAgentConfig doesn't run and the UClass flags are different from the previous cook. This triggers useless recooks.

UpdateAgentConfig uses the CDO of UNavigationSystemV1 to change the UClass of UNavAreaMeta_SwitchByAgent. It means UpdateAgentConfig has to run after the initialization of both classes. To do so I use the callback FCoreDelegates::OnAllModuleLoadingPhasesComplete. This way UpdateAgentConfig is always called and the UClass of UNavAreaMeta_SwitchByAgent is deterministic.

#rnx

diff --git a/Engine/Source/Runtime/NavigationSystem/Private/NavAreas/NavAreaMeta_SwitchByAgent.cpp b/Engine/Source/Runtime/NavigationSystem/Private/NavAreas/NavAreaMeta_SwitchByAgent.cpp
--- a/Engine/Source/Runtime/NavigationSystem/Private/NavAreas/NavAreaMeta_SwitchByAgent.cpp
+++ b/Engine/Source/Runtime/NavigationSystem/Private/NavAreas/NavAreaMeta_SwitchByAgent.cpp
@@ -7,6 +7,10 @@
 
 #include UE_INLINE_GENERATED_CPP_BY_NAME(NavAreaMeta_SwitchByAgent)
 
+#if WITH_EDITOR
+bool UNavAreaMeta_SwitchByAgent::bUpdateAgentConfigCallbackRegistered = false;
+#endif //#if WITH_EDITOR
+
 UNavAreaMeta_SwitchByAgent::UNavAreaMeta_SwitchByAgent(const FObjectInitializer& ObjectInitializer) : Super(ObjectInitializer)
 {
 	Agent0Area = UNavArea_Default::StaticClass();
@@ -25,6 +29,14 @@ UNavAreaMeta_SwitchByAgent::UNavAreaMeta_SwitchByAgent(const FObjectInitializer&
 	Agent13Area = UNavArea_Default::StaticClass();
 	Agent14Area = UNavArea_Default::StaticClass();
 	Agent15Area = UNavArea_Default::StaticClass();
+
+#if WITH_EDITOR
+	if (HasAnyFlags(RF_ClassDefaultObject) && !bUpdateAgentConfigCallbackRegistered)
+	{
+		FCoreDelegates::OnAllModuleLoadingPhasesComplete.AddLambda([]() { GetMutableDefault<UNavAreaMeta_SwitchByAgent>()->UpdateAgentConfig(); });
+		bUpdateAgentConfigCallbackRegistered = true;
+	}
+#endif //#if WITH_EDITOR
 }
 
 TSubclassOf<UNavAreaBase> UNavAreaMeta_SwitchByAgent::PickAreaClassForAgent(const AActor& Actor, const FNavAgentProperties& NavAgent) const
@@ -60,7 +72,7 @@ TSubclassOf<UNavAreaBase> UNavAreaMeta_SwitchByAgent::PickAreaClassForAgent(cons
 #if WITH_EDITOR
 void UNavAreaMeta_SwitchByAgent::UpdateAgentConfig()
 {
-	const UNavigationSystemV1* DefNavSys = (UNavigationSystemV1*)(UNavigationSystemV1::StaticClass()->GetDefaultObject());
+	const UNavigationSystemV1* DefNavSys = GetDefault<UNavigationSystemV1>();
 	check(DefNavSys);
 
 	const int32 MaxAllowedAgents = 16;
diff --git a/Engine/Source/Runtime/NavigationSystem/Public/NavAreas/NavAreaMeta_SwitchByAgent.h b/Engine/Source/Runtime/NavigationSystem/Public/NavAreas/NavAreaMeta_SwitchByAgent.h
--- a/Engine/Source/Runtime/NavigationSystem/Public/NavAreas/NavAreaMeta_SwitchByAgent.h
+++ b/Engine/Source/Runtime/NavigationSystem/Public/NavAreas/NavAreaMeta_SwitchByAgent.h
@@ -71,5 +71,7 @@ class UNavAreaMeta_SwitchByAgent : public UNavAreaMeta
 #if WITH_EDITOR
 	/** setup AgentXArea properties */
 	NAVIGATIONSYSTEM_API virtual void UpdateAgentConfig();
+
+	static bool bUpdateAgentConfigCallbackRegistered;
 #endif
 };
