Fix ensure occuring when undoing a change for a NavLinkJumpConfig with a LinkProxyClass.

diff --git a/Engine/Source/Runtime/NavigationSystem/Public/NavMesh/LinkGenerationConfig.h b/Engine/Source/Runtime/NavigationSystem/Public/NavMesh/LinkGenerationConfig.h
--- a/Engine/Source/Runtime/NavigationSystem/Public/NavMesh/LinkGenerationConfig.h
+++ b/Engine/Source/Runtime/NavigationSystem/Public/NavMesh/LinkGenerationConfig.h
@@ -125,12 +125,12 @@ struct FNavLinkGenerationJumpConfig
 	/** Current proxy. The proxy instance is build from the LinkProxyClass (provided it's not null).
 	 * A proxy will be created if a @see LinkProxyClass is used.
 	 */
-	UPROPERTY(Transient)
+	UPROPERTY(SkipSerialization)
 	TObjectPtr<UBaseGeneratedNavLinksProxy> LinkProxy = nullptr;
 
 	/** Is the link proxy registered to the navigation system CustomNavLinksMap.
 	 * Registration occurs on PostRegisterAllComponents or on PostLoadPreRebuild if a new proxy was created. */
-	UPROPERTY(Transient)
+	UPROPERTY(SkipSerialization)
 	bool bLinkProxyRegistered = false;
 
 	/** Implemented for deprecated property cleanup purposes. */
