Navigation log improvements

diff --git a/Engine/Source/Runtime/NavigationSystem/Private/NavMesh/RecastNavMesh.cpp b/Engine/Source/Runtime/NavigationSystem/Private/NavMesh/RecastNavMesh.cpp
--- a/Engine/Source/Runtime/NavigationSystem/Private/NavMesh/RecastNavMesh.cpp
+++ b/Engine/Source/Runtime/NavigationSystem/Private/NavMesh/RecastNavMesh.cpp
@@ -1183,6 +1183,10 @@ void ARecastNavMesh::SerializeRecastNavMesh(FArchive& Ar, FPImplRecastNavMesh*&
 
 void ARecastNavMesh::Serialize( FArchive& Ar )
 {
+	UE_LOG(LogNavigation, Verbose, TEXT("%hs %s '%s'"), __FUNCTION__,
+		Ar.IsLoading() ? TEXT("IsLoading") : (Ar.IsSaving() ? TEXT("IsSaving") : TEXT("-")),
+		*GetFullNameSafe(this));
+	
 	Super::Serialize(Ar);
 
 	Ar << NavMeshVersion;
diff --git a/Engine/Source/Runtime/NavigationSystem/Private/NavMesh/RecastNavMeshGenerator.cpp b/Engine/Source/Runtime/NavigationSystem/Private/NavMesh/RecastNavMeshGenerator.cpp
--- a/Engine/Source/Runtime/NavigationSystem/Private/NavMesh/RecastNavMeshGenerator.cpp
+++ b/Engine/Source/Runtime/NavigationSystem/Private/NavMesh/RecastNavMeshGenerator.cpp
@@ -5037,6 +5037,7 @@ int32 CalculateMaxTilesCount(const TNavStatArray<FBox>& NavigableAreas, FVector:
 				{
 					continue;
 				}
+
 				const FBox& OtherBox = NavigableAreas[OtherIndex];
 				if (OtherBox.IsInsideXY(AreaBounds))
 				{
@@ -5053,6 +5054,8 @@ int32 CalculateMaxTilesCount(const TNavStatArray<FBox>& NavigableAreas, FVector:
 			}
 		}
 
+		UE_LOG(LogNavigation, Verbose, TEXT("%hs AreaBounds: %s"), __FUNCTION__, *AreaBounds.ToCompactString());
+
 		// TODO: need more precise calculation, currently we don't take into account that volumes can be overlapped
 		const FBox RCBox = Unreal2RecastBox(AreaBounds);
 
@@ -5400,7 +5403,9 @@ void FRecastNavMeshGenerator::Init()
 					if (FMath::CeilToInt(FMath::Log2(static_cast<float>(MaxTiles))) != FMath::CeilToInt(FMath::Log2(static_cast<float>(SavedNavParams->maxTiles))))
 					{
 						bRecreateNavmesh = true;
-						UE_LOG(LogNavigation, Warning, TEXT("Recreating dtNavMesh instance due mismatch in number of bytes required to store serialized maxTiles (%d, %d bits) vs calculated maxtiles (%d, %d bits)")
+						UE_LOG(LogNavigation, Warning, TEXT("Recreating dtNavMesh instance (%s) due mismatch in number of bytes required to store serialized maxTiles "
+							"(serialized: %d, %d bits) vs calculated required (%d, %d bits)")
+							, *DestNavMesh->GetName()
 							, SavedNavParams->maxTiles, FMath::CeilToInt(FMath::Log2(static_cast<float>(SavedNavParams->maxTiles)))
 							, MaxTiles, FMath::CeilToInt(FMath::Log2(static_cast<float>(MaxTiles))));
 					}
@@ -5432,6 +5437,8 @@ void FRecastNavMeshGenerator::Init()
 
 void FRecastNavMeshGenerator::UpdateNavigationBounds()
 {
+	UE_LOG(LogNavigation, Log, TEXT("%hs"), __FUNCTION__);
+	
 	const UNavigationSystemV1* NavSys = FNavigationSystem::GetCurrent<UNavigationSystemV1>(GetWorld());
 	if (NavSys)
 	{		
@@ -5451,11 +5458,13 @@ void FRecastNavMeshGenerator::UpdateNavigationBounds()
 				}
 			}
 			TotalNavBounds = BoundsSum;
+			UE_LOG(LogNavigation, Log, TEXT("  %hs TotalNavBounds: %s."), __FUNCTION__, *TotalNavBounds.ToCompactString());
 		}
 		else
 		{
 			InclusionBounds.Reset(1);
 			TotalNavBounds = NavSys->GetWorldBounds();
+			UE_LOG(LogNavigation, Log, TEXT("  %hs bWholeWorldNavigable==true, TotalNavBounds: %s."), __FUNCTION__, *TotalNavBounds.ToCompactString());
 			if (TotalNavBounds.IsValid)
 			{
 				InclusionBounds.Add(TotalNavBounds);
diff --git a/Engine/Source/Runtime/NavigationSystem/Private/NavigationData.cpp b/Engine/Source/Runtime/NavigationSystem/Private/NavigationData.cpp
--- a/Engine/Source/Runtime/NavigationSystem/Private/NavigationData.cpp
+++ b/Engine/Source/Runtime/NavigationSystem/Private/NavigationData.cpp
@@ -779,7 +779,7 @@ void ANavigationData::OnNavAreaAdded(const UClass* NavAreaClass, int32 AgentInde
 		{
 			SupportedAreas[i].AreaClass = NavAreaClass;
 			AreaClassToIdMap.Add(NavAreaClass, SupportedAreas[i].AreaID);
-			UE_VLOG_UELOG(this, LogNavigation, Verbose, TEXT("%s: updated area %s with ID %d"), *GetFullName(), *AreaClassName, SupportedAreas[i].AreaID);
+			UE_VLOG_UELOG(this, LogNavigation, VeryVerbose, TEXT("%s: updated area %s with ID %d"), *GetFullName(), *AreaClassName, SupportedAreas[i].AreaID);
 			return;
 		}
 	}
diff --git a/Engine/Source/Runtime/NavigationSystem/Private/NavigationObjectRepository.cpp b/Engine/Source/Runtime/NavigationSystem/Private/NavigationObjectRepository.cpp
--- a/Engine/Source/Runtime/NavigationSystem/Private/NavigationObjectRepository.cpp
+++ b/Engine/Source/Runtime/NavigationSystem/Private/NavigationObjectRepository.cpp
@@ -125,7 +125,7 @@ TSharedPtr<const FNavigationElement> UNavigationObjectRepository::RegisterNavRel
 			(void)OnNavigationElementAddedDelegate.ExecuteIfBound(ExistingElement.ToSharedRef());
 		}
 
-		UE_LOG(LogNavigation, Verbose, TEXT("%hs [already registered - updating] (%s:%s) Bounds: [%s]->[%s]"), __FUNCTION__,
+		UE_LOG(LogNavigation, VeryVerbose, TEXT("%hs [already registered - updating] (%s:%s) Bounds: [%s]->[%s]"), __FUNCTION__,
 			*GetNameSafe(NavRelevantObject.GetOuter()), *GetNameSafe(&NavRelevantObject),
 			*PreviousBounds.ToString(), *ExistingElement->GetBounds().ToString());
 
@@ -139,7 +139,7 @@ TSharedPtr<const FNavigationElement> UNavigationObjectRepository::RegisterNavRel
 			UE_MT_SCOPED_WRITE_ACCESS(NavElementAccessDetector);
 			ObjectsToHandleMap.Emplace(FObjectKey(&NavRelevantObject), SharedElement->GetHandle());
 
-			UE_LOG(LogNavigation, Verbose, TEXT("%hs [registered] (%s:%s) Bounds: [%s]"), __FUNCTION__,
+			UE_LOG(LogNavigation, VeryVerbose, TEXT("%hs [registered] (%s:%s) Bounds: [%s]"), __FUNCTION__,
 				*GetNameSafe(NavRelevantObject.GetOuter()), *GetNameSafe(&NavRelevantObject),
 				*NavRelevantInterface.GetNavigationBounds().ToString());
 
@@ -161,7 +161,7 @@ void UNavigationObjectRepository::UnregisterNavRelevantObject(const INavRelevant
 
 void UNavigationObjectRepository::UnregisterNavRelevantObject(const UObject* NavRelevantObject)
 {
-	UE_LOG(LogNavigation, Verbose, TEXT("%hs (%s:%s)"), __FUNCTION__,
+	UE_LOG(LogNavigation, VeryVerbose, TEXT("%hs (%s:%s)"), __FUNCTION__,
 		NavRelevantObject ? *GetNameSafe(NavRelevantObject->GetOuter()) : TEXT("null outer"),
 		*GetNameSafe(NavRelevantObject));
 
