Fix issue with navmesh when reloading map.
Increase default navbounds volume to make sure it encompass base map navbounds.
Remove unnecessary IsValid checks.
Small change to log verbosity when using bWholeWorldNavigable.

diff --git a/Engine/Source/Runtime/NavigationSystem/Private/NavigationSystem.cpp b/Engine/Source/Runtime/NavigationSystem/Private/NavigationSystem.cpp
--- a/Engine/Source/Runtime/NavigationSystem/Private/NavigationSystem.cpp
+++ b/Engine/Source/Runtime/NavigationSystem/Private/NavigationSystem.cpp
@@ -2608,7 +2608,7 @@ FBox UNavigationSystemV1::GetWorldBounds() const
 			{
 				if (IsNavigationRelevant(*It))
 				{
-					UE_LOG(LogNavigation, Verbose, TEXT("%hs %s is NavRelevant with bounds %s (%s)."),
+					UE_LOG(LogNavigation, VeryVerbose, TEXT("%hs %s is NavRelevant with bounds %s (%s)."),
 						__FUNCTION__, *It->GetName(), *(*It)->GetComponentsBoundingBox().ToCompactString(), *It->GetFullName());
 					NavigableWorldBounds += (*It)->GetComponentsBoundingBox();
 				}
