Fix debug command ai.debug.nav.DirtyAreaAroundPlayer not working on server.

diff --git a/Engine/Source/Runtime/NavigationSystem/Private/NavigationSystem.cpp b/Engine/Source/Runtime/NavigationSystem/Private/NavigationSystem.cpp
--- a/Engine/Source/Runtime/NavigationSystem/Private/NavigationSystem.cpp
+++ b/Engine/Source/Runtime/NavigationSystem/Private/NavigationSystem.cpp
@@ -186,12 +186,18 @@ namespace UE::Navigation::Private
 
 static FAutoConsoleCommandWithWorldArgsAndOutputDevice CmdNavDirtyAreaAroundPlayer(
 	TEXT("ai.debug.nav.DirtyAreaAroundPlayer"),
-	TEXT("Dirty all tiles in a square area around the local player using provided value as extent (in cm), using 10 meters if not specified."),
+	TEXT("Dirty all tiles in a square area around the player using provided value as extent (in cm), using 10 meters if not specified."),
 	FConsoleCommandWithWorldArgsAndOutputDeviceDelegate::CreateLambda([](const TArray<FString>& Args, const UWorld* World, FOutputDevice& OutputDevice)
 		{
-			if (const ULocalPlayer* LocalPlayer = World->GetFirstLocalPlayerFromController<ULocalPlayer>())
+			if (const APlayerController* Player = World->GetFirstPlayerController())
 			{
-				const FVector Center = LocalPlayer->LastViewLocation;
+				APawn* Pawn = Player->GetPawn();
+				if (!Pawn)
+				{
+					return;
+				}
+
+				const FVector Center = Pawn->GetActorLocation();
 
 				FVector::FReal Extent = 1000;
 				if (Args.Num() > 0)
