Using PostTransacted to catch NavMesh-rebuild necessitating changes for Undo/Redo.

diff --git a/Engine/Source/Runtime/NavigationSystem/Private/NavMesh/RecastNavMesh.cpp b/Engine/Source/Runtime/NavigationSystem/Private/NavMesh/RecastNavMesh.cpp
--- a/Engine/Source/Runtime/NavigationSystem/Private/NavMesh/RecastNavMesh.cpp
+++ b/Engine/Source/Runtime/NavigationSystem/Private/NavMesh/RecastNavMesh.cpp
@@ -33,6 +33,7 @@
 #include "WorldPartition/WorldPartitionEditorLoaderAdapter.h"
 #include "WorldPartition/LoaderAdapter/LoaderAdapterShape.h"
 #include "LevelInstance/LevelInstanceSubsystem.h"
+#include "Misc/TransactionObjectEvent.h"
 #endif
 
 #if WITH_RECAST
@@ -3631,6 +3632,18 @@ void ARecastNavMesh::UpdateNavVersion()
 
 #if WITH_EDITOR
 
+bool ARecastNavMesh::CanUpdateNavMeshInEditor() const
+{
+	const UNavigationSystemV1* NavSys = FNavigationSystem::GetCurrent<UNavigationSystemV1>(GetWorld());
+	if (!HasAnyFlags(RF_ClassDefaultObject)
+		&& NavSys && NavSys->GetIsAutoUpdateEnabled())
+	{
+		return true;
+	}
+
+	return false;
+}
+
 void ARecastNavMesh::PostEditChangeChainProperty(FPropertyChangedChainEvent& PropertyChangedChainEvent)
 {
 	static const FName NAME_Generation = FName("Generation");
@@ -3704,9 +3717,7 @@ void ARecastNavMesh::PostEditChangeChainProperty(FPropertyChangedChainEvent& Pro
 				// update config
 				FillConfig(NavDataConfig);
 
-				const UNavigationSystemV1* NavSys = FNavigationSystem::GetCurrent<UNavigationSystemV1>(GetWorld());
-				if (!HasAnyFlags(RF_ClassDefaultObject)
-					&& NavSys && NavSys->GetIsAutoUpdateEnabled()
+				if (CanUpdateNavMeshInEditor()
 					&& PropName != GET_MEMBER_NAME_CHECKED(ARecastNavMesh, MaxSimultaneousTileGenerationJobsCount))
 				{
 					RebuildAll();
@@ -3732,7 +3743,6 @@ void ARecastNavMesh::PostEditChangeProperty(FPropertyChangedEvent& PropertyChang
 		if (CategoryName == NAME_Generation)
 		{
 			const FName PropName = PropertyChangedEvent.Property->GetFName();
-			const FName MemberName = PropertyChangedEvent.MemberProperty->GetFName();
 			
 			if (PropName == GET_MEMBER_NAME_CHECKED(ARecastNavMesh, AgentRadius))
 			{
@@ -3766,9 +3776,7 @@ void ARecastNavMesh::PostEditChangeProperty(FPropertyChangedEvent& PropertyChang
 				UpdatePolyRefBitsPreview();
 			}
 
-			UNavigationSystemV1* NavSys = FNavigationSystem::GetCurrent<UNavigationSystemV1>(GetWorld());
-			if (!HasAnyFlags(RF_ClassDefaultObject)
-				&& NavSys && NavSys->GetIsAutoUpdateEnabled()
+			if (CanUpdateNavMeshInEditor()
 				&& PropName != GET_MEMBER_NAME_CHECKED(ARecastNavMesh, MaxSimultaneousTileGenerationJobsCount))
 			{
 				RebuildAll();
@@ -3809,6 +3817,61 @@ void ARecastNavMesh::PostEditUndo()
 	Super::PostEditUndo();
 }
 
+void ARecastNavMesh::PostTransacted(const FTransactionObjectEvent& TransactionEvent)
+{
+	Super::PostTransacted(TransactionEvent);
+
+	if (TransactionEvent.GetEventType() == ETransactionObjectEventType::UndoRedo)
+	{
+		// (Note, that this doesn't handle property chains at all. It can also match identical property names in a category outside "Generation", which is undesirable.)
+
+		if (CanUpdateNavMeshInEditor())
+		{
+			struct FNavMeshGenerationPropertyNames
+			{
+				FNavMeshGenerationPropertyNames()
+				{
+					static const FName NAME_Generation = FName("Generation");
+
+					// Assume almost all matches in the "Generation" category will trigger a rebuild.
+
+					for (TFieldIterator<FProperty> PropertyIt(ARecastNavMesh::StaticClass()); PropertyIt; ++PropertyIt)
+					{
+						const FProperty* Property = *PropertyIt;
+						const FName CategoryName = FObjectEditorUtils::GetCategoryFName(Property);
+						if (CategoryName != NAME_Generation)
+							continue;
+
+						const FName PropName = Property->GetFName();
+						if (PropName != GET_MEMBER_NAME_CHECKED(ARecastNavMesh, MaxSimultaneousTileGenerationJobsCount))
+							ValidProperties.Add(PropName);
+					}
+
+				}
+
+				TArray<FName> ValidProperties;
+			};
+			static FNavMeshGenerationPropertyNames BuiltPropertyNames;
+
+			const TArray<FName>& PropertiesChanged = TransactionEvent.GetChangedProperties();
+
+			const bool bRebuildNavMesh = PropertiesChanged.ContainsByPredicate([](const FName Name)
+				{
+					return BuiltPropertyNames.ValidProperties.Contains(Name);
+				});
+
+			if (bRebuildNavMesh)
+			{
+				// update config
+				FillConfig(NavDataConfig);
+
+				RebuildAll();
+			}
+		}
+	}
+}
+
+
 #endif // WITH_EDITOR
 
 bool ARecastNavMesh::NeedsRebuild() const
diff --git a/Engine/Source/Runtime/NavigationSystem/Public/NavMesh/RecastNavMesh.h b/Engine/Source/Runtime/NavigationSystem/Public/NavMesh/RecastNavMesh.h
--- a/Engine/Source/Runtime/NavigationSystem/Public/NavMesh/RecastNavMesh.h
+++ b/Engine/Source/Runtime/NavigationSystem/Public/NavMesh/RecastNavMesh.h
@@ -1005,14 +1005,20 @@ public:
 	NAVIGATIONSYSTEM_API virtual void BeginDestroy() override;
 
 #if WITH_EDITOR
+
+private:
+	bool CanUpdateNavMeshInEditor() const;
+
+public:
 	NAVIGATIONSYSTEM_API virtual void PostEditChangeChainProperty( struct FPropertyChangedChainEvent& PropertyChangedChainEvent) override;
 	NAVIGATIONSYSTEM_API virtual void PostEditChangeProperty( struct FPropertyChangedEvent& PropertyChangedEvent) override;
 	NAVIGATIONSYSTEM_API virtual void PostEditUndo() override;
+	NAVIGATIONSYSTEM_API virtual void PostTransacted(const FTransactionObjectEvent& TransactionEvent) override;
 #endif // WITH_EDITOR
 	//~ End UObject Interface
 
 #if WITH_EDITOR
-	/** RecastNavMesh instances are dynamically spawned and should not be coppied */
+	/** RecastNavMesh instances are dynamically spawned and should not be copied */
 	virtual bool ShouldExport() override { return false; }
 #endif
 
