[RecastNavMesh]
- Fix race on unneeded static

diff --git a/Engine/Source/Runtime/NavigationSystem/Public/NavMesh/RecastNavMesh.h b/Engine/Source/Runtime/NavigationSystem/Public/NavMesh/RecastNavMesh.h
--- a/Engine/Source/Runtime/NavigationSystem/Public/NavMesh/RecastNavMesh.h
+++ b/Engine/Source/Runtime/NavigationSystem/Public/NavMesh/RecastNavMesh.h
@@ -468,10 +468,6 @@ struct FNavMeshTileData
 		{
 			if (RawNavData != nullptr)
 			{
-				// Temporary test to help reproduce a crash.
-				static uint8 Temp = 0;
-				Temp = *RawNavData;
-				
 				AllocatedSize = FMemory::GetAllocSize((void*)RawNavData);
 				check(AllocatedSize == 0 || AllocatedSize >= InDataSize);
 			}
