Fix ensure in ComputeCustomLinkBounds

diff --git a/Engine/Source/Runtime/NavigationSystem/Private/NavigationSystem.cpp b/Engine/Source/Runtime/NavigationSystem/Private/NavigationSystem.cpp
--- a/Engine/Source/Runtime/NavigationSystem/Private/NavigationSystem.cpp
+++ b/Engine/Source/Runtime/NavigationSystem/Private/NavigationSystem.cpp
@@ -3056,7 +3056,10 @@ FBox UNavigationSystemV1::ComputeCustomLinkBounds(const INavLinkCustomInterface&
 	const AActor* OwnerActor = OwnerComp ? OwnerComp->GetOwner() : Cast<AActor>(CustomLinkOb);
 
 	FBox LinkBounds(ForceInitToZero);
-	if (OwnerActor)
+
+	// If the owner is a navdata, the NavLink Custom Interface is shared for multiple link instances.
+	// It does not have individual link data so act has if it had no owner actor.
+	if (OwnerActor && !OwnerActor->IsA(ANavigationData::StaticClass()))
 	{
 		ENavLinkDirection::Type DummyDir = ENavLinkDirection::BothWays;
 		FVector RelativePtA, RelativePtB;
