Faster FRecastNavMeshGenerator::RemoveTiles (tentative fix for hang when there are too many dirty tiles)

diff --git a/Engine/Source/Runtime/NavigationSystem/Private/NavMesh/RecastNavMeshGenerator.cpp b/Engine/Source/Runtime/NavigationSystem/Private/NavMesh/RecastNavMeshGenerator.cpp
--- a/Engine/Source/Runtime/NavigationSystem/Private/NavMesh/RecastNavMeshGenerator.cpp
+++ b/Engine/Source/Runtime/NavigationSystem/Private/NavMesh/RecastNavMeshGenerator.cpp
@@ -5919,7 +5919,7 @@ void FRecastNavMeshGenerator::RemoveTiles(const TArray<FIntPoint>& Tiles)
 		{
 			FPendingTileElement DirtyTile;
 			DirtyTile.Coord = TileXY;
-			PendingDirtyTiles.Remove(DirtyTile);
+			PendingDirtyTiles.RemoveSwap(DirtyTile);
 		}
 
 		if (SyncTimeSlicedData.TileGeneratorSync.IsValid())
@@ -6013,21 +6013,6 @@ void FRecastNavMeshGenerator::ReAddTiles(const TArray<FNavMeshDirtyTileElement>&
 	{
 		SortPendingBuildTiles();
 	}
-
-	/*TArray<FNavigationDirtyArea> DirtyAreasContainer;
-	DirtyAreasContainer.Reserve(Tiles.Num());
-
-	TSet<FPendingTileElement> DirtyTiles;
-
-	for (const FIntPoint& TileCoords : Tiles)
-	{
-		const FVector TileCenter = Recast2UnrealPoint(SavedNavParams->orig) + FVector(TileDim * float(TileCoords.X), TileDim * float(TileCoords.Y), 0);
-		
-		FNavigationDirtyArea DirtyArea(FBox(TileCenter - Expansion, TileCenter - 1), ENavigationDirtyFlag::All);
-		DirtyAreasContainer.Add(DirtyArea);
-	}
-
-	MarkDirtyTiles(DirtyAreasContainer);*/
 }
 
 namespace RecastTileVersionHelper
