Fix to generate the spline points in local space instead of world space. Previously, the world transform was applied twice.

diff --git a/Engine/Source/Runtime/NavigationSystem/Private/SplineNavModifierComponent.cpp b/Engine/Source/Runtime/NavigationSystem/Private/SplineNavModifierComponent.cpp
--- a/Engine/Source/Runtime/NavigationSystem/Private/SplineNavModifierComponent.cpp
+++ b/Engine/Source/Runtime/NavigationSystem/Private/SplineNavModifierComponent.cpp
@@ -26,8 +26,8 @@ namespace
 		{
 			if (PrevIndex >= 0)
 			{
-				const FSplinePoint PrevSplinePoint = Spline.GetSplinePointAt(PrevIndex, ESplineCoordinateSpace::World);
-				const FSplinePoint CurrSplinePoint = Spline.GetSplinePointAt(SplinePointIndex, ESplineCoordinateSpace::World);
+				const FSplinePoint PrevSplinePoint = Spline.GetSplinePointAt(PrevIndex, ESplineCoordinateSpace::Local);
+				const FSplinePoint CurrSplinePoint = Spline.GetSplinePointAt(SplinePointIndex, ESplineCoordinateSpace::Local);
 
 				// The first point of the segment is appended before tessellation since UE::CubicBezier::Tessellate does not add it
 				OutSubdivisions.Add(PrevSplinePoint.Position);
