Fix to make it so EncompassesPoint works on ANavMeshBoundsVolumes

diff --git a/Engine/Source/Runtime/Engine/Classes/Engine/CollisionProfile.h b/Engine/Source/Runtime/Engine/Classes/Engine/CollisionProfile.h
--- a/Engine/Source/Runtime/Engine/Classes/Engine/CollisionProfile.h
+++ b/Engine/Source/Runtime/Engine/Classes/Engine/CollisionProfile.h
@@ -188,6 +188,7 @@ public:
 	/** default property name for no collision - this is very popular **/
 	static ENGINE_API const FName NoCollision_ProfileName;
 	static ENGINE_API const FName BlockAll_ProfileName;
+	static ENGINE_API const FName OverlapAll_ProfileName;
 	static ENGINE_API const FName PhysicsActor_ProfileName;
 	static ENGINE_API const FName BlockAllDynamic_ProfileName;
 	static ENGINE_API const FName Pawn_ProfileName;
diff --git a/Engine/Source/Runtime/Engine/Private/Collision/CollisionProfile.cpp b/Engine/Source/Runtime/Engine/Private/Collision/CollisionProfile.cpp
--- a/Engine/Source/Runtime/Engine/Private/Collision/CollisionProfile.cpp
+++ b/Engine/Source/Runtime/Engine/Private/Collision/CollisionProfile.cpp
@@ -41,6 +41,7 @@ bool FCollisionResponseTemplate::IsEqual(const TEnumAsByte<ECollisionEnabled::Ty
 
 ENGINE_API const FName UCollisionProfile::NoCollision_ProfileName = FName("NoCollision");
 ENGINE_API const FName UCollisionProfile::BlockAll_ProfileName = FName("BlockAll");
+ENGINE_API const FName UCollisionProfile::OverlapAll_ProfileName = FName("OverlapAll");
 ENGINE_API const FName UCollisionProfile::PhysicsActor_ProfileName = FName("PhysicsActor");
 ENGINE_API const FName UCollisionProfile::BlockAllDynamic_ProfileName = FName("BlockAllDynamic");
 ENGINE_API const FName UCollisionProfile::Pawn_ProfileName = FName("Pawn");
diff --git a/Engine/Source/Runtime/NavigationSystem/Private/NavMesh/NavMeshBoundsVolume.cpp b/Engine/Source/Runtime/NavigationSystem/Private/NavMesh/NavMeshBoundsVolume.cpp
--- a/Engine/Source/Runtime/NavigationSystem/Private/NavMesh/NavMeshBoundsVolume.cpp
+++ b/Engine/Source/Runtime/NavigationSystem/Private/NavMesh/NavMeshBoundsVolume.cpp
@@ -12,6 +12,26 @@
 #include "Editor.h"
 #endif // WITH_EDITOR
 
+namespace UE::Navigation::Private
+{
+	static void ApplyNavMeshBoundsCollision(UBrushComponent* Brush, const bool bAllowOverlap)
+	{
+		if (!Brush)
+		{
+			return;
+		}
+
+		if (bAllowOverlap)
+		{
+			Brush->SetCollisionProfileName(UCollisionProfile::OverlapAll_ProfileName);
+		}
+		else
+		{
+			Brush->SetCollisionProfileName(UCollisionProfile::NoCollision_ProfileName);
+		}
+	}
+}
+
 ANavMeshBoundsVolume::ANavMeshBoundsVolume(const FObjectInitializer& ObjectInitializer)
 	: Super(ObjectInitializer)
 {
@@ -23,6 +43,8 @@ ANavMeshBoundsVolume::ANavMeshBoundsVolume(const FObjectInitializer& ObjectIniti
 
 	bColored = true;
 
+	bAllowPhysicsOverlap = false;
+
 #if WITH_EDITORONLY_DATA
 	bIsSpatiallyLoaded = false;
 #endif
@@ -49,6 +71,15 @@ void ANavMeshBoundsVolume::PostEditChangeProperty( struct FPropertyChangedEvent&
 			NavSys->OnNavigationBoundsUpdated(this);
 		}
 	}
+
+	if (PropertyChangedEvent.Property != nullptr)
+	{
+		static const FName NAME_bAllowPhysicsOverlap(TEXT("bAllowPhysicsOverlap"));
+		if (PropertyChangedEvent.Property->GetFName() == NAME_bAllowPhysicsOverlap)
+		{
+			UE::Navigation::Private::ApplyNavMeshBoundsCollision(GetBrushComponent(), bAllowPhysicsOverlap);
+		}
+	}
 }
 
 void ANavMeshBoundsVolume::PostEditUndo()
@@ -59,6 +90,8 @@ void ANavMeshBoundsVolume::PostEditUndo()
 	{
 		NavSys->OnNavigationBoundsUpdated(this);
 	}
+
+	UE::Navigation::Private::ApplyNavMeshBoundsCollision(GetBrushComponent(), bAllowPhysicsOverlap);
 }
 
 void ANavMeshBoundsVolume::OnPostEngineInit()
@@ -90,6 +123,8 @@ void ANavMeshBoundsVolume::PostRegisterAllComponents()
 	{
 		NavSys->OnNavigationBoundsAdded(this);
 	}
+
+	UE::Navigation::Private::ApplyNavMeshBoundsCollision(GetBrushComponent(), bAllowPhysicsOverlap);
 }
 
 void ANavMeshBoundsVolume::PostUnregisterAllComponents() 
diff --git a/Engine/Source/Runtime/NavigationSystem/Public/NavMesh/NavMeshBoundsVolume.h b/Engine/Source/Runtime/NavigationSystem/Public/NavMesh/NavMeshBoundsVolume.h
--- a/Engine/Source/Runtime/NavigationSystem/Public/NavMesh/NavMeshBoundsVolume.h
+++ b/Engine/Source/Runtime/NavigationSystem/Public/NavMesh/NavMeshBoundsVolume.h
@@ -27,6 +27,10 @@ class ANavMeshBoundsVolume : public AVolume
 	UPROPERTY(EditAnywhere, Category = Navigation)
 	FNavAgentSelector SupportedAgents;
 
+	/* If true, then this volume will have a BodyInstance and calls to EncompassesPoint will function as expected. */
+	UPROPERTY(EditAnywhere, Category = Navigation)
+	uint8 bAllowPhysicsOverlap : 1 = false;
+
 	//~ Begin AActor Interface
 	NAVIGATIONSYSTEM_API virtual void PostRegisterAllComponents() override;
 	NAVIGATIONSYSTEM_API virtual void PostUnregisterAllComponents() override;
