Miscellaneous improvements to navigation logs

diff --git a/Engine/Source/Runtime/Engine/Private/AI/NavigationSystemBase.cpp b/Engine/Source/Runtime/Engine/Private/AI/NavigationSystemBase.cpp
--- a/Engine/Source/Runtime/Engine/Private/AI/NavigationSystemBase.cpp
+++ b/Engine/Source/Runtime/Engine/Private/AI/NavigationSystemBase.cpp
@@ -11,7 +11,6 @@
 #include "AI/Navigation/NavRelevantInterface.h"
 #include "AI/Navigation/NavigationElement.h"
 
-
 #include UE_INLINE_GENERATED_CPP_BY_NAME(NavigationSystemBase)
 
 DEFINE_LOG_CATEGORY(LogNavigation);
@@ -21,9 +20,6 @@ DEFINE_LOG_CATEGORY(LogNavInvokers);
 DEFINE_LOG_CATEGORY(LogNavLink);
 DEFINE_LOG_CATEGORY(LogAStar);
 
-#if !UE_BUILD_SHIPPING
-#endif // !UE_BUILD_SHIPPING
-
 namespace FNavigationSystem
 {
 	void DiscardNavigationDataChunks(UWorld& InWorld)
diff --git a/Engine/Source/Runtime/Engine/Private/World.cpp b/Engine/Source/Runtime/Engine/Private/World.cpp
--- a/Engine/Source/Runtime/Engine/Private/World.cpp
+++ b/Engine/Source/Runtime/Engine/Private/World.cpp
@@ -9473,6 +9473,7 @@ void UWorld::SetNavigationSystem(UNavigationSystemBase* InNavigationSystem)
 		NavigationSystem->CleanUp(FNavigationSystem::ECleanupMode::CleanupWithWorld);
 	}
 
+	UE_LOG(LogNavigation, Verbose, TEXT("   %hs to %s"), __FUNCTION__, *GetNameSafe(NavigationSystem));
 	NavigationSystem = InNavigationSystem;
 }
 
diff --git a/Engine/Source/Runtime/NavigationSystem/Private/NavMesh/RecastNavMesh.cpp b/Engine/Source/Runtime/NavigationSystem/Private/NavMesh/RecastNavMesh.cpp
--- a/Engine/Source/Runtime/NavigationSystem/Private/NavMesh/RecastNavMesh.cpp
+++ b/Engine/Source/Runtime/NavigationSystem/Private/NavMesh/RecastNavMesh.cpp
@@ -1186,9 +1186,11 @@ void ARecastNavMesh::SerializeRecastNavMesh(FArchive& Ar, FPImplRecastNavMesh*&
 void ARecastNavMesh::Serialize( FArchive& Ar )
 {
 	UE_LOG(LogNavigation, Verbose, TEXT("%hs %s '%s'"), __FUNCTION__,
-		Ar.IsLoading() ? TEXT("IsLoading") : (Ar.IsSaving() ? TEXT("IsSaving") : TEXT("-")),
+		Ar.IsLoading() ? TEXT("'loading'") :
+			(Ar.IsSaving() ? TEXT("'saving'") :
+				(Ar.IsTransacting() ? TEXT("'transacting'") : TEXT("'unknown'"))),
 		*GetFullNameSafe(this));
-	
+
 	Super::Serialize(Ar);
 
 	Ar << NavMeshVersion;
diff --git a/Engine/Source/Runtime/NavigationSystem/Private/NavigationData.cpp b/Engine/Source/Runtime/NavigationSystem/Private/NavigationData.cpp
--- a/Engine/Source/Runtime/NavigationSystem/Private/NavigationData.cpp
+++ b/Engine/Source/Runtime/NavigationSystem/Private/NavigationData.cpp
@@ -799,7 +799,7 @@ void ANavigationData::OnNavAreaAdded(const UClass* NavAreaClass, int32 AgentInde
 	SupportedAreas.Add(NewAgentData);
 	AreaClassToIdMap.Add(NavAreaClass, NewAgentData.AreaID);
 
-	UE_VLOG_UELOG(this, LogNavigation, Verbose, TEXT("%s registered area %s with ID %d"), *GetName(), *AreaClassName, NewAgentData.AreaID);
+	UE_VLOG_UELOG(this, LogNavigation, VeryVerbose, TEXT("  %s registered area %s with ID %d"), *GetPathName(), *AreaClassName, NewAgentData.AreaID);
 }
 
 void ANavigationData::OnNavAreaEvent(const UClass* NavAreaClass, ENavAreaEvent::Type Event)
diff --git a/Engine/Source/Runtime/NavigationSystem/Private/NavigationSystem.cpp b/Engine/Source/Runtime/NavigationSystem/Private/NavigationSystem.cpp
--- a/Engine/Source/Runtime/NavigationSystem/Private/NavigationSystem.cpp
+++ b/Engine/Source/Runtime/NavigationSystem/Private/NavigationSystem.cpp
@@ -4898,6 +4898,11 @@ void UNavigationSystemV1::RebuildAll(bool bIsLoadTime)
 #endif
 			NavData->RebuildAll();
 		}
+		else if (NavData)
+		{
+			UE_LOG(LogNavigationDataBuild, Verbose, TEXT("   RebuildAll skipped.  bIsLoadTime=%s, NavData->NeedsRebuildOnLoad()=%s, bIsInGame=%s, NavData->SupportsRuntimeGeneration()=%s, BuildBounds.IsValid=%s"),
+				*LexToString(bIsLoadTime), *LexToString(NavData->NeedsRebuildOnLoad()), *LexToString(bIsInGame), *LexToString(NavData->SupportsRuntimeGeneration()), *LexToString(BuildBounds.IsValid));
+		}
 	}
 }
 
