Add way to visualize nav area modifiers in navmesh generation

#Synced-CL 48202918

diff --git a/Engine/Source/Runtime/NavigationSystem/Private/NavMesh/RecastNavMesh.cpp b/Engine/Source/Runtime/NavigationSystem/Private/NavMesh/RecastNavMesh.cpp
--- a/Engine/Source/Runtime/NavigationSystem/Private/NavMesh/RecastNavMesh.cpp
+++ b/Engine/Source/Runtime/NavigationSystem/Private/NavMesh/RecastNavMesh.cpp
@@ -443,6 +443,7 @@ FRecastNavMeshTileGenerationDebug::FRecastNavMeshTileGenerationDebug()
 	MaxTileCoordinate = FIntVector::ZeroValue;
 	bGenerateDebugTileOnly = false;
 	bCollisionGeometry = false;
+	bDynamicAreas = false;
 	HeightFieldRenderMode = EHeightFieldRenderMode::Solid;
 	bHeightfieldFromRasterization = false;
 	bHeightfieldPostInclusionBoundsFiltering = false;
diff --git a/Engine/Source/Runtime/NavigationSystem/Private/NavMesh/RecastNavMeshGenerator.cpp b/Engine/Source/Runtime/NavigationSystem/Private/NavMesh/RecastNavMeshGenerator.cpp
--- a/Engine/Source/Runtime/NavigationSystem/Private/NavMesh/RecastNavMeshGenerator.cpp
+++ b/Engine/Source/Runtime/NavigationSystem/Private/NavMesh/RecastNavMeshGenerator.cpp
@@ -2395,6 +2395,27 @@ void FRecastTileGenerator::PrepareVoxelCache(const TNavStatArray<uint8>& RawColl
 	return PrepareVoxelCache(RawCollisionCache, FBox(ForceInit), InModifier, SpanData);
 }
 
+// Deprecated
+void FRecastTileGenerator::MarkDynamicAreas(dtTileCacheLayer& Layer)
+{
+	FNavMeshBuildContext BuildContext(*this);
+	return MarkDynamicAreas(BuildContext, Layer);
+}
+
+// Deprecated
+void FRecastTileGenerator::MarkDynamicArea(const FAreaNavModifier& Modifier, const FTransform& LocalToWorld, dtTileCacheLayer& Layer)
+{
+	FNavMeshBuildContext BuildContext(*this);
+	return MarkDynamicArea(BuildContext, Modifier, LocalToWorld, Layer);
+}
+
+// Deprecated
+void FRecastTileGenerator::MarkDynamicArea(const FAreaNavModifier& Modifier, const FTransform& LocalToWorld, dtTileCacheLayer& Layer, const int32 AreaId, const int32* ReplaceAreaIdPtr)
+{
+	FNavMeshBuildContext BuildContext(*this);
+	return MarkDynamicArea(BuildContext, Modifier, LocalToWorld, Layer, AreaId, ReplaceAreaIdPtr);
+}
+
 void FRecastTileGenerator::PrepareVoxelCache(const TNavStatArray<uint8>& RawCollisionCache, const FBox& RecastBounds, const FCompositeNavModifier& InModifier, TNavStatArray<rcSpanCache>& SpanData)
 {
 	// tile's geometry: voxel cache (only for synchronous rebuilds)
@@ -4099,7 +4120,7 @@ bool FRecastTileGenerator::GenerateNavigationDataLayer(FNavMeshBuildContext& Bui
 	}
 
 	// Rasterize obstacles.
-	MarkDynamicAreas(*GenerationContext.Layer);
+	MarkDynamicAreas(BuildContext, *GenerationContext.Layer);
 
 	{
 		SCOPE_CYCLE_COUNTER(STAT_Navigation_RecastBuildRegions)
@@ -4548,7 +4569,7 @@ void FRecastTileGenerator::ComputeRasterizationMasks(FNavMeshBuildContext& Build
 	}
 }
 
-void FRecastTileGenerator::MarkDynamicAreas(dtTileCacheLayer& Layer)
+void FRecastTileGenerator::MarkDynamicAreas(FNavMeshBuildContext& BuildContext, dtTileCacheLayer& Layer)
 {
 	SCOPE_CYCLE_COUNTER(STAT_Navigation_RecastMarkAreas);
 
@@ -4584,12 +4605,12 @@ void FRecastTileGenerator::MarkDynamicAreas(dtTileCacheLayer& Layer)
 						{
 							for (const FTransform& LocalToWorld : Element.PerInstanceTransform)
 							{
-								MarkDynamicArea(AreaMod, LocalToWorld, Layer, *AreaIDPtr, ReplaceAreaIDPtr);
+								MarkDynamicArea(BuildContext, AreaMod, LocalToWorld, Layer, *AreaIDPtr, ReplaceAreaIDPtr);
 							}
 
 							if (Element.PerInstanceTransform.Num() == 0)
 							{
-								MarkDynamicArea(AreaMod, FTransform::Identity, Layer, *AreaIDPtr, ReplaceAreaIDPtr);
+								MarkDynamicArea(BuildContext, AreaMod, FTransform::Identity, Layer, *AreaIDPtr, ReplaceAreaIDPtr);
 							}
 						}
 					}
@@ -4618,12 +4639,12 @@ void FRecastTileGenerator::MarkDynamicAreas(dtTileCacheLayer& Layer)
 				{
 					for (const FTransform& LocalToWorld : Element.PerInstanceTransform)
 					{
-						MarkDynamicArea(Area, LocalToWorld, Layer, *AreaIDPtr, ReplaceIDPtr);
+						MarkDynamicArea(BuildContext, Area, LocalToWorld, Layer, *AreaIDPtr, ReplaceIDPtr);
 					}
 
 					if (Element.PerInstanceTransform.Num() == 0)
 					{
-						MarkDynamicArea(Area, FTransform::Identity, Layer, *AreaIDPtr, ReplaceIDPtr);
+						MarkDynamicArea(BuildContext, Area, FTransform::Identity, Layer, *AreaIDPtr, ReplaceIDPtr);
 					}
 				}
 			}
@@ -4638,13 +4659,13 @@ void FRecastTileGenerator::MarkDynamicAreas(dtTileCacheLayer& Layer)
 	}
 }
 
-void FRecastTileGenerator::MarkDynamicArea(const FAreaNavModifier& Modifier, const FTransform& LocalToWorld, dtTileCacheLayer& Layer)
+void FRecastTileGenerator::MarkDynamicArea(FNavMeshBuildContext& BuildContext, const FAreaNavModifier& Modifier, const FTransform& LocalToWorld, dtTileCacheLayer& Layer)
 {
 	const int32* AreaIDPtr = AdditionalCachedData.AreaClassToIdMap.Find(Modifier.GetAreaClass());
 	const int32* ReplaceIDPtr = Modifier.GetAreaClassToReplace() ? AdditionalCachedData.AreaClassToIdMap.Find(Modifier.GetAreaClassToReplace()) : nullptr;
 	if (AreaIDPtr)
 	{
-		MarkDynamicArea(Modifier, LocalToWorld, Layer, *AreaIDPtr, ReplaceIDPtr);
+		MarkDynamicArea(BuildContext, Modifier, LocalToWorld, Layer, *AreaIDPtr, ReplaceIDPtr);
 	}
 }
 
@@ -4809,7 +4830,7 @@ void FRecastTileGenerator::MarkRasterizationMask(rcContext* /*BuildContext*/, rc
 	}
 }
 
-void FRecastTileGenerator::MarkDynamicArea(const FAreaNavModifier& Modifier, const FTransform& LocalToWorld, dtTileCacheLayer& Layer, const int32 AreaID, const int32* ReplaceIDPtr)
+void FRecastTileGenerator::MarkDynamicArea(FNavMeshBuildContext& BuildContext, const FAreaNavModifier& Modifier, const FTransform& LocalToWorld, dtTileCacheLayer& Layer, const int32 AreaID, const int32* ReplaceIDPtr)
 {
 	const float ExpandBy = TileConfig.AgentRadius;
 
@@ -4829,6 +4850,23 @@ void FRecastTileGenerator::MarkDynamicArea(const FAreaNavModifier& Modifier, con
 		return;
 	}
 
+#if RECAST_INTERNAL_DEBUG_DATA
+	const bool bAddShapesToDebugDraw = IsTileDebugActive() && TileDebugSettings.bDynamicAreas;
+	const ARecastNavMesh* RecastParentNavmesh = nullptr;
+	if (bAddShapesToDebugDraw)
+	{
+		TSharedPtr<FNavDataGenerator, ESPMode::ThreadSafe> ParentGenerator = ParentGeneratorWeakPtr.Pin();
+		if (ParentGenerator.IsValid())
+		{
+			TSharedPtr<FRecastNavMeshGenerator> RecastParentGenerator = StaticCastSharedPtr<FRecastNavMeshGenerator>(ParentGenerator);
+			if (const FRecastNavMeshGenerator* RecastNavMeshGeneratorPtr = RecastParentGenerator.Get())
+			{
+				RecastParentNavmesh = RecastNavMeshGeneratorPtr->GetOwner();
+			}
+		}
+	}
+#endif // RECAST_INTERNAL_DEBUG_DATA
+
 	const FVector::FReal* LayerRecastOrig = Layer.header->bmin;
 	switch (Modifier.GetShapeType())
 	{
@@ -4849,7 +4887,7 @@ void FRecastTileGenerator::MarkDynamicArea(const FAreaNavModifier& Modifier, con
 			CylinderData.Height += FMath::Abs(OffsetZMid) * 2.f;
 			CylinderData.Radius += ExpandBy;
 			
-			FVector RecastPos = Unreal2RecastPoint(CylinderData.Origin);
+			const FVector RecastPos = Unreal2RecastPoint(CylinderData.Origin);
 
 			if (ReplaceIDPtr)
 			{
@@ -4859,7 +4897,16 @@ void FRecastTileGenerator::MarkDynamicArea(const FAreaNavModifier& Modifier, con
 			else
 			{
 				dtMarkCylinderArea(Layer, LayerRecastOrig, TileConfig.cs, TileConfig.ch,
-					&(RecastPos.X), CylinderData.Radius, CylinderData.Height, IntCastChecked<unsigned char>(AreaID));			}
+					&(RecastPos.X), CylinderData.Radius, CylinderData.Height, IntCastChecked<unsigned char>(AreaID));
+			}
+
+#if RECAST_INTERNAL_DEBUG_DATA
+			if (bAddShapesToDebugDraw)
+			{
+				const FColor AreaColor = RecastParentNavmesh ? RecastParentNavmesh->GetAreaIDColor(IntCastChecked<unsigned char>(AreaID)) : FColor::White;
+				duDrawCylinderArea(&BuildContext.InternalDebugData, &(RecastPos.X), CylinderData.Radius, CylinderData.Height, AreaColor.ToPackedRGBA());
+			}
+#endif // RECAST_INTERNAL_DEBUG_DATA
 		}
 		break;
 
@@ -4887,6 +4934,14 @@ void FRecastTileGenerator::MarkDynamicArea(const FAreaNavModifier& Modifier, con
 				dtMarkBoxArea(Layer, LayerRecastOrig, TileConfig.cs, TileConfig.ch,
 					&(RecastPos.X), &(RecastExtent.X), IntCastChecked<unsigned char>(AreaID));
 			}
+
+#if RECAST_INTERNAL_DEBUG_DATA
+			if (bAddShapesToDebugDraw)
+			{
+				const FColor AreaColor = RecastParentNavmesh ? RecastParentNavmesh->GetAreaIDColor(IntCastChecked<unsigned char>(AreaID)) : FColor::White;
+				duDrawBoxArea(&BuildContext.InternalDebugData, &(RecastPos.X), &(RecastExtent.X), AreaColor.ToPackedRGBA());
+			}
+#endif // RECAST_INTERNAL_DEBUG_DATA
 		}
 		break;
 
@@ -4934,6 +4989,14 @@ void FRecastTileGenerator::MarkDynamicArea(const FAreaNavModifier& Modifier, con
 					dtMarkConvexArea(Layer, LayerRecastOrig, TileConfig.cs, TileConfig.ch,
 						ConvexCoords.GetData(), ConvexVerts.Num(), ConvexData.MinZ, ConvexData.MaxZ, IntCastChecked<unsigned char>(AreaID));
 				}
+
+#if RECAST_INTERNAL_DEBUG_DATA
+				if (bAddShapesToDebugDraw)
+				{
+					const FColor AreaColor = RecastParentNavmesh ? RecastParentNavmesh->GetAreaIDColor(IntCastChecked<unsigned char>(AreaID)) : FColor::White;
+					duDrawConvexArea(&BuildContext.InternalDebugData,ConvexCoords.GetData(), ConvexVerts.Num(), ConvexData.MinZ, ConvexData.MaxZ, AreaColor.ToPackedRGBA());
+				}
+#endif // RECAST_INTERNAL_DEBUG_DATA
 			}
 		}
 		break;
diff --git a/Engine/Source/Runtime/NavigationSystem/Public/NavMesh/RecastNavMesh.h b/Engine/Source/Runtime/NavigationSystem/Public/NavMesh/RecastNavMesh.h
--- a/Engine/Source/Runtime/NavigationSystem/Public/NavMesh/RecastNavMesh.h
+++ b/Engine/Source/Runtime/NavigationSystem/Public/NavMesh/RecastNavMesh.h
@@ -375,6 +375,11 @@ struct FRecastNavMeshTileGenerationDebug
 	UPROPERTY(EditAnywhere, Category = Debug)
 	uint32 bCollisionGeometry : 1;
 
+	/** Display the nav area modifier shapes.
+	 * Note: The visualization is affected by the DrawOffset of the RecastNavmesh owner*/
+	UPROPERTY(EditAnywhere, Category = Debug)
+	uint32 bDynamicAreas : 1;
+
 	UPROPERTY(EditAnywhere, Category = Debug)
 	EHeightFieldRenderMode HeightFieldRenderMode;
 
diff --git a/Engine/Source/Runtime/NavigationSystem/Public/NavMesh/RecastNavMeshGenerator.h b/Engine/Source/Runtime/NavigationSystem/Public/NavMesh/RecastNavMeshGenerator.h
--- a/Engine/Source/Runtime/NavigationSystem/Public/NavMesh/RecastNavMeshGenerator.h
+++ b/Engine/Source/Runtime/NavigationSystem/Public/NavMesh/RecastNavMeshGenerator.h
@@ -488,9 +488,9 @@ protected:
 		const FAreaNavModifier& Modifier, const FTransform& LocalToWorld, const int32 Mask, TInlineMaskArray& OutMaskArray);
 
 	/** apply areas from DynamicAreas to layer */
-	NAVIGATIONSYSTEM_API void MarkDynamicAreas(dtTileCacheLayer& Layer);
-	NAVIGATIONSYSTEM_API void MarkDynamicArea(const FAreaNavModifier& Modifier, const FTransform& LocalToWorld, dtTileCacheLayer& Layer);
-	NAVIGATIONSYSTEM_API void MarkDynamicArea(const FAreaNavModifier& Modifier, const FTransform& LocalToWorld, dtTileCacheLayer& Layer, const int32 AreaId, const int32* ReplaceAreaIdPtr);
+	NAVIGATIONSYSTEM_API void MarkDynamicAreas(FNavMeshBuildContext& BuildContext, dtTileCacheLayer& Layer);
+	NAVIGATIONSYSTEM_API void MarkDynamicArea(FNavMeshBuildContext& BuildContext, const FAreaNavModifier& Modifier, const FTransform& LocalToWorld, dtTileCacheLayer& Layer);
+	NAVIGATIONSYSTEM_API void MarkDynamicArea(FNavMeshBuildContext& BuildContext, const FAreaNavModifier& Modifier, const FTransform& LocalToWorld, dtTileCacheLayer& Layer, const int32 AreaId, const int32* ReplaceAreaIdPtr);
 
 	NAVIGATIONSYSTEM_API void AppendModifier(const FCompositeNavModifier& Modifier, const FNavDataPerInstanceTransformDelegate& InTransformsDelegate);
 	/** Appends specified geometry to tile's geometry */
@@ -546,6 +546,13 @@ protected:
 	NAVIGATIONSYSTEM_API void RasterizeGeometry(FNavMeshBuildContext& BuildContext, const TArray<FVector::FReal>& Coords, const TArray<int32>& Indices, const FTransform& LocalToWorld, const rcRasterizationFlags RasterizationFlags, FTileRasterizationContext& RasterContext);
 	UE_DEPRECATED(5.6, "Please use the overload with the added RecastBounds.")
 	NAVIGATIONSYSTEM_API void PrepareVoxelCache(const TNavStatArray<uint8>& RawCollisionCache, const FCompositeNavModifier& InModifier, TNavStatArray<rcSpanCache>& SpanData);
+	UE_DEPRECATED(5.8, "Please use the overload with the added BuildContext.")
+	NAVIGATIONSYSTEM_API void MarkDynamicAreas(dtTileCacheLayer& Layer);
+	UE_DEPRECATED(5.8, "Please use the overload with the added BuildContext.")
+	NAVIGATIONSYSTEM_API void MarkDynamicArea(const FAreaNavModifier& Modifier, const FTransform& LocalToWorld, dtTileCacheLayer& Layer);
+	UE_DEPRECATED(5.8, "Please use the overload with the added BuildContext.")
+	NAVIGATIONSYSTEM_API void MarkDynamicArea(const FAreaNavModifier& Modifier, const FTransform& LocalToWorld, dtTileCacheLayer& Layer, const int32 AreaId, const int32* ReplaceAreaIdPtr);
+
 
 protected:
 	uint32 bRegenerateCompressedLayers : 1;
diff --git a/Engine/Source/Runtime/Navmesh/Private/DebugUtils/RecastDebugDraw.cpp b/Engine/Source/Runtime/Navmesh/Private/DebugUtils/RecastDebugDraw.cpp
--- a/Engine/Source/Runtime/Navmesh/Private/DebugUtils/RecastDebugDraw.cpp
+++ b/Engine/Source/Runtime/Navmesh/Private/DebugUtils/RecastDebugDraw.cpp
@@ -1093,3 +1093,37 @@ void duDebugDrawPolyMeshDetail(duDebugDraw* dd, const struct rcPolyMeshDetail& d
 	}
 	dd->end();
 }
+
+void duDrawCylinderArea(duDebugDraw* dd, const duReal* pos, const duReal radius, const duReal height, const unsigned int areaRGBAColor)
+{
+	duDebugDrawCylinderWire(dd, pos[0] - radius, pos[1], pos[2] - radius, pos[0] + radius, pos[1] + height, pos[2] + radius, duRGBA(areaRGBAColor), 4.f);
+}
+
+void duDrawBoxArea(duDebugDraw* dd, const duReal* pos, const duReal* extent, const unsigned int areaRGBAColor)
+{
+	duDebugDrawBoxWire(dd, pos[0] - extent[0], pos[1] - extent[1], pos[2] - extent[2], pos[0] + extent[0], pos[1] + extent[1], pos[2] + extent[2], duRGBA(areaRGBAColor), 4.f);
+}
+
+void duDrawConvexArea(duDebugDraw* dd, const duReal* verts, const int nverts, const duReal hmin, const duReal hmax, const unsigned int areaRGBAColor)
+{
+	dd->begin(DU_DRAW_LINES, 4.0f);
+
+	const unsigned int areaColor = duRGBA(areaRGBAColor);
+	for (int i = 0; i < nverts; ++i)
+	{
+		const duReal* vert = &verts[i*3];
+
+		// vertical line
+		dd->vertex(vert[0], hmin, vert[2], areaColor);
+		dd->vertex(vert[0], hmax, vert[2], areaColor);
+
+		// horizontal lines
+		int ni = (i+1)%nverts;
+		const duReal* nextVert = &verts[ni*3];
+		dd->vertex(vert[0], hmin, vert[2], areaColor);
+		dd->vertex(nextVert[0], hmin, nextVert[2], areaColor);
+		dd->vertex(vert[0], hmax, vert[2], areaColor);
+		dd->vertex(nextVert[0], hmax, nextVert[2], areaColor);
+	}
+	dd->end();
+}
diff --git a/Engine/Source/Runtime/Navmesh/Public/DebugUtils/RecastDebugDraw.h b/Engine/Source/Runtime/Navmesh/Public/DebugUtils/RecastDebugDraw.h
--- a/Engine/Source/Runtime/Navmesh/Public/DebugUtils/RecastDebugDraw.h
+++ b/Engine/Source/Runtime/Navmesh/Public/DebugUtils/RecastDebugDraw.h
@@ -39,6 +39,10 @@ NAVMESH_API void duDebugDrawCompactHeightfieldSolid(struct duDebugDraw* dd, cons
 NAVMESH_API void duDebugDrawCompactHeightfieldRegions(struct duDebugDraw* dd, const struct rcCompactHeightfield& chf);
 NAVMESH_API void duDebugDrawCompactHeightfieldDistance(struct duDebugDraw* dd, const struct rcCompactHeightfield& chf);
 
+NAVMESH_API void duDrawCylinderArea(duDebugDraw* dd, const duReal* pos, const duReal radius, const duReal height, const unsigned int areaRGBAColor);
+NAVMESH_API void duDrawBoxArea(duDebugDraw* dd, const duReal* pos, const duReal* extent, const unsigned int areaRGBAColor);
+NAVMESH_API void duDrawConvexArea(duDebugDraw* dd, const duReal* verts, const int nverts, const duReal hmin, const duReal hmax, const unsigned int areaRGBAColor);
+
 NAVMESH_API void duDebugDrawHeightfieldLayer(duDebugDraw* dd, const struct rcHeightfieldLayer& layer, const int idx);
 NAVMESH_API void duDebugDrawHeightfieldLayers(duDebugDraw* dd, const struct rcHeightfieldLayerSet& lset);
 NAVMESH_API void duDebugDrawHeightfieldLayersRegions(duDebugDraw* dd, const struct rcHeightfieldLayerSet& lset);
