Change to make sure, when a PreferredNavData isn't specified on a Pawn, that it doesn't get associated with a NavData that doesn't match its flying capabilities.

This is fixing an issue that was causing walking Pawns to try to use flying NavData to reach their destinations.

[RN] Change to how Pawns without a PreferredNavData set work. This makes it so, if unset, that the Pawn will now only find and use NavData with matching CanFly values. This prevents ground-based Pawns from trying to use flying navigation data.
#Synced-CL 48361187

diff --git a/Engine/Source/Runtime/Engine/Classes/AI/Navigation/NavigationTypes.h b/Engine/Source/Runtime/Engine/Classes/AI/Navigation/NavigationTypes.h
--- a/Engine/Source/Runtime/Engine/Classes/AI/Navigation/NavigationTypes.h
+++ b/Engine/Source/Runtime/Engine/Classes/AI/Navigation/NavigationTypes.h
@@ -404,7 +404,7 @@ struct FMovementProperties
 	UPROPERTY(EditAnywhere, BlueprintReadWrite, Category=MovementProperties)
 	uint8 bCanSwim:1;
 
-	/** If true, this Pawn is capable of flying. */
+	/** If true, this Pawn is capable of flying. By default, Pawn can only use NavData with matching bCanFly values.*/
 	UPROPERTY(EditAnywhere, BlueprintReadWrite, Category=MovementProperties)
 	uint8 bCanFly:1;
 
diff --git a/Engine/Source/Runtime/NavigationSystem/Private/NavigationSystem.cpp b/Engine/Source/Runtime/NavigationSystem/Private/NavigationSystem.cpp
--- a/Engine/Source/Runtime/NavigationSystem/Private/NavigationSystem.cpp
+++ b/Engine/Source/Runtime/NavigationSystem/Private/NavigationSystem.cpp
@@ -2336,6 +2336,14 @@ const ANavigationData* UNavigationSystemV1::GetNavDataForProps(const FNavAgentPr
 				continue;
 			}
 
+			// Make sure to only allow NavData and Agent matches that have matching bCanFly.
+			// This prevents ground-based Pawns from inadvertently getting associated with flying-based NavData when their PreferredNavData isn't set.
+
+			if (NavIt.bCanFly != AgentProperties.bCanFly)
+			{
+				continue;
+			}
+
 			ExcessRadius = NavIt.AgentRadius - AgentProperties.AgentRadius;
 			ExcessHeight = bSkipAgentHeightCheckWhenPickingNavData ? 0.f : (NavIt.AgentHeight - AgentHeight);
 
