[NavigationSystem] fixed deprecation warning for NavLinkJumpDownConfig

#Synced-CL 48614523

diff --git a/Engine/Source/Runtime/NavigationSystem/Private/NavMesh/RecastNavMesh.cpp b/Engine/Source/Runtime/NavigationSystem/Private/NavMesh/RecastNavMesh.cpp
--- a/Engine/Source/Runtime/NavigationSystem/Private/NavMesh/RecastNavMesh.cpp
+++ b/Engine/Source/Runtime/NavigationSystem/Private/NavMesh/RecastNavMesh.cpp
@@ -266,9 +266,11 @@ FVector::FReal ARecastNavMesh::DrawDistanceSq = 0.;
 float ARecastNavMesh::MinimumSizeForChaosNavMeshInfluenceSq = 0.0f;
 #if !WITH_RECAST
 
+PRAGMA_DISABLE_DEPRECATION_WARNINGS
 ARecastNavMesh::ARecastNavMesh(const FObjectInitializer& ObjectInitializer) : Super(ObjectInitializer)
 {
 }
+PRAGMA_ENABLE_DEPRECATION_WARNINGS
 
 void ARecastNavMesh::Serialize( FArchive& Ar )
 {
