Adding some ensures and guards against a perplexing crash.

#Synced-CL 48922636

diff --git a/Engine/Source/Runtime/NavigationSystem/Private/NavMesh/RecastNavMeshGenerator.cpp b/Engine/Source/Runtime/NavigationSystem/Private/NavMesh/RecastNavMeshGenerator.cpp
--- a/Engine/Source/Runtime/NavigationSystem/Private/NavMesh/RecastNavMeshGenerator.cpp
+++ b/Engine/Source/Runtime/NavigationSystem/Private/NavMesh/RecastNavMeshGenerator.cpp
@@ -5589,10 +5589,26 @@ bool FRecastNavMeshGenerator::ConstructTiledNavMesh()
 			else
 			{
 				bSuccess = true;
-				NumActiveTiles = GetTilesCountHelper(DetourMesh);
-				DestNavMesh->GetRecastNavMeshImpl()->SetRecastMesh(DetourMesh);
+				if (!ensureMsgf(DestNavMesh, TEXT("ConstructTiledNavMesh doesn't have a valid DestNavMesh")))
+				{
+					bSuccess = false;
+				}
 
-				UE::NavMesh::Private::CheckTileIndicesInValidRange(InclusionBounds, *DestNavMesh);
+				if (bSuccess)
+				{
+					if (!ensureMsgf(DestNavMesh->GetRecastNavMeshImpl(), TEXT("ConstructTiledNavMesh doesn't have a valid DestNavMesh->GetRecastNavMeshImpl")))
+					{
+						bSuccess = false;
+					}
+				}
+
+				if (bSuccess)
+				{
+					NumActiveTiles = GetTilesCountHelper(DetourMesh);
+					DestNavMesh->GetRecastNavMeshImpl()->SetRecastMesh(DetourMesh);
+
+					UE::NavMesh::Private::CheckTileIndicesInValidRange(InclusionBounds, *DestNavMesh);
+				}
 			}
 		}
 
