Fix for a static analysis warning about an ensure being called too late.

#Synced-CL 48922636

diff --git a/Engine/Source/Runtime/NavigationSystem/Private/NavMesh/RecastNavMeshGenerator.cpp b/Engine/Source/Runtime/NavigationSystem/Private/NavMesh/RecastNavMeshGenerator.cpp
--- a/Engine/Source/Runtime/NavigationSystem/Private/NavMesh/RecastNavMeshGenerator.cpp
+++ b/Engine/Source/Runtime/NavigationSystem/Private/NavMesh/RecastNavMeshGenerator.cpp
@@ -5551,6 +5551,8 @@ bool FRecastNavMeshGenerator::ConstructTiledNavMesh()
 	dtNavMesh* DetourMesh = dtAllocNavMesh();	
 	if (DetourMesh)
 	{
+		ensureMsgf(DestNavMesh, TEXT("ConstructTiledNavMesh doesn't have a valid DestNavMesh"));
+
 		++Version;
 		
 		dtNavMeshParams TiledMeshParameters;
@@ -5589,17 +5591,10 @@ bool FRecastNavMeshGenerator::ConstructTiledNavMesh()
 			else
 			{
 				bSuccess = true;
-				if (!ensureMsgf(DestNavMesh, TEXT("ConstructTiledNavMesh doesn't have a valid DestNavMesh")))
-				{
-					bSuccess = false;
-				}
 
-				if (bSuccess)
+				if (!ensureMsgf(DestNavMesh->GetRecastNavMeshImpl(), TEXT("ConstructTiledNavMesh doesn't have a valid DestNavMesh->GetRecastNavMeshImpl")))
 				{
-					if (!ensureMsgf(DestNavMesh->GetRecastNavMeshImpl(), TEXT("ConstructTiledNavMesh doesn't have a valid DestNavMesh->GetRecastNavMeshImpl")))
-					{
-						bSuccess = false;
-					}
+					bSuccess = false;
 				}
 
 				if (bSuccess)
