Adding batching support for navmesh FindMoveAlongSurface calls

#Synced-CL 49330094

diff --git a/Engine/Source/Runtime/Engine/Classes/AI/Navigation/NavigationTypes.h b/Engine/Source/Runtime/Engine/Classes/AI/Navigation/NavigationTypes.h
--- a/Engine/Source/Runtime/Engine/Classes/AI/Navigation/NavigationTypes.h
+++ b/Engine/Source/Runtime/Engine/Classes/AI/Navigation/NavigationTypes.h
@@ -622,6 +622,25 @@ struct FNavigationRaycastWork : FRayStartEnd
 	{}
 };
 
+struct FNavigationFindMoveAlongSurfaceWork
+{
+	FNavLocation StartLocation;
+
+	/** If the move along the surface was found (ie, bResult is set to TRUE), stores the navmesh surface location and poly where it ends */
+	FNavLocation ResultNavLocation;
+
+	FVector DesiredEndLocation;
+
+	/** If set, stores the result of the underlying find surface move call
+	 *  TRUE if the underlying find surface move call succeeded -- in this case ResultNavLocation will store the navmesh surface location and poly where it ends
+	 *  FALSE if the underlying find surface move call failed */
+	TOptional<bool> bResult;
+
+	FNavigationFindMoveAlongSurfaceWork(const FNavLocation& InStartLocation = FNavLocation(), const FVector& InDesiredEndLocation = FNavigationSystem::InvalidLocation)
+		: StartLocation(InStartLocation), ResultNavLocation(), DesiredEndLocation(InDesiredEndLocation), bResult()
+	{}
+};
+
 UENUM()
 namespace ENavigationQueryResult
 {
diff --git a/Engine/Source/Runtime/NavigationSystem/Private/NavMesh/RecastNavMesh.cpp b/Engine/Source/Runtime/NavigationSystem/Private/NavMesh/RecastNavMesh.cpp
--- a/Engine/Source/Runtime/NavigationSystem/Private/NavMesh/RecastNavMesh.cpp
+++ b/Engine/Source/Runtime/NavigationSystem/Private/NavMesh/RecastNavMesh.cpp
@@ -1908,6 +1908,18 @@ bool ARecastNavMesh::FindMoveAlongSurface(const FNavLocation& StartLocation, con
 	return bSuccess;
 }
 
+void ARecastNavMesh::BatchFindMoveAlongSurface(TArray<FNavigationFindMoveAlongSurfaceWork>& Workload, FSharedConstNavQueryFilter Filter, const UObject* Querier) const
+{
+	if (RecastNavMeshImpl)
+	{
+		const FNavigationQueryFilter& FilterToUse = GetRightFilterRef(Filter);
+		for (FNavigationFindMoveAlongSurfaceWork& Work : Workload)
+		{
+			Work.bResult = RecastNavMeshImpl->FindMoveAlongSurface(Work.StartLocation, Work.DesiredEndLocation, Work.ResultNavLocation, FilterToUse, Querier);
+		}
+	}
+}
+
 bool ARecastNavMesh::FindOverlappingEdges(const FNavLocation& StartLocation, TConstArrayView<FVector> ConvexPolygon, TArray<FVector>& OutEdges, FSharedConstNavQueryFilter Filter, const UObject* Querier) const
 {
 	if (RecastNavMeshImpl == NULL || RecastNavMeshImpl->DetourNavMesh == NULL)
diff --git a/Engine/Source/Runtime/NavigationSystem/Public/AbstractNavData.h b/Engine/Source/Runtime/NavigationSystem/Public/AbstractNavData.h
--- a/Engine/Source/Runtime/NavigationSystem/Public/AbstractNavData.h
+++ b/Engine/Source/Runtime/NavigationSystem/Public/AbstractNavData.h
@@ -75,6 +75,7 @@ public:
 	// Begin ANavigationData overrides
 	virtual void BatchRaycast(TArray<FNavigationRaycastWork>& Workload, FSharedConstNavQueryFilter QueryFilter, const UObject* Querier = NULL) const override {};
 	virtual bool FindMoveAlongSurface(const FNavLocation& StartLocation, const FVector& TargetPosition, FNavLocation& OutLocation, FSharedConstNavQueryFilter Filter = NULL, const UObject* Querier = NULL) const override { return false;  };
+	virtual void BatchFindMoveAlongSurface(TArray<FNavigationFindMoveAlongSurfaceWork>& Workload, FSharedConstNavQueryFilter Filter = NULL, const UObject* Querier = NULL) const override {};
 	virtual bool FindOverlappingEdges(const FNavLocation& StartLocation, TConstArrayView<FVector> ConvexPolygon, TArray<FVector>& OutEdges, FSharedConstNavQueryFilter Filter = NULL, const UObject* Querier = NULL) const override { return false; };
 	virtual bool GetPathSegmentBoundaryEdges(const FNavigationPath& Path, const FNavPathPoint& StartPoint, const FNavPathPoint& EndPoint, const TConstArrayView<FVector> SearchArea, TArray<FVector>& OutEdges, const float MaxAreaEnterCost, FSharedConstNavQueryFilter Filter = NULL, const UObject* Querier = NULL) const override { return false; }
 	virtual FBox GetBounds() const override { return FBox(ForceInit); };
diff --git a/Engine/Source/Runtime/NavigationSystem/Public/NavMesh/RecastNavMesh.h b/Engine/Source/Runtime/NavigationSystem/Public/NavMesh/RecastNavMesh.h
--- a/Engine/Source/Runtime/NavigationSystem/Public/NavMesh/RecastNavMesh.h
+++ b/Engine/Source/Runtime/NavigationSystem/Public/NavMesh/RecastNavMesh.h
@@ -1039,6 +1039,7 @@ public:
 	NAVIGATIONSYSTEM_API virtual bool GetRandomPointInNavigableRadius(const FVector& Origin, float Radius, FNavLocation& OutResult, FSharedConstNavQueryFilter Filter = NULL, const UObject* Querier = NULL) const override;
 
 	NAVIGATIONSYSTEM_API virtual bool FindMoveAlongSurface(const FNavLocation& StartLocation, const FVector& TargetPosition, FNavLocation& OutLocation, FSharedConstNavQueryFilter Filter = NULL, const UObject* Querier = NULL) const override;
+	NAVIGATIONSYSTEM_API virtual void BatchFindMoveAlongSurface(TArray<FNavigationFindMoveAlongSurfaceWork>& Workload, FSharedConstNavQueryFilter Filter = NULL, const UObject* Querier = NULL) const override;
 	NAVIGATIONSYSTEM_API virtual bool FindOverlappingEdges(const FNavLocation& StartLocation, TConstArrayView<FVector> ConvexPolygon, TArray<FVector>& OutEdges, FSharedConstNavQueryFilter Filter = NULL, const UObject* Querier = NULL) const override;
 	NAVIGATIONSYSTEM_API virtual bool GetPathSegmentBoundaryEdges(const FNavigationPath& Path, const FNavPathPoint& StartPoint, const FNavPathPoint& EndPoint, const TConstArrayView<FVector> SearchArea, TArray<FVector>& OutEdges, const float MaxAreaEnterCost, FSharedConstNavQueryFilter Filter = NULL, const UObject* Querier = NULL) const override;
 	NAVIGATIONSYSTEM_API virtual bool ProjectPoint(const FVector& Point, FNavLocation& OutLocation, const FVector& Extent, FSharedConstNavQueryFilter Filter = NULL, const UObject* Querier = NULL) const override;
@@ -1595,6 +1596,7 @@ private:
 	virtual FBox GetBounds() const override { return FBox(); }
 	virtual void BatchRaycast(TArray<FNavigationRaycastWork>& Workload, FSharedConstNavQueryFilter QueryFilter, const UObject* Querier = nullptr) const override {}
 	virtual bool FindMoveAlongSurface(const FNavLocation& StartLocation, const FVector& TargetPosition, FNavLocation& OutLocation, FSharedConstNavQueryFilter Filter = nullptr, const UObject* Querier = nullptr) const override { return false; }
+	virtual void BatchFindMoveAlongSurface(TArray<FNavigationFindMoveAlongSurfaceWork>& Workload, FSharedConstNavQueryFilter Filter = NULL, const UObject* Querier = NULL) const override {}
 	virtual bool FindOverlappingEdges(const FNavLocation& StartLocation, TConstArrayView<FVector> ConvexPolygon, TArray<FVector>& OutEdges, FSharedConstNavQueryFilter Filter = nullptr, const UObject* Querier = nullptr) const override { return false; }
 	virtual bool GetPathSegmentBoundaryEdges(const FNavigationPath& Path, const FNavPathPoint& StartPoint, const FNavPathPoint& EndPoint, const TConstArrayView<FVector> SearchArea, TArray<FVector>& OutEdges, const float MaxAreaEnterCost, FSharedConstNavQueryFilter Filter = nullptr, const UObject* Querier = nullptr) const override { return false; }
 	virtual FNavLocation GetRandomPoint(FSharedConstNavQueryFilter Filter = nullptr, const UObject* Querier = nullptr) const override { return FNavLocation(); }
diff --git a/Engine/Source/Runtime/NavigationSystem/Public/NavigationData.h b/Engine/Source/Runtime/NavigationSystem/Public/NavigationData.h
--- a/Engine/Source/Runtime/NavigationSystem/Public/NavigationData.h
+++ b/Engine/Source/Runtime/NavigationSystem/Public/NavigationData.h
@@ -904,6 +904,11 @@ public:
 	 */
 	virtual bool FindMoveAlongSurface(const FNavLocation& StartLocation, const FVector& TargetPosition, FNavLocation& OutLocation, FSharedConstNavQueryFilter Filter = NULL, const UObject* Querier = NULL) const PURE_VIRTUAL(ANavigationData::FindMoveAlongSurface, return false;);
 
+	/**	Batches FindMoveAlongSurface calls for efficiency.
+	 *	@note Refer to FindMoveAlongSurface function documentation for more information
+	 */
+	virtual void BatchFindMoveAlongSurface(TArray<FNavigationFindMoveAlongSurfaceWork>& Workload, FSharedConstNavQueryFilter Filter = NULL, const UObject* Querier = NULL) const PURE_VIRTUAL(ANavigationData::BatchFindMoveAlongSurface, );
+
 	/**	Returns the navmesh edges that touch the convex polygon. The edges are not clipped by the polygon. 
 	 *	@param StartLocation a location on the navmesh where to start searching.
 	 *	@param ConvexPolygon 2D convex polygon that describes the search area. 
