Cleanup: remove UE 5.5 deprecated navigation code

#Synced-CL 49337999

diff --git a/Engine/Source/Runtime/NavigationSystem/Private/NavigationDataHandler.cpp b/Engine/Source/Runtime/NavigationSystem/Private/NavigationDataHandler.cpp
--- a/Engine/Source/Runtime/NavigationSystem/Private/NavigationDataHandler.cpp
+++ b/Engine/Source/Runtime/NavigationSystem/Private/NavigationDataHandler.cpp
@@ -41,12 +41,6 @@ void FNavigationDataHandler::ConstructNavOctree(const FVector& Origin, const dou
 #endif // !UE_BUILD_SHIPPING
 }
 
-// Deprecated
-void FNavigationDataHandler::RemoveNavOctreeElementId(const FOctreeElementId2& ElementId, const int32 UpdateFlags)
-{
-	RemoveFromNavOctree(ElementId, UpdateFlags);
-}
-
 void FNavigationDataHandler::RemoveFromNavOctree(const FOctreeElementId2& ElementId, const int32 UpdateFlags)
 {
 	if (ensure(OctreeController.IsValidElement(ElementId)))
@@ -71,13 +65,6 @@ void FNavigationDataHandler::RemoveFromNavOctree(const FOctreeElementId2& Elemen
 	}
 }
 
-// Deprecated
-FSetElementId FNavigationDataHandler::RegisterNavOctreeElement(UObject& ElementOwner, INavRelevantInterface& ElementInterface, int32 UpdateFlags)
-{
-	check(false);
-	return RegisterElementWithNavOctree(FNavigationElement::CreateFromNavRelevantInterface(ElementInterface), UpdateFlags);
-}
-
 FSetElementId FNavigationDataHandler::RegisterElementWithNavOctree(const TSharedRef<const FNavigationElement>& ElementRef, const int32 UpdateFlags)
 {
 	SCOPE_CYCLE_COUNTER(STAT_Navigation_RegisterNavOctreeElement);
@@ -245,12 +232,6 @@ void FNavigationDataHandler::AddElementToNavOctree(const FNavigationDirtyElement
 	}
 }
 
-// Deprecated
-bool FNavigationDataHandler::UnregisterNavOctreeElement(UObject& ElementOwner, INavRelevantInterface& ElementInterface, int32 UpdateFlags)
-{
-	return UnregisterElementWithNavOctree(FNavigationElement::CreateFromNavRelevantInterface(ElementInterface), UpdateFlags);
-}
-
 bool FNavigationDataHandler::UnregisterElementWithNavOctree(const TSharedRef<const FNavigationElement>& ElementRef, const int32 UpdateFlags)
 {
 	SCOPE_CYCLE_COUNTER(STAT_Navigation_UnregisterNavOctreeElement);
@@ -320,12 +301,6 @@ bool FNavigationDataHandler::UnregisterElementWithNavOctree(const TSharedRef<con
 	return bUnregistered;
 }
 
-// Deprecated
-void FNavigationDataHandler::UpdateNavOctreeElement(UObject& ElementOwner, INavRelevantInterface& ElementInterface, int32 UpdateFlags)
-{
-	UpdateNavOctreeElement(FNavigationElementHandle(&ElementOwner), FNavigationElement::CreateFromNavRelevantInterface(ElementInterface), UpdateFlags);
-}
-
 void FNavigationDataHandler::UpdateNavOctreeElement(FNavigationElementHandle ElementHandle, const TSharedRef<const FNavigationElement>& UpdatedElement, int32 UpdateFlags)
 {
 	INC_DWORD_STAT(STAT_Navigation_UpdateNavOctree);
@@ -385,15 +360,6 @@ void FNavigationDataHandler::UpdateNavOctreeElement(FNavigationElementHandle Ele
 	UpdateNavOctreeParentChain(UpdatedElement, /*bSkipElementOwnerUpdate=*/ true);
 }
 
-// Deprecated
-void FNavigationDataHandler::UpdateNavOctreeParentChain(UObject& ElementOwner, bool bSkipElementOwnerUpdate)
-{
-	if (const INavRelevantInterface* NavRelevantInterface = Cast<INavRelevantInterface>(&ElementOwner))
-	{
-		UpdateNavOctreeParentChain(FNavigationElement::CreateFromNavRelevantInterface(*NavRelevantInterface), bSkipElementOwnerUpdate);
-	}
-}
-
 void FNavigationDataHandler::UpdateNavOctreeParentChain(const TSharedRef<const FNavigationElement>& Element, const bool bSkipElementOwnerUpdate)
 {
 	constexpr int32 UpdateFlags = FNavigationOctreeController::OctreeUpdate_ParentChain | FNavigationOctreeController::OctreeUpdate_Refresh;
@@ -446,12 +412,6 @@ void FNavigationDataHandler::UpdateNavOctreeParentChain(const TSharedRef<const F
 	}
 }
 
-// Deprecated
-bool FNavigationDataHandler::UpdateNavOctreeElementBounds(UObject& ElementOwner, const FBox& NewBounds, const TConstArrayView<FBox> DirtyAreas)
-{
-	return UpdateNavOctreeElementBounds(FNavigationElementHandle(&ElementOwner), NewBounds, DirtyAreas);
-}
-
 bool FNavigationDataHandler::UpdateNavOctreeElementBounds(const FNavigationElementHandle ElementHandle, const FBox& NewBounds, const TConstArrayView<FBox> DirtyAreas)
 {
 	const FOctreeElementId2* ElementId = OctreeController.GetNavOctreeIdForElement(ElementHandle);
@@ -515,12 +475,6 @@ void FNavigationDataHandler::FindElementsInNavOctree(const FBox& QueryBox, const
 	});
 }
 
-// Deprecated
-bool FNavigationDataHandler::ReplaceAreaInOctreeData(const UObject& Object, TSubclassOf<UNavArea> OldArea, TSubclassOf<UNavArea> NewArea, bool bReplaceChildClasses)
-{
-	return ReplaceAreaInOctreeData(FNavigationElementHandle(&Object), OldArea, NewArea, bReplaceChildClasses);
-}
-
 bool FNavigationDataHandler::ReplaceAreaInOctreeData(
 	const FNavigationElementHandle Element,
 	const TSubclassOf<UNavArea> OldArea,
diff --git a/Engine/Source/Runtime/NavigationSystem/Private/NavigationDirtyAreasController.cpp b/Engine/Source/Runtime/NavigationSystem/Private/NavigationDirtyAreasController.cpp
--- a/Engine/Source/Runtime/NavigationSystem/Private/NavigationDirtyAreasController.cpp
+++ b/Engine/Source/Runtime/NavigationSystem/Private/NavigationDirtyAreasController.cpp
@@ -116,19 +116,6 @@ void FNavigationDirtyAreasController::Tick(const float DeltaSeconds, const TArra
 	}
 }
 
-// Deprecated
-void FNavigationDirtyAreasController::AddArea(const FBox& NewArea, const int32 Flags, const TFunction<UObject*()>& ObjectProviderFunc /*= nullptr*/,
-	const FNavigationDirtyElement* DirtyElement /*= nullptr*/, const FName& DebugReason /*= NAME_None*/)
-{
-	AddAreas({NewArea}, static_cast<ENavigationDirtyFlag>(Flags), /*ElementProviderFunc*/nullptr, DirtyElement, DebugReason);
-}
-
-// Deprecated
-void FNavigationDirtyAreasController::AddAreas(const TConstArrayView<FBox> NewAreas, const int32 Flags, const TFunction<UObject*()>& ObjectProviderFunc, const FNavigationDirtyElement* DirtyElement, const FName& DebugReason)
-{
-	AddAreas(NewAreas, static_cast<ENavigationDirtyFlag>(Flags), /*ElementProviderFunc*/nullptr, DirtyElement, DebugReason);
-}
-
 void FNavigationDirtyAreasController::AddArea(const FBox& NewArea, const ENavigationDirtyFlag Flags, const TFunction<const TSharedPtr<const FNavigationElement>()>& ElementProviderFunc /*= nullptr*/,
 	const FNavigationDirtyElement* DirtyElement /*= nullptr*/, const FName& DebugReason /*= NAME_None*/)
 {
diff --git a/Engine/Source/Runtime/NavigationSystem/Private/NavigationOctree.cpp b/Engine/Source/Runtime/NavigationSystem/Private/NavigationOctree.cpp
--- a/Engine/Source/Runtime/NavigationSystem/Private/NavigationOctree.cpp
+++ b/Engine/Source/Runtime/NavigationSystem/Private/NavigationOctree.cpp
@@ -122,13 +122,6 @@ void FNavigationOctree::DemandLazyDataGathering(FNavigationRelevantData& Element
 	INC_MEMORY_STAT_BY(STAT_Navigation_CollisionTreeMemory, ElementMemoryChange);
 }
 
-// Deprecated
-void FNavigationOctree::DemandChildLazyDataGathering(FNavigationRelevantData& ElementData, INavRelevantInterface& ChildNavInterface)
-{
-	const TSharedRef<const FNavigationElement> TmpElement = FNavigationElement::CreateFromNavRelevantInterface(ChildNavInterface);
-	DemandChildLazyDataGathering(ElementData, TmpElement.Get());
-}
-
 void FNavigationOctree::DemandChildLazyDataGathering(FNavigationRelevantData& ElementData, const FNavigationElement& ChildElement) const
 {
 	LLM_SCOPE_BYTAG(NavigationOctree);
@@ -149,13 +142,6 @@ void FNavigationOctree::SetGatheringNavModifiersTimeLimitWarning(const float Thr
 }
 #endif // !UE_BUILD_SHIPPING
 
-// Deprecated
-bool FNavigationOctree::IsLazyGathering(const INavRelevantInterface& ChildNavInterface) const
-{
-	const TSharedRef<const FNavigationElement> TmpElement = FNavigationElement::CreateFromNavRelevantInterface(ChildNavInterface);
-	return IsLazyGathering(TmpElement.Get());
-}
-
 bool FNavigationOctree::IsLazyGathering(const FNavigationElement& NavigationElement) const
 {
 	const ENavDataGatheringMode GatheringMode = NavigationElement.GetGeometryGatheringMode();
@@ -165,15 +151,6 @@ bool FNavigationOctree::IsLazyGathering(const FNavigationElement& NavigationElem
 	return !bDoInstantGathering;
 }
 
-// Deprecated
-void FNavigationOctree::AddNode(UObject* ElementOb, INavRelevantInterface* NavElement, const FBox& Bounds, FNavigationOctreeElement& Element)
-{
-	if (NavElement)
-	{
-		AddNode(Bounds, Element);
-	}
-}
-
 void FNavigationOctree::AddNode(const FBox& Bounds, FNavigationOctreeElement& OctreeElement)
 {
 	LLM_SCOPE_BYTAG(NavigationOctree);
@@ -278,15 +255,6 @@ void FNavigationOctree::AddNode(const FBox& Bounds, FNavigationOctreeElement& Oc
 	AddElement(OctreeElement);
 }
 
-// Deprecated
-void FNavigationOctree::AppendToNode(const FOctreeElementId2& Id, INavRelevantInterface* NavElement, const FBox& Bounds, FNavigationOctreeElement& Element)
-{
-	if (NavElement)
-	{
-		AppendToNode(Id, FNavigationElement::CreateFromNavRelevantInterface(*NavElement), Bounds, Element);
-	}
-}
-
 void FNavigationOctree::AppendToNode(const FOctreeElementId2& Id, const TSharedRef<const FNavigationElement>& ElementRef, const FBox& Bounds, FNavigationOctreeElement& Element)
 {
 	LLM_SCOPE_BYTAG(NavigationOctree);
diff --git a/Engine/Source/Runtime/NavigationSystem/Private/NavigationOctreeController.cpp b/Engine/Source/Runtime/NavigationSystem/Private/NavigationOctreeController.cpp
--- a/Engine/Source/Runtime/NavigationSystem/Private/NavigationOctreeController.cpp
+++ b/Engine/Source/Runtime/NavigationSystem/Private/NavigationOctreeController.cpp
@@ -43,21 +43,6 @@ bool FNavigationOctreeController::GetNavOctreeElementData(const FNavigationEleme
 	return false;
 }
 
-// Deprecated
-bool FNavigationOctreeController::GetNavOctreeElementData(const UObject& NodeOwner, int32& DirtyFlags, FBox& DirtyBounds)
-{
-	ENavigationDirtyFlag TmpDirtyFlags = ENavigationDirtyFlag::None;
-	const bool bSuccess = GetNavOctreeElementData(FNavigationElementHandle(&NodeOwner), TmpDirtyFlags, DirtyBounds);
-	DirtyFlags = static_cast<int32>(TmpDirtyFlags);
-	return bSuccess;
-}
-
-// Deprecated
-const FNavigationRelevantData* FNavigationOctreeController::GetDataForObject(const UObject& Object) const
-{
-	return GetDataForElement(FNavigationElementHandle(&Object));
-}
-
 const FNavigationRelevantData* FNavigationOctreeController::GetDataForElement(const FNavigationElementHandle Element) const
 {
 	if (const FOctreeElementId2* ElementId = GetNavOctreeIdForElement(Element); IsValidElement(ElementId))
@@ -68,12 +53,6 @@ const FNavigationRelevantData* FNavigationOctreeController::GetDataForElement(co
 	return nullptr;
 }
 
-// Deprecated
-FNavigationRelevantData* FNavigationOctreeController::GetMutableDataForObject(const UObject& Object)
-{
-	return GetMutableDataForElement(FNavigationElementHandle(&Object));
-}
-
 FNavigationRelevantData* FNavigationOctreeController::GetMutableDataForElement(const FNavigationElementHandle Element)
 {
 	if (const FOctreeElementId2* ElementId = GetNavOctreeIdForElement(Element); IsValidElement(ElementId))
@@ -83,39 +62,3 @@ FNavigationRelevantData* FNavigationOctreeController::GetMutableDataForElement(c
 
 	return nullptr;
 }
-
-//----------------------------------------------------------------------//
-// Deprecated methods
-//----------------------------------------------------------------------//
-
-PRAGMA_DISABLE_DEPRECATION_WARNINGS
-
-// Deprecated
-bool FNavigationOctreeController::HasPendingObjectNavOctreeId(UObject& Object) const
-{
-	return HasPendingUpdateForElement(FNavigationElementHandle(&Object));
-}
-
-// Deprecated
-bool FNavigationOctreeController::HasObjectsNavOctreeId(const UObject& Object) const
-{
-	return HasElementNavOctreeId(FNavigationElementHandle(&Object));
-}
-
-// Deprecated
-const FOctreeElementId2* FNavigationOctreeController::GetObjectsNavOctreeId(const UObject& Object) const
-{
-	return GetNavOctreeIdForElement(FNavigationElementHandle(&Object));
-}
-
-// Deprecated
-void FNavigationOctreeController::RemoveObjectsNavOctreeId(const UObject& Object)
-{
-	const FNavigationElementHandle ElementHandle(&Object);
-	if (const FOctreeElementId2* ElementId = GetNavOctreeIdForElement(ElementHandle); IsValidElement(ElementId))
-	{
-		RemoveNode(*ElementId, ElementHandle);
-	}	
-}
-
-PRAGMA_ENABLE_DEPRECATION_WARNINGS
\ No newline at end of file
diff --git a/Engine/Source/Runtime/NavigationSystem/Private/NavigationSystem.cpp b/Engine/Source/Runtime/NavigationSystem/Private/NavigationSystem.cpp
--- a/Engine/Source/Runtime/NavigationSystem/Private/NavigationSystem.cpp
+++ b/Engine/Source/Runtime/NavigationSystem/Private/NavigationSystem.cpp
@@ -3324,25 +3324,6 @@ bool UNavigationSystemV1::ShouldCreateNavigationSystemInstance(const UWorld* Wor
 		&& (ShouldAllowClientSideNavigation() || (World->GetNetMode() != NM_Client));
 }
 
-// Deprecated
-UNavigationSystemV1* UNavigationSystemV1::CreateNavigationSystem(UWorld* WorldOwner)
-{
-	UNavigationSystemV1* NavSys = nullptr;
-
-	// create navigation system for editor and server targets, but remove it from game clients
-	if (GetDefault<UNavigationSystemV1>()->ShouldCreateNavigationSystemInstance(WorldOwner))
-	{
-		AWorldSettings* WorldSettings = WorldOwner->GetWorldSettings();
-		if (WorldSettings == nullptr || WorldSettings->IsNavigationSystemEnabled())
-		{
-			NavSys = NewObject<UNavigationSystemV1>(WorldOwner, GEngine->NavigationSystemClass);
-			WorldOwner->SetNavigationSystem(NavSys);
-		}
-	}
-
-	return NavSys;
-}
-
 void UNavigationSystemV1::InitializeForWorld(UWorld& World, FNavigationSystemRunMode Mode)
 {
 	OnWorldInitDone(Mode);
@@ -3488,24 +3469,6 @@ void UNavigationSystemV1::UnregisterComponentToNavOctree(UActorComponent* Comp)
 	}
 }
 
-// Deprecated
-void UNavigationSystemV1::AddDirtyArea(const FBox& NewArea, int32 Flags, const FName& DebugReason)
-{
-	DefaultDirtyAreasController.AddArea(NewArea, static_cast<ENavigationDirtyFlag>(Flags), /*ElementProvideFunc*/ nullptr, /*DirtyElement*/nullptr, DebugReason);
-}
-
-// Deprecated
-void UNavigationSystemV1::AddDirtyArea(const FBox& NewArea, int32 Flags, const TFunction<UObject*()>&, const FName& DebugReason /*= NAME_None*/)
-{
-	DefaultDirtyAreasController.AddArea(NewArea, static_cast<ENavigationDirtyFlag>(Flags), /*ElementProvideFunc*/ nullptr, /*DirtyElement*/nullptr, DebugReason);
-}
-
-// Deprecated
-void UNavigationSystemV1::AddDirtyAreas(const TArray<FBox>& NewAreas, int32 Flags, const FName& DebugReason /*= NAME_None*/)
-{ 
-	AddDirtyAreas(NewAreas, static_cast<ENavigationDirtyFlag>(Flags), DebugReason);
-}
-
 void UNavigationSystemV1::AddDirtyArea(const FBox& NewArea, const ENavigationDirtyFlag Flags, const FName& DebugReason /*= NAME_None*/)
 {
 	DefaultDirtyAreasController.AddArea(NewArea, Flags, nullptr, nullptr, DebugReason);
@@ -3544,14 +3507,6 @@ FSetElementId UNavigationSystemV1::RegisterNavigationElementWithNavOctree(const
 	return FNavigationDataHandler(DefaultOctreeController, DefaultDirtyAreasController).RegisterElementWithNavOctree(Element, UpdateFlags);
 }
 
-// Deprecated
-FSetElementId UNavigationSystemV1::RegisterNavOctreeElement(UObject* ElementOwner, INavRelevantInterface* ElementInterface, int32 UpdateFlags)
-{
-	return (ElementOwner && ElementInterface) 
-		? FNavigationDataHandler(DefaultOctreeController, DefaultDirtyAreasController).RegisterElementWithNavOctree(FNavigationElement::CreateFromNavRelevantInterface(*ElementInterface), UpdateFlags)
-		: FSetElementId();
-}
-
 void UNavigationSystemV1::AddElementToNavOctree(const FNavigationDirtyElement& DirtyElement)
 {
 	FNavigationDataHandler(DefaultOctreeController, DefaultDirtyAreasController).AddElementToNavOctree(DirtyElement);
@@ -3562,41 +3517,11 @@ bool UNavigationSystemV1::GetNavOctreeElementData(const FNavigationElementHandle
 	return DefaultOctreeController.GetNavOctreeElementData(Element, OutDirtyFlags, OutDirtyBounds);
 }
 
-// Deprecated
-bool UNavigationSystemV1::GetNavOctreeElementData(const UObject& NodeOwner, int32& DirtyFlags, FBox& DirtyBounds)
-{
-	ENavigationDirtyFlag TmpDirtyFlags = ENavigationDirtyFlag::None;
-	const bool bSuccess = GetNavOctreeElementData(FNavigationElementHandle(&NodeOwner), TmpDirtyFlags, DirtyBounds);
-	DirtyFlags = static_cast<int32>(TmpDirtyFlags);
-	return bSuccess;
-}
-
-// Deprecated
-void UNavigationSystemV1::UnregisterNavOctreeElement(UObject* ElementOwner, INavRelevantInterface* ElementInterface, int32 UpdateFlags)
-{
-	if (ElementOwner && ElementInterface)
-	{
-		UnregisterNavRelevantObjectInternal(*ElementOwner);
-	}
-}
-
 void UNavigationSystemV1::UnregisterNavigationElementWithOctree(const TSharedRef<const FNavigationElement>& Element, const int32 UpdateFlags)
 {
 	FNavigationDataHandler(DefaultOctreeController, DefaultDirtyAreasController).UnregisterElementWithNavOctree(Element, UpdateFlags);
 }
 
-// Deprecated
-void UNavigationSystemV1::RemoveObjectsNavOctreeId(const UObject& Object)
-{
-	// doing nothing since we don't want external calls to remove a mapping without properly update the nodes
-}
-
-// Deprecated
-void UNavigationSystemV1::RemoveNavOctreeElementId(const FOctreeElementId2& ElementId, const int32 UpdateFlags)
-{
-	RemoveFromNavOctree(ElementId, UpdateFlags);
-}
-
 void UNavigationSystemV1::RemoveFromNavOctree(const FOctreeElementId2& ElementId, const int32 UpdateFlags)
 {
 	FNavigationDataHandler(DefaultOctreeController, DefaultDirtyAreasController).RemoveFromNavOctree(ElementId, UpdateFlags);
@@ -3607,18 +3532,6 @@ void UNavigationSystemV1::DemandLazyDataGathering(FNavigationRelevantData& Eleme
 	FNavigationDataHandler(DefaultOctreeController, DefaultDirtyAreasController).DemandLazyDataGathering(ElementData);
 }
 
-// Deprecated
-const FNavigationRelevantData* UNavigationSystemV1::GetDataForObject(const UObject& Object) const
-{
-	return GetDataForElement(FNavigationElementHandle(&Object));
-}
-
-// Deprecated
-FNavigationRelevantData* UNavigationSystemV1::GetMutableDataForObject(const UObject& Object)
-{
-	return GetMutableDataForElement(FNavigationElementHandle(&Object));
-}
-
 const FNavigationRelevantData* UNavigationSystemV1::GetDataForElement(const FNavigationElementHandle Element) const
 {
 	return DefaultOctreeController.GetDataForElement(Element);
@@ -3996,26 +3909,6 @@ void UNavigationSystemV1::ClearNavOctreeAll(AActor* Actor)
 	}
 }
 
-// Deprecated
-void UNavigationSystemV1::UpdateNavOctreeElement(UObject* ElementOwner, INavRelevantInterface* ElementInterface, int32 UpdateFlags)
-{
-	if (IsNavigationSystemStatic())
-	{
-		return;
-	}
-
-	if (ElementOwner && ElementInterface)
-	{
-		if (const FNavigationElementHandle Handle = GetNavigationElementHandleForUObject(ElementOwner))
-		{
-			FNavigationDataHandler(DefaultOctreeController, DefaultDirtyAreasController).UpdateNavOctreeElement(
-				Handle,
-				FNavigationElement::CreateFromNavRelevantInterface(*ElementInterface),
-				UpdateFlags);
-		}
-	}
-}
-
 void UNavigationSystemV1::UpdateNavOctreeElement(const FNavigationElementHandle Handle, const TSharedRef<const FNavigationElement>& Element, const int32 UpdateFlags)
 {
 	if (IsNavigationSystemStatic())
@@ -4026,34 +3919,6 @@ void UNavigationSystemV1::UpdateNavOctreeElement(const FNavigationElementHandle
 	FNavigationDataHandler(DefaultOctreeController, DefaultDirtyAreasController).UpdateNavOctreeElement(Handle, Element, UpdateFlags);
 }
 
-// Deprecated
-void UNavigationSystemV1::UpdateNavOctreeParentChain(UObject* ElementOwner, bool bSkipElementOwnerUpdate)
-{
-	PRAGMA_DISABLE_DEPRECATION_WARNINGS
-	if (ElementOwner)
-	{
-		FNavigationDataHandler(DefaultOctreeController, DefaultDirtyAreasController).UpdateNavOctreeParentChain(*ElementOwner, bSkipElementOwnerUpdate);
-	}
-	PRAGMA_ENABLE_DEPRECATION_WARNINGS
-}
-
-// Deprecated
-bool UNavigationSystemV1::UpdateNavOctreeElementBounds(UObject& Object, const FBox& NewBounds, TConstArrayView<FBox> DirtyAreas)
-{
-	if (Repository == nullptr)
-	{
-		return false;
-	}
-
-	const FNavigationElementHandle Handle = Repository->GetNavigationElementHandleForUObject(&Object);
-	if (Handle.IsValid())
-	{
-		return UpdateNavOctreeElementBounds(Handle, NewBounds, DirtyAreas);
-	}
-
-	return false;
-}
-
 bool UNavigationSystemV1::UpdateNavOctreeElementBounds(const FNavigationElementHandle Handle, const FBox& NewBounds, const TConstArrayView<FBox> DirtyAreas)
 {
 	if (IsNavigationSystemStatic())
@@ -4069,22 +3934,6 @@ bool UNavigationSystemV1::UpdateNavOctreeElementBounds(const FNavigationElementH
 	return FNavigationDataHandler(DefaultOctreeController, DefaultDirtyAreasController).UpdateNavOctreeElementBounds(Handle, NewBounds, DirtyAreas);
 }
 
-// Deprecated
-bool UNavigationSystemV1::ReplaceAreaInOctreeData(const UObject& Object, TSubclassOf<UNavArea> OldArea, TSubclassOf<UNavArea> NewArea, bool bReplaceChildClasses)
-{
-	if (Repository == nullptr)
-	{
-		return false;
-	}
-
-	if (const FNavigationElementHandle Handle = Repository->GetNavigationElementHandleForUObject(&Object))
-	{
-		return ReplaceAreaInOctreeData(Handle, OldArea, NewArea, bReplaceChildClasses);
-	}
-
-	return false;
-}
-
 bool UNavigationSystemV1::ReplaceAreaInOctreeData(const FNavigationElementHandle Handle, const TSubclassOf<UNavArea> OldArea, const TSubclassOf<UNavArea> NewArea, const bool bReplaceChildClasses)
 {
 	if (IsNavigationSystemStatic())
@@ -6457,29 +6306,4 @@ void UNavigationSystemModuleConfig::PostEditChangeProperty(FPropertyChangedEvent
 }
 #endif // WITH_EDITOR
 
-//----------------------------------------------------------------------//
-// Deprecated methods
-//----------------------------------------------------------------------//
-PRAGMA_DISABLE_DEPRECATION_WARNINGS
-
-// Deprecated
-uint32 UNavigationSystemV1::HashObject(const UObject& Object)
-{
-	return FNavigationOctree::HashObject(Object);
-}
-
-// Deprecated
-const FOctreeElementId2* UNavigationSystemV1::GetObjectsNavOctreeId(const UObject& Object) const
-{
-	return GetNavOctreeIdForElement(FNavigationElementHandle(&Object));
-}
-
-// Deprecated
-bool UNavigationSystemV1::HasPendingObjectNavOctreeId(UObject* Object) const
-{
-	return HasPendingUpdateForElement(FNavigationElementHandle(Object));
-}
-
-PRAGMA_ENABLE_DEPRECATION_WARNINGS
-
 #undef LOCTEXT_NAMESPACE
diff --git a/Engine/Source/Runtime/NavigationSystem/Public/NavigationDataHandler.h b/Engine/Source/Runtime/NavigationSystem/Public/NavigationDataHandler.h
--- a/Engine/Source/Runtime/NavigationSystem/Public/NavigationDataHandler.h
+++ b/Engine/Source/Runtime/NavigationSystem/Public/NavigationDataHandler.h
@@ -17,8 +17,6 @@ struct FNavigationDataHandler
 
 	NAVIGATIONSYSTEM_API void ConstructNavOctree(const FVector& Origin, const double Radius, const ENavDataGatheringModeConfig DataGatheringMode, const float GatheringNavModifiersWarningLimitTime);
 
-	UE_DEPRECATED(5.5, "Use RegisterFromNavOctree instead.")
-	NAVIGATIONSYSTEM_API void RemoveNavOctreeElementId(const FOctreeElementId2& ElementId, int32 UpdateFlags);
 	/**
 	 * Removes the octree node and the NavigationElementHandle-OctreeElementId pair associated to the specified OctreeElementId.
 	 * It will also dirty the area base of the NavigationElement values and the specified update flags.
@@ -26,8 +24,6 @@ struct FNavigationDataHandler
 	NAVIGATIONSYSTEM_API void RemoveFromNavOctree(const FOctreeElementId2& ElementId, int32 UpdateFlags);
 
 	NAVIGATIONSYSTEM_API FSetElementId RegisterElementWithNavOctree(const TSharedRef<const FNavigationElement>& ElementRef, int32 UpdateFlags);
-	UE_DEPRECATED(5.5, "Use RegisterNavigationElementWithNavOctree instead.")
-	NAVIGATIONSYSTEM_API FSetElementId RegisterNavOctreeElement(UObject& ElementOwner, INavRelevantInterface& ElementInterface, int32 UpdateFlags);
 	NAVIGATIONSYSTEM_API void AddElementToNavOctree(const FNavigationDirtyElement& DirtyElement);
 
 	/**
@@ -39,36 +35,22 @@ struct FNavigationDataHandler
 	 * @return True if associated NavOctreeElement has been removed or pending update has been invalidated; false otherwise.
 	 */
 	NAVIGATIONSYSTEM_API bool UnregisterElementWithNavOctree(const TSharedRef<const FNavigationElement>& ElementRef, int32 UpdateFlags);
-	UE_DEPRECATED(5.5, "Use the overloaded version with FNavigationElement instead.")
-	NAVIGATIONSYSTEM_API bool UnregisterNavOctreeElement(UObject& ElementOwner, INavRelevantInterface& ElementInterface, int32 UpdateFlags);
 
 	/**
 	 * Unregister element associated with the provided handle and register the new element.
 	 * Also update any pending update associated to that element.
 	 */
 	NAVIGATIONSYSTEM_API void UpdateNavOctreeElement(FNavigationElementHandle ElementHandle, const TSharedRef<const FNavigationElement>& UpdatedElement, int32 UpdateFlags);
-	UE_DEPRECATED(5.5, "Use the overloaded version with FNavigationElement instead.")
-	NAVIGATIONSYSTEM_API void UpdateNavOctreeElement(UObject& ElementOwner, INavRelevantInterface& ElementInterface, int32 UpdateFlags);
-
-	UE_DEPRECATED(5.5, "This method is no longer public and should not be called directly.")
-	NAVIGATIONSYSTEM_API void UpdateNavOctreeParentChain(UObject& ElementOwner, bool bSkipElementOwnerUpdate = false);
 
 	NAVIGATIONSYSTEM_API bool UpdateNavOctreeElementBounds(FNavigationElementHandle Element, const FBox& NewBounds, const TConstArrayView<FBox> DirtyAreas);
-	UE_DEPRECATED(5.5, "Use the overloaded version with FNavigationElement instead.")
-	NAVIGATIONSYSTEM_API bool UpdateNavOctreeElementBounds(UObject& Object, const FBox& NewBounds, const TConstArrayView<FBox> DirtyAreas);
 
 	NAVIGATIONSYSTEM_API void FindElementsInNavOctree(const FBox& QueryBox, const FNavigationOctreeFilter& Filter, TArray<FNavigationOctreeElement>& Elements);
 
 	NAVIGATIONSYSTEM_API bool ReplaceAreaInOctreeData(FNavigationElementHandle Element, TSubclassOf<UNavArea> OldArea, TSubclassOf<UNavArea> NewArea, bool bReplaceChildClasses) const;
-	UE_DEPRECATED(5.5, "Use the overloaded version with FNavigationElement instead.")
-	NAVIGATIONSYSTEM_API bool ReplaceAreaInOctreeData(const UObject& Object, TSubclassOf<UNavArea> OldArea, TSubclassOf<UNavArea> NewArea, bool bReplaceChildClasses);
 
 	NAVIGATIONSYSTEM_API void AddLevelCollisionToOctree(ULevel& Level);
 	NAVIGATIONSYSTEM_API void RemoveLevelCollisionFromOctree(ULevel& Level);
 
-	UE_DEPRECATED(5.5, "This method will be removed. Use UNavigationSystemV1 version instead.")
-	NAVIGATIONSYSTEM_API void UpdateActorAndComponentsInNavOctree(AActor& Actor);
-
 	NAVIGATIONSYSTEM_API void ProcessPendingOctreeUpdates();
 
 	NAVIGATIONSYSTEM_API void DemandLazyDataGathering(FNavigationRelevantData& ElementData);
diff --git a/Engine/Source/Runtime/NavigationSystem/Public/NavigationDirtyAreasController.h b/Engine/Source/Runtime/NavigationSystem/Public/NavigationDirtyAreasController.h
--- a/Engine/Source/Runtime/NavigationSystem/Public/NavigationDirtyAreasController.h
+++ b/Engine/Source/Runtime/NavigationSystem/Public/NavigationDirtyAreasController.h
@@ -69,9 +69,6 @@ public:
 	 */
 	NAVIGATIONSYSTEM_API void AddArea(const FBox& NewArea, const ENavigationDirtyFlag Flags, const TFunction<const TSharedPtr<const FNavigationElement>()>& ElementProviderFunc = nullptr,
 		const FNavigationDirtyElement* DirtyElement = nullptr, const FName& DebugReason = NAME_None);
-	UE_DEPRECATED(5.5, "Use the version taking ENavigationDirtyFlag and FNavigationElement instead.")
-	NAVIGATIONSYSTEM_API void AddArea(const FBox& NewArea, const int32 Flags, const TFunction<UObject*()>& ObjectProviderFunc = nullptr,
-		const FNavigationDirtyElement* DirtyElement = nullptr, const FName& DebugReason = NAME_None);
 
 	/** Add non-empty list of dirty areas to the queue based on the provided bounds and flags.
 	 * Bounds must be valid and non-empty otherwise the request will be ignored and a warning reported.
@@ -85,9 +82,6 @@ public:
 	 */
 	NAVIGATIONSYSTEM_API void AddAreas(const TConstArrayView<FBox> NewAreas, const ENavigationDirtyFlag Flags, const TFunction<const TSharedPtr<const FNavigationElement>()>& ElementProviderFunc = nullptr,
 		const FNavigationDirtyElement* DirtyElement = nullptr, const FName& DebugReason = NAME_None);
-	UE_DEPRECATED(5.5, "Use the version taking ENavigationDirtyFlag and FNavigationElement instead.")
-	NAVIGATIONSYSTEM_API void AddAreas(const TConstArrayView<FBox> NewAreas, const int32 Flags, const TFunction<UObject*()>& ObjectProviderFunc = nullptr,
-		const FNavigationDirtyElement* DirtyElement = nullptr, const FName& DebugReason = NAME_None);
 	
 	bool IsDirty() const { return GetNumDirtyAreas() > 0; }
 	int32 GetNumDirtyAreas() const { return DirtyAreas.Num(); }
diff --git a/Engine/Source/Runtime/NavigationSystem/Public/NavigationOctree.h b/Engine/Source/Runtime/NavigationSystem/Public/NavigationOctree.h
--- a/Engine/Source/Runtime/NavigationSystem/Public/NavigationOctree.h
+++ b/Engine/Source/Runtime/NavigationSystem/Public/NavigationOctree.h
@@ -35,20 +35,6 @@ struct FNavigationOctreeElement
 	FBoxSphereBounds Bounds;
 	TSharedRef<FNavigationRelevantData, ESPMode::ThreadSafe> Data;
 
-#if WITH_EDITORONLY_DATA
-	UE_DEPRECATED(5.5, "ID no longer used.")
-	uint32 OwnerUniqueId = INDEX_NONE;
-#endif // WITH_EDITORONLY_DATA
-
-	UE_DEPRECATED(5.5, "Use the constructor taking a FNavigationElement instead.")
-	explicit FNavigationOctreeElement(UObject& SourceObject)
-		: Data(MakeShared<FNavigationRelevantData>(MakeShared<const FNavigationElement>(SourceObject)))
-#if WITH_EDITORONLY_DATA
-		, OwnerUniqueId(SourceObject.GetUniqueID())
-#endif // WITH_EDITORONLY_DATA
-	{
-	}
-
 	explicit FNavigationOctreeElement(const TSharedRef<const FNavigationElement>& SourceObject)
 		: Bounds(SourceObject.Get().GetBounds())
 		, Data(MakeShared<FNavigationRelevantData>(SourceObject))
@@ -98,9 +84,9 @@ struct FNavigationOctreeElement
 	}
 
 	/** 
-	 *	retrieves Modifier, if it doesn't contain any "Meta Navigation Areas". 
-	 *	If it does then retrieves a copy with meta areas substituted with
-	 *	appropriate non-meta areas, depending on NavAgent
+	 *	Retrieves Modifier if it doesn't contain any "Meta Navigation Areas". 
+	 *	If it does, then retrieves a copy with meta areas substituted with
+	 *	appropriate non-meta areas, depending on NavAgent.
 	 */
 	FORCEINLINE FCompositeNavModifier GetModifierForAgent(const FNavAgentProperties* NavAgent = nullptr) const
 	{ 
@@ -127,12 +113,6 @@ struct FNavigationOctreeElement
 		Data->ValidateAndShrink();
 	}
 
-	UE_DEPRECATED(5.5, "Use GetSourceElement instead.")
-	FORCEINLINE UObject* GetOwner(bool bEvenIfPendingKill = false) const
-	{
-		return const_cast<UObject*>(Data.Get().SourceElement.Get().GetWeakUObject().Get());
-	}
-
 	const TSharedRef<const FNavigationElement>& GetSourceElement() const
 	{
 		return Data->SourceElement;
@@ -172,12 +152,6 @@ struct FNavigationOctreeSemantics
 class FNavigationOctree : public TOctree2<FNavigationOctreeElement, FNavigationOctreeSemantics>, public TSharedFromThis<FNavigationOctree, ESPMode::ThreadSafe>
 {
 public:
-	UE_DEPRECATED(5.5, "Use FGeometryExportDelegate.")
-	DECLARE_DELEGATE_TwoParams(FNavRelevantGeometryExportDelegate, INavRelevantInterface&, FNavigationRelevantData&);
-PRAGMA_DISABLE_DEPRECATION_WARNINGS
-	FNavRelevantGeometryExportDelegate NavRelevantGeometryExportDelegate;
-PRAGMA_ENABLE_DEPRECATION_WARNINGS
-
 	DECLARE_DELEGATE_TwoParams(FGeometryExportDelegate, const FNavigationElement&, FNavigationRelevantData&);
     FGeometryExportDelegate GeometryExportDelegate;
 
@@ -221,13 +195,9 @@ PRAGMA_ENABLE_DEPRECATION_WARNINGS
 
 	/** Add new node and fill it with navigation export data */
 	NAVIGATIONSYSTEM_API void AddNode(const FBox& Bounds, FNavigationOctreeElement& OctreeElement);
-	UE_DEPRECATED(5.5, "Use the overloaded version with only FBox and FNavigationOctreeElement instead.")
-	NAVIGATIONSYSTEM_API void AddNode(UObject* ElementOb, INavRelevantInterface* NavElement, const FBox& Bounds, FNavigationOctreeElement& Data);
 
 	/** Append new data to existing node */
 	NAVIGATIONSYSTEM_API void AppendToNode(const FOctreeElementId2& Id, const TSharedRef<const FNavigationElement>& ElementRef, const FBox& Bounds, FNavigationOctreeElement& Data);
-	UE_DEPRECATED(5.5, "Use the overloaded version with FNavigationElement instead.")
-	NAVIGATIONSYSTEM_API void AppendToNode(const FOctreeElementId2& Id, INavRelevantInterface* NavElement, const FBox& Bounds, FNavigationOctreeElement& Data);
 
 	/** Updates element bounds remove/add operation */
 	NAVIGATIONSYSTEM_API void UpdateNode(const FOctreeElementId2& Id, const FBox& NewBounds);
@@ -250,20 +220,10 @@ PRAGMA_ENABLE_DEPRECATION_WARNINGS
 
 	// Lazy data gathering methods
 	NAVIGATIONSYSTEM_API bool IsLazyGathering(const FNavigationElement& NavigationElement) const;
-	UE_DEPRECATED(5.5, "Use the overloaded version with FNavigationElement instead.")
-	NAVIGATIONSYSTEM_API bool IsLazyGathering(const INavRelevantInterface& ChildNavInterface) const;
 
 	NAVIGATIONSYSTEM_API void DemandLazyDataGathering(FNavigationRelevantData& ElementData);
 
 	NAVIGATIONSYSTEM_API void DemandChildLazyDataGathering(FNavigationRelevantData& ElementData, const FNavigationElement& ChildElement) const;
-	UE_DEPRECATED(5.5, "Use the overloaded version with FNavigationElement instead.")
-	NAVIGATIONSYSTEM_API void DemandChildLazyDataGathering(FNavigationRelevantData& ElementData, INavRelevantInterface& ChildNavInterface);
-
-	UE_DEPRECATED(5.5, "This method is no longer used by the navigation system.")
-	FORCEINLINE static uint32 HashObject(const UObject& Object)
-	{
-		return Object.GetUniqueID();
-	}
 
 #if !UE_BUILD_SHIPPING
 	NAVIGATIONSYSTEM_API void SetGatheringNavModifiersTimeLimitWarning(const float Threshold);
@@ -272,14 +232,8 @@ protected:
 	friend struct FNavigationOctreeController;
 	friend struct FNavigationOctreeSemantics;
 
-	UE_DEPRECATED(5.5, "Use the version taking a FNavigationElementHandle instead.")
-	NAVIGATIONSYSTEM_API void SetElementIdImpl(const uint32 OwnerUniqueId, FOctreeElementId2 Id);
 	NAVIGATIONSYSTEM_API void SetElementIdImpl(FNavigationElementHandle ElementHandle, FOctreeElementId2 Id);
 
-#if WITH_EDITORONLY_DATA
-	UE_DEPRECATED(5.5, "Use ElementToOctreeId instead.")
-	TMap<uint32, FOctreeElementId2> ObjectToOctreeId;
-#endif // WITH_EDITORONLY_DATA
 	TMap<FNavigationElementHandle, FOctreeElementId2> ElementToOctreeId;
 	ENavDataGatheringMode DefaultGeometryGatheringMode;
 	uint32 bGatherGeometry : 1;
diff --git a/Engine/Source/Runtime/NavigationSystem/Public/NavigationOctreeController.h b/Engine/Source/Runtime/NavigationSystem/Public/NavigationOctreeController.h
--- a/Engine/Source/Runtime/NavigationSystem/Public/NavigationOctreeController.h
+++ b/Engine/Source/Runtime/NavigationSystem/Public/NavigationOctreeController.h
@@ -34,28 +34,18 @@ struct FNavigationOctreeController
 		OctreeUpdate_ParentChain = 8,					// update child nodes, don't remove anything
 	};
 
-	UE_DEPRECATED(5.5, "Use PendingUpdates instead.")
-	TSet<FNavigationDirtyElement> PendingOctreeUpdates;
 	TSet<FNavigationDirtyElement, FNavigationDirtyElementKeyFunctions> PendingUpdates;
 	TSharedPtr<FNavigationOctree, ESPMode::ThreadSafe> NavOctree;
 
-	UE_DEPRECATED(5.5, "This container is no longer used. Use AddChild/RemoveChild/GetChildren methods instead.")
-	TMultiMap<UObject*, FWeakObjectPtr> OctreeChildNodesMap;
-
 	/** if set, navoctree updates are ignored, use with caution! */
 	uint8 bNavOctreeLock : 1 = false;
 
 	inline void SetNavigationOctreeLock(bool bLock);
 
 	NAVIGATIONSYSTEM_API bool HasPendingUpdateForElement(FNavigationElementHandle Element) const;
-	UE_DEPRECATED(5.5, "Use HasPendingUpdateForElement instead.")
-	NAVIGATIONSYSTEM_API bool HasPendingObjectNavOctreeId(UObject& Object) const;
 
 	inline void RemoveNode(FOctreeElementId2 ElementId, FNavigationElementHandle GetHandle);
 
-	UE_DEPRECATED(5.5, "Use RemoveNode instead.")
-	NAVIGATIONSYSTEM_API void RemoveObjectsNavOctreeId(const UObject& Object);
-
 	NAVIGATIONSYSTEM_API void SetNavigableGeometryStoringMode(FNavigationOctree::ENavGeometryStoringMode NavGeometryMode);
 
 	NAVIGATIONSYSTEM_API void Reset();
@@ -64,24 +54,14 @@ struct FNavigationOctreeController
 	inline FNavigationOctree* GetMutableOctree();
 
 	inline const FOctreeElementId2* GetNavOctreeIdForElement(FNavigationElementHandle Element) const;
-	UE_DEPRECATED(5.5, "Use GetNavOctreeIdForElement instead.")
-	NAVIGATIONSYSTEM_API const FOctreeElementId2* GetObjectsNavOctreeId(const UObject& Object) const;
 
 	NAVIGATIONSYSTEM_API bool GetNavOctreeElementData(FNavigationElementHandle Element, ENavigationDirtyFlag& DirtyFlags, FBox& DirtyBounds);
-	UE_DEPRECATED(5.5, "Use the version taking ENavigationDirtyFlag& and FNavigationElementHandle as parameter instead.")
-	NAVIGATIONSYSTEM_API bool GetNavOctreeElementData(const UObject& NodeOwner, int32& DirtyFlags, FBox& DirtyBounds);
 
 	NAVIGATIONSYSTEM_API const FNavigationRelevantData* GetDataForElement(FNavigationElementHandle Element) const;
-	UE_DEPRECATED(5.5, "Use GetDataForElement instead.")
-	NAVIGATIONSYSTEM_API const FNavigationRelevantData* GetDataForObject(const UObject& Object) const;
 
 	NAVIGATIONSYSTEM_API FNavigationRelevantData* GetMutableDataForElement(FNavigationElementHandle Element);
-	UE_DEPRECATED(5.5, "Use GetMutableDataForElement instead.")
-	NAVIGATIONSYSTEM_API FNavigationRelevantData* GetMutableDataForObject(const UObject& Object);
 
 	inline bool HasElementNavOctreeId(const FNavigationElementHandle Element) const;
-	UE_DEPRECATED(5.5, "Use HasElementNavOctreeId instead.")
-	NAVIGATIONSYSTEM_API bool HasObjectsNavOctreeId(const UObject& Object) const;
 
 	inline bool IsNavigationOctreeLocked() const;
 
@@ -101,9 +81,6 @@ struct FNavigationOctreeController
 	inline void GetChildren(FNavigationElementHandle Parent, TArray<const TSharedRef<const FNavigationElement>>& OutChildren) const;
 
 private:
-	UE_DEPRECATED(5.5, "This method will no be longer be used by the navigation system.")
-	static uint32 HashObject(const UObject& Object);
-
 	/** Map of all elements that are tied to indexed navigation parent */
 	TMultiMap<FNavigationElementHandle, const TSharedRef<const FNavigationElement>> OctreeParentChildNodesMap;
 };
@@ -112,14 +89,6 @@ private:
 // inlines
 //----------------------------------------------------------------------//
 
-// deprecated
-inline uint32 FNavigationOctreeController::HashObject(const UObject& Object)
-{
-PRAGMA_DISABLE_DEPRECATION_WARNINGS
-	return FNavigationOctree::HashObject(Object);
-PRAGMA_ENABLE_DEPRECATION_WARNINGS
-}
-
 inline const FOctreeElementId2* FNavigationOctreeController::GetNavOctreeIdForElement(const FNavigationElementHandle Element) const
 { 
 	return NavOctree.IsValid()
diff --git a/Engine/Source/Runtime/NavigationSystem/Public/NavigationSystem.h b/Engine/Source/Runtime/NavigationSystem/Public/NavigationSystem.h
--- a/Engine/Source/Runtime/NavigationSystem/Public/NavigationSystem.h
+++ b/Engine/Source/Runtime/NavigationSystem/Public/NavigationSystem.h
@@ -962,12 +962,6 @@ public:
 	NAVIGATIONSYSTEM_API void AddDirtyArea(const FBox& NewArea, ENavigationDirtyFlag Flags, const FName& DebugReason = NAME_None);
 	NAVIGATIONSYSTEM_API void AddDirtyArea(const FBox& NewArea, ENavigationDirtyFlag Flags, const TFunction<const TSharedPtr<const FNavigationElement>()>& ElementProviderFunc, const FName& DebugReason = NAME_None);
 	NAVIGATIONSYSTEM_API void AddDirtyAreas(const TArray<FBox>& NewAreas, ENavigationDirtyFlag Flags, const FName& DebugReason = NAME_None);
-	UE_DEPRECATED(5.5, "Use the version taking ENavigationDirtyFlag instead.")
-	NAVIGATIONSYSTEM_API void AddDirtyArea(const FBox& NewArea, int32 Flags, const FName& DebugReason = NAME_None);
-	UE_DEPRECATED(5.5, "Use the version taking ENavigationDirtyFlag instead.")
-	NAVIGATIONSYSTEM_API void AddDirtyArea(const FBox& NewArea, int32 Flags, const TFunction<UObject*()>& ObjectProviderFunc, const FName& DebugReason = NAME_None);
-	UE_DEPRECATED(5.5, "Use the version taking ENavigationDirtyFlag instead.")
-	NAVIGATIONSYSTEM_API void AddDirtyAreas(const TArray<FBox>& NewAreas, int32 Flags, const FName& DebugReason = NAME_None);
 	NAVIGATIONSYSTEM_API bool HasDirtyAreasQueued() const;
 	NAVIGATIONSYSTEM_API int32 GetNumDirtyAreas() const;
 	float GetDirtyAreaWarningSizeThreshold() const { return DirtyAreaWarningSizeThreshold; }
@@ -975,63 +969,35 @@ public:
 	const FNavigationOctree* GetNavOctree() const { return DefaultOctreeController.GetOctree(); }
 	FNavigationOctree* GetMutableNavOctree() { return DefaultOctreeController.GetMutableOctree(); }
 
-	UE_DEPRECATED(5.5, "This method is no longer used by the navigation system.")
-	NAVIGATIONSYSTEM_API static uint32 HashObject(const UObject& Object);
-
 	inline const FOctreeElementId2* GetNavOctreeIdForElement(FNavigationElementHandle Element) const
 	{
 		return DefaultOctreeController.GetNavOctreeIdForElement(Element);
 	}
-	UE_DEPRECATED(5.5, "Use GetNavOctreeIdForElement instead.")
-	NAVIGATIONSYSTEM_API const FOctreeElementId2* GetObjectsNavOctreeId(const UObject& Object) const;
 
 	inline bool HasPendingUpdateForElement(FNavigationElementHandle Element) const
 	{
 		return DefaultOctreeController.HasPendingUpdateForElement(Element);
 	}
-	UE_DEPRECATED(5.5, "Use HasPendingUpdateForElement instead.")
-	NAVIGATIONSYSTEM_API bool HasPendingObjectNavOctreeId(UObject* Object) const;
-
-	UE_DEPRECATED(5.5, "This functionality should not be exposed and will be removed.")
-	NAVIGATIONSYSTEM_API void RemoveObjectsNavOctreeId(const UObject& Object);
-
-	UE_DEPRECATED(5.5, "This functionality should not be exposed publicly and will be removed. Derived classes should use RemoveFromNavOctree instead.")
-	NAVIGATIONSYSTEM_API void RemoveNavOctreeElementId(const FOctreeElementId2& ElementId, int32 UpdateFlags);
 
 protected:
 	NAVIGATIONSYSTEM_API void RemoveFromNavOctree(const FOctreeElementId2& ElementId, int32 UpdateFlags);
 public:
 
 	NAVIGATIONSYSTEM_API const FNavigationRelevantData* GetDataForElement(FNavigationElementHandle Element) const;
-	UE_DEPRECATED(5.5, "Use GetDataForElement instead.")
-	NAVIGATIONSYSTEM_API const FNavigationRelevantData* GetDataForObject(const UObject& Object) const;
 
 	NAVIGATIONSYSTEM_API FNavigationRelevantData* GetMutableDataForElement(FNavigationElementHandle Element);
 
-	UE_DEPRECATED(5.5, "Use GetMutableDataForElement instead.")
-	NAVIGATIONSYSTEM_API FNavigationRelevantData* GetMutableDataForObject(const UObject& Object);
-
 	/** find all elements in navigation octree within given box (intersection) */
 	NAVIGATIONSYSTEM_API void FindElementsInNavOctree(const FBox& QueryBox, const FNavigationOctreeFilter& Filter, TArray<FNavigationOctreeElement>& Elements);
 
 	/** update single element in navoctree */
 	NAVIGATIONSYSTEM_API void UpdateNavOctreeElement(FNavigationElementHandle Handle, const TSharedRef<const FNavigationElement>& Element, int32 UpdateFlags);
-	UE_DEPRECATED(5.5, "Use the version taking FNavigationElement instead.")
-	void UpdateNavOctreeElement(UObject* ElementOwner, INavRelevantInterface* ElementInterface, int32 UpdateFlags);
-
-	UE_DEPRECATED(5.5, "This method will no longer be public and should not be called directly.")
-	NAVIGATIONSYSTEM_API void UpdateNavOctreeParentChain(UObject* ElementOwner, bool bSkipElementOwnerUpdate = false);
 
 	/** update element's bounds in navigation octree and mark only specified area as dirty, doesn't re-export geometry */
 	NAVIGATIONSYSTEM_API bool UpdateNavOctreeElementBounds(FNavigationElementHandle Handle, const FBox& NewBounds, TConstArrayView<FBox> DirtyAreas);
 
-	UE_DEPRECATED(5.5, "Use the version taking FNavigationElement instead.")
-	NAVIGATIONSYSTEM_API bool UpdateNavOctreeElementBounds(UObject& Object, const FBox& NewBounds, TConstArrayView<FBox> DirtyAreas);
-
 	/** fetched element's data from the octree and replaces occurrences of OldArea with NewArea */
 	NAVIGATIONSYSTEM_API bool ReplaceAreaInOctreeData(FNavigationElementHandle Handle, TSubclassOf<UNavArea> OldArea, TSubclassOf<UNavArea> NewArea, bool bReplaceChildClasses = false);
-	UE_DEPRECATED(5.5, "Use the version taking FNavigationElement instead.")
-	NAVIGATIONSYSTEM_API bool ReplaceAreaInOctreeData(const UObject& Object, TSubclassOf<UNavArea> OldArea, TSubclassOf<UNavArea> NewArea, bool bReplaceChildClasses = false);
 
 	//----------------------------------------------------------------------//
 	// Custom navigation links
@@ -1190,9 +1156,6 @@ public:
 
 	inline bool IsSetUpForLazyGeometryExporting() const { return bGenerateNavigationOnlyAroundNavigationInvokers; }
 
-	UE_DEPRECATED(5.5, "This creation flow is deprecation. Use FNavigationSystem::AddNavigationSystemToWorld instead.")
-	static NAVIGATIONSYSTEM_API UNavigationSystemV1* CreateNavigationSystem(UWorld* WorldOwner);
-
 	static NAVIGATIONSYSTEM_API UNavigationSystemV1* GetCurrent(UWorld* World);
 	static NAVIGATIONSYSTEM_API UNavigationSystemV1* GetCurrent(UObject* WorldContextObject);
 	
@@ -1352,17 +1315,11 @@ protected:
 	NAVIGATIONSYSTEM_API void OnNavigationAreaEvent(UClass* AreaClass, ENavAreaEvent::Type Event);
 
 	NAVIGATIONSYSTEM_API FSetElementId RegisterNavigationElementWithNavOctree(const TSharedRef<const FNavigationElement>& Element, int32 UpdateFlags);
-	UE_DEPRECATED(5.5, "Use the version taking FNavigationElement as parameter instead.")
-	NAVIGATIONSYSTEM_API FSetElementId RegisterNavOctreeElement(UObject* ElementOwner, INavRelevantInterface* ElementInterface, int32 UpdateFlags);
 
 	NAVIGATIONSYSTEM_API void UnregisterNavigationElementWithOctree(const TSharedRef<const FNavigationElement>& Element, int32 UpdateFlags);
-	UE_DEPRECATED(5.5, "Use the version taking FNavigationElement as parameter instead.")
-	NAVIGATIONSYSTEM_API void UnregisterNavOctreeElement(UObject* ElementOwner, INavRelevantInterface* ElementInterface, int32 UpdateFlags);
 
 	/** read element data from navigation octree */
 	NAVIGATIONSYSTEM_API bool GetNavOctreeElementData(FNavigationElementHandle Element, ENavigationDirtyFlag& OutDirtyFlags, FBox& OutDirtyBounds);
-	UE_DEPRECATED(5.5, "Use the version taking FNavigationElement instead.")
-	NAVIGATIONSYSTEM_API bool GetNavOctreeElementData(const UObject& NodeOwner, int32& DirtyFlags, FBox& DirtyBounds);
 
 	/** Adds given element to NavOctree. No check for owner's validity are performed, 
 	 *	nor its presence in NavOctree - function assumes callee responsibility 
