[Navigation]
- Small change to make it so the AllowLinkGeneration cvar is enforced more uniformly.

#Synced-CL 49217423

diff --git a/Engine/Source/Runtime/NavigationSystem/Private/NavMesh/RecastNavMesh.cpp b/Engine/Source/Runtime/NavigationSystem/Private/NavMesh/RecastNavMesh.cpp
--- a/Engine/Source/Runtime/NavigationSystem/Private/NavMesh/RecastNavMesh.cpp
+++ b/Engine/Source/Runtime/NavigationSystem/Private/NavMesh/RecastNavMesh.cpp
@@ -837,13 +837,18 @@ void ARecastNavMesh::PostRegisterAllComponents()
 	}
 
 	// Create and register link proxy if enabled.
-	for (FNavLinkGenerationJumpConfig& NavLinkJumpConfig : NavLinkJumpConfigs)
+
+	const FRecastNavMeshGenerator* MyGenerator = static_cast<FRecastNavMeshGenerator*>(GetGenerator());
+	if (MyGenerator && MyGenerator->IsGeneratingLinks())
 	{
-		if (bGenerateNavLinks && NavLinkJumpConfig.bEnabled && NavLinkJumpConfig.LinkProxyClass.Get() != nullptr)
+		for (FNavLinkGenerationJumpConfig& NavLinkJumpConfig : NavLinkJumpConfigs)
 		{
-			// If it's the first time, LinkProxyId will have the default invalid value so we need to generate an Id.
-			const FNavLinkId Id = NavLinkJumpConfig.LinkProxyId.IsValid() ? NavLinkJumpConfig.LinkProxyId :  FNavLinkId::GenerateUniqueId(); 
-			CreateAndRegisterJumpLinksProxy(NavLinkJumpConfig, Id);
+			if (NavLinkJumpConfig.bEnabled && NavLinkJumpConfig.LinkProxyClass.Get() != nullptr)
+			{
+				// If it's the first time, LinkProxyId will have the default invalid value so we need to generate an Id.
+				const FNavLinkId Id = NavLinkJumpConfig.LinkProxyId.IsValid() ? NavLinkJumpConfig.LinkProxyId : FNavLinkId::GenerateUniqueId();
+				CreateAndRegisterJumpLinksProxy(NavLinkJumpConfig, Id);
+			}
 		}
 	}
 }
diff --git a/Engine/Source/Runtime/NavigationSystem/Private/NavMesh/RecastNavMeshGenerator.cpp b/Engine/Source/Runtime/NavigationSystem/Private/NavMesh/RecastNavMeshGenerator.cpp
--- a/Engine/Source/Runtime/NavigationSystem/Private/NavMesh/RecastNavMeshGenerator.cpp
+++ b/Engine/Source/Runtime/NavigationSystem/Private/NavMesh/RecastNavMeshGenerator.cpp
@@ -5305,7 +5305,7 @@ void FRecastNavMeshGenerator::ConfigureBuildProperties(FRecastBuildConfig& OutCo
 	OutConfig.AgentMaxClimb = AgentMaxClimb;
 	OutConfig.AgentRadius = AgentRadius;
 
-	UE::NavMesh::Private::ComputeConfigBorderSizes(DestNavMesh->bGenerateNavLinks, OutConfig);
+	UE::NavMesh::Private::ComputeConfigBorderSizes(IsGeneratingLinks(), OutConfig);
 
 	OutConfig.maxEdgeLen = (int32)(1200.0f / CellSize);
 
