[Navigation] Wait for initial async loading to be completed before allowing to rebuild navigation automatically
Added a property bWaitForAsyncLoadingBeforeBuildingNavigationAutomatically to make this behavior optional (enabled by default)
#Synced-CL 49590991

diff --git a/Engine/Source/Runtime/NavigationSystem/Private/NavigationSystem.cpp b/Engine/Source/Runtime/NavigationSystem/Private/NavigationSystem.cpp
--- a/Engine/Source/Runtime/NavigationSystem/Private/NavigationSystem.cpp
+++ b/Engine/Source/Runtime/NavigationSystem/Private/NavigationSystem.cpp
@@ -49,6 +49,7 @@
 #endif // WITH_RECAST
 
 #if WITH_EDITOR
+#include "AssetCompilingManager.h"
 #include "EditorModeManager.h"
 #include "EditorModes.h"
 #include "LevelEditor.h"
@@ -1050,6 +1051,18 @@ void UNavigationSystemV1::DoInitialSetup()
 
 	RegisterToRepositoryDelegates();
 
+#if WITH_EDITOR		
+	if (bWaitForAsyncLoadingBeforeBuildingNavigationAutomatically && GetIsAutoUpdateEnabled() && OperationMode == FNavigationSystemRunMode::EditorMode)
+	{
+		AddNavigationBuildLock(ENavigationBuildLock::AsyncLoadLock);
+
+		FAssetCompilingManager::Get().OnAssetPostCompileEvent().AddUObject(this, &UNavigationSystemV1::UpdateDelayedUnlockeRegistration);
+
+		// Run it at least once to make sure we don't stay locked.
+		UpdateDelayedUnlockeRegistration({});
+	}
+#endif // WITH_EDITOR	
+
 	bInitialSetupHasBeenPerformed = true;
 }
 
@@ -5032,6 +5045,9 @@ void UNavigationSystemV1::CleanUp(FNavigationSystem::ECleanupMode Mode)
 	{
 		GEngine->OnActorMoved().RemoveAll(this);
 	}
+
+	FAssetCompilingManager::Get().OnAssetPostCompileEvent().RemoveAll(this);
+	UnregisterDelayedUnlock();
 #endif // WITH_EDITOR
 
 	FCoreUObjectDelegates::PostLoadMapWithWorld.RemoveAll(this);
@@ -5717,6 +5733,63 @@ void UNavigationSystemV1::UnregisterFromRepositoryDelegates() const
 	Repository->OnNavigationElementRemovedDelegate = nullptr;
 }
 
+#if WITH_EDITOR
+void UNavigationSystemV1::UpdateDelayedUnlockeRegistration(const TArray<FAssetCompileData>& AssetCompileData)
+{
+	// Register a delayed unlock navigation when all assets are finished compiling.
+	if (FAssetCompilingManager::Get().GetNumRemainingAssets() == 0)
+	{
+		// Wait until we get some idle time to avoid unlocking too soon
+		constexpr float TimeDelay = 2.f;
+		RegisterDelayedUnlock(TimeDelay);
+	}
+	else
+	{
+		// We're not idle anymore, unregister unlocink until we get to 0 assets again
+		UnregisterDelayedUnlock();
+	}
+}
+
+void UNavigationSystemV1::RegisterDelayedUnlock(float DelayInSeconds)
+{
+	// Unregister the previous one and register a new one.
+	UnregisterDelayedUnlock();
+
+	// The editor might be unresponsive during heavy asset compilation so we 
+	// not only need a delay, but also a minimum amount of frames
+	// to pass until we call the actual refresh.
+	DelayedUnlockHandle = FTSTicker::GetCoreTicker().AddTicker(
+		TEXT("UNavigationSystemV1 async loading wait for unlock"),
+		0.f,
+		[this, FireInTickCount = 16, DelayInSeconds](float DeltaTime) mutable
+		{
+			DelayInSeconds -= DeltaTime;
+			if (--FireInTickCount <= 0 && DelayInSeconds <= 0.f && FAssetCompilingManager::Get().GetNumRemainingAssets() == 0)
+			{
+				FAssetCompilingManager::Get().OnAssetPostCompileEvent().RemoveAll(this);
+				
+				// Cleared initial async load dirty areas and remove AsyncLoadLock
+				DefaultDirtyAreasController.Reset();
+				RemoveNavigationBuildLock(ENavigationBuildLock::AsyncLoadLock, ELockRemovalRebuildAction::NoRebuild);
+				
+				return false; // stop ticker
+			}
+
+			return true;
+		}
+	);
+}
+
+void UNavigationSystemV1::UnregisterDelayedUnlock()
+{
+	if (DelayedUnlockHandle.IsValid())
+	{
+		FTSTicker::RemoveTicker(DelayedUnlockHandle);
+		DelayedUnlockHandle.Reset();
+	}
+}
+#endif // WITH_EDITOR
+
 void UNavigationSystemV1::UpdateInvokers()
 {
 	QUICK_SCOPE_CYCLE_COUNTER(STAT_Navigation_UpdateInvokers);
diff --git a/Engine/Source/Runtime/NavigationSystem/Public/NavigationSystem.h b/Engine/Source/Runtime/NavigationSystem/Public/NavigationSystem.h
--- a/Engine/Source/Runtime/NavigationSystem/Public/NavigationSystem.h
+++ b/Engine/Source/Runtime/NavigationSystem/Public/NavigationSystem.h
@@ -47,6 +47,7 @@ struct FDebugText;
 struct FNavigationRelevantData;
 struct FNavigationOctreeElement;
 struct FNavigationElement;
+struct FAssetCompileData;
 
 #if !UE_BUILD_SHIPPING
 #define ALLOW_TIME_SLICE_DEBUG 1
@@ -112,6 +113,7 @@ namespace ENavigationBuildLock
 		NoUpdateInPIE = 1 << 2,			// PIE doesn't allow automatic updates
 		InitialLock = 1 << 3,			// initial lock, release manually after levels are ready for rebuild (e.g. streaming)
 		Custom = 1 << 4,
+		AsyncLoadLock = 1 << 5			// used in editor worlds to prevent automatic update when loading a map using async loading 
 	};
 }
 
@@ -351,6 +353,12 @@ protected:
 	UPROPERTY()
 	uint32 bDirtyNavDataPackageOnNavigationGenerationFinished : 1;
 
+#if WITH_EDITORONLY_DATA	
+	/** In editor worlds, wait for async loading to finish before allowing to build navigation automatically (true by default). */
+	UPROPERTY(config, EditAnywhere, Category=NavigationSystem)
+	uint32 bWaitForAsyncLoadingBeforeBuildingNavigationAutomatically : 1 = true;
+#endif	
+
 public:
 	/** if set to true will result navigation system not rebuild navigation until 
 	 *	a call to ReleaseInitialBuildingLock() is called. Does not influence 
@@ -482,6 +490,14 @@ private:
 	/** Unregisters from the navigation objects repository subsystem delegates */
 	void UnregisterFromRepositoryDelegates() const;
 
+#if WITH_EDITOR	
+	void UpdateDelayedUnlockeRegistration(const TArray<FAssetCompileData>& AssetCompileData);
+	void RegisterDelayedUnlock(float DelayInSeconds);
+	void UnregisterDelayedUnlock();
+
+	FTSTicker::FDelegateHandle DelayedUnlockHandle;
+#endif // WITH_EDITOR
+
 public:
 	//----------------------------------------------------------------------//
 	// Blueprint functions
