Fix navmesh loading quantization factors not matching with newly updated CellSize settings.
This can happen if settings are changed out of the editor through the config files.
#Synced-CL 49590991

diff --git a/Engine/Source/Runtime/NavigationSystem/Private/NavMesh/PImplRecastNavMesh.cpp b/Engine/Source/Runtime/NavigationSystem/Private/NavMesh/PImplRecastNavMesh.cpp
--- a/Engine/Source/Runtime/NavigationSystem/Private/NavMesh/PImplRecastNavMesh.cpp
+++ b/Engine/Source/Runtime/NavigationSystem/Private/NavMesh/PImplRecastNavMesh.cpp
@@ -563,6 +563,7 @@ void FPImplRecastNavMesh::Serialize( FArchive& Ar, int32 NavMeshVersion )
 	Ar << Params.tileWidth;				///< The width of each tile. (Along the x-axis.)
 	Ar << Params.tileHeight;			///< The height of each tile. (Along the z-axis.)
 	Ar << Params.maxTiles;				///< The maximum number of tiles the navigation mesh can contain.
+	UE_LOG(LogNavigation, VeryVerbose, TEXT("%hs maxTiles: %i"), __FUNCTION__, Params.maxTiles);
 	Ar << Params.maxPolys;
 
 	if (NavMeshOwner->NavMeshVersion >= NAVMESHVER_TILE_RESOLUTIONS)
@@ -608,6 +609,30 @@ void FPImplRecastNavMesh::Serialize( FArchive& Ar, int32 NavMeshVersion )
 		
 		float DefaultCellSize = NavMeshOwner->GetCellSize(ENavigationDataResolution::Default);
 
+		// Detect mismatch with quantization factor
+		bool bQuantizationFactorMismatch = false;
+		if (NavMeshOwner->NavMeshVersion >= NAVMESHVER_TILE_RESOLUTIONS)
+		{
+			for (uint8 Index = 0; Index < (uint8)ENavigationDataResolution::MAX; Index++)
+			{
+				const double expectedQuantFactor =  (NavMeshOwner->NavMeshResolutionParams[Index].CellSize != 0.f) ? (1.f / NavMeshOwner->NavMeshResolutionParams[Index].CellSize) : 0.f;
+				if (!FMath::IsNearlyEqual(Params.resolutionParams[(uint8)Index].bvQuantFactor, expectedQuantFactor))
+				{
+					UE_LOG(LogNavigation, Warning, TEXT("%hs Navmesh %s will be loaded empty because it has been built with different navmesh settings than the current settings. "
+						"If possible, this will trigger a navmesh rebuild. "
+						"The mismatch was detected on stored quantization factor %f (expected: %f) for CellSize %f (resolution: %i)."),
+						__FUNCTION__,
+						*NavMeshOwner->GetName(),
+						Params.resolutionParams[(uint8)Index].bvQuantFactor,
+						expectedQuantFactor,
+						NavMeshOwner->NavMeshResolutionParams[Index].CellSize,
+						Index);
+				
+					bQuantizationFactorMismatch = true;
+				}
+			}
+		}
+
 #if WITH_EDITORONLY_DATA
 		PRAGMA_DISABLE_DEPRECATION_WARNINGS
 		if (NavMeshVersion < NAVMESHVER_TILE_RESOLUTIONS)
@@ -621,8 +646,9 @@ void FPImplRecastNavMesh::Serialize( FArchive& Ar, int32 NavMeshVersion )
 #endif // WITH_EDITORONLY_DATA
 
 		const FVector::FReal ActorsTileSize = static_cast<FVector::FReal>(FMath::TruncToInt(NavMeshOwner->TileSizeUU / DefaultCellSize)) * DefaultCellSize;
+		const bool bTileSizeMismatch = ActorsTileSize != Params.tileWidth;
 
-		if (ActorsTileSize != Params.tileWidth)
+		if (bTileSizeMismatch || bQuantizationFactorMismatch)
 		{
 			// just move archive position
 			ReleaseDetourNavMesh();
