Fix navmesh issue with unloaded objects in a WP maps when Build Paths is run.
#Synced-CL 49693995

diff --git a/Engine/Source/Runtime/NavigationSystem/Private/NavigationData.cpp b/Engine/Source/Runtime/NavigationSystem/Private/NavigationData.cpp
--- a/Engine/Source/Runtime/NavigationSystem/Private/NavigationData.cpp
+++ b/Engine/Source/Runtime/NavigationSystem/Private/NavigationData.cpp
@@ -589,8 +589,20 @@ void ANavigationData::ConditionalConstructGenerator()
 void ANavigationData::RebuildAll()
 {
 	const double LoadTime = FPlatformTime::Seconds();
+
 	LoadBeforeGeneratorRebuild();
 	FAssetCompilingManager::Get().FinishAllCompilation();
+	
+	if (UWorld* World = GetWorld())
+	{
+		UNavigationSystemV1* NavSys = FNavigationSystem::GetCurrent<UNavigationSystemV1>(World);
+		if (NavSys)
+		{
+			// This used to be in UNavigationSystemV1::RebuildAll but since we just loaded new assets it needs to be done here.
+			NavSys->FlushPendingOperations();
+		}
+	}
+	
 	UE_LOG(LogNavigationDataBuild, Display, TEXT("   %s load time: %.2fs"), ANSI_TO_TCHAR(__FUNCTION__), (FPlatformTime::Seconds() - LoadTime));
 
 	PostLoadPreRebuild();
diff --git a/Engine/Source/Runtime/NavigationSystem/Private/NavigationSystem.cpp b/Engine/Source/Runtime/NavigationSystem/Private/NavigationSystem.cpp
--- a/Engine/Source/Runtime/NavigationSystem/Private/NavigationSystem.cpp
+++ b/Engine/Source/Runtime/NavigationSystem/Private/NavigationSystem.cpp
@@ -1260,8 +1260,7 @@ bool UNavigationSystemV1::ConditionalPopulateNavOctree()
 	}
 
 	// Add all found elements to octree, this will not add new dirty areas to navigation
-	FNavigationDataHandler NavHandler(DefaultOctreeController, DefaultDirtyAreasController);
-	NavHandler.ProcessPendingOctreeUpdates();
+	ProcessPendingOctreeUpdates();
 
 	return bSupportRebuilding;
 }
@@ -3510,6 +3509,11 @@ int32 UNavigationSystemV1::GetNumDirtyAreas() const
 	return DefaultDirtyAreasController.GetNumDirtyAreas();
 }
 
+void UNavigationSystemV1::ResetDefaultDirtyAreasController()
+{
+	DefaultDirtyAreasController.Reset();
+}
+
 bool UNavigationSystemV1::HasDirtyAreasQueued() const
 {
 	return DefaultDirtyAreasController.IsDirty();
@@ -3540,6 +3544,22 @@ void UNavigationSystemV1::RemoveFromNavOctree(const FOctreeElementId2& ElementId
 	FNavigationDataHandler(DefaultOctreeController, DefaultDirtyAreasController).RemoveFromNavOctree(ElementId, UpdateFlags);
 }
 
+void UNavigationSystemV1::ProcessPendingOctreeUpdates()
+{
+	FNavigationDataHandler NavHandler(DefaultOctreeController, DefaultDirtyAreasController);
+	NavHandler.ProcessPendingOctreeUpdates();
+}
+
+void UNavigationSystemV1::FlushPendingOperations()
+{
+	// Make sure all changes are reflected in the octree
+	ProcessPendingOctreeUpdates();
+	
+	PendingNavBoundsUpdates.Reset();
+
+	ResetDefaultDirtyAreasController();
+}
+
 void UNavigationSystemV1::DemandLazyDataGathering(FNavigationRelevantData& ElementData)
 {
 	FNavigationDataHandler(DefaultOctreeController, DefaultDirtyAreasController).DemandLazyDataGathering(ElementData);
@@ -4745,14 +4765,6 @@ void UNavigationSystemV1::RebuildAll(bool bIsLoadTime)
 	
 	GatherNavigationBounds();
 
-	// make sure that octree is up to date
-	FNavigationDataHandler NavHandler(DefaultOctreeController, DefaultDirtyAreasController);
-	NavHandler.ProcessPendingOctreeUpdates();
-	
-	PendingNavBoundsUpdates.Reset();
-
-	DefaultDirtyAreasController.Reset();
-
 	for (int32 NavDataIndex = 0; NavDataIndex < NavDataSet.Num(); ++NavDataIndex)
 	{
 		ANavigationData* NavData = NavDataSet[NavDataIndex];
diff --git a/Engine/Source/Runtime/NavigationSystem/Public/NavigationSystem.h b/Engine/Source/Runtime/NavigationSystem/Public/NavigationSystem.h
--- a/Engine/Source/Runtime/NavigationSystem/Public/NavigationSystem.h
+++ b/Engine/Source/Runtime/NavigationSystem/Public/NavigationSystem.h
@@ -995,10 +995,17 @@ public:
 		return DefaultOctreeController.HasPendingUpdateForElement(Element);
 	}
 
+	NAVIGATIONSYSTEM_API void FlushPendingOperations();
+
 protected:
+	NAVIGATIONSYSTEM_API void ResetDefaultDirtyAreasController();
+
+	/** Update any pending FNavigationDirtyElement in the navigation system default navigation octree. */
+	NAVIGATIONSYSTEM_API void ProcessPendingOctreeUpdates();
+
 	NAVIGATIONSYSTEM_API void RemoveFromNavOctree(const FOctreeElementId2& ElementId, int32 UpdateFlags);
-public:
 
+public:
 	NAVIGATIONSYSTEM_API const FNavigationRelevantData* GetDataForElement(FNavigationElementHandle Element) const;
 
 	NAVIGATIONSYSTEM_API FNavigationRelevantData* GetMutableDataForElement(FNavigationElementHandle Element);
