Remove deprecated navigation code
#Synced-CL 49739737

diff --git a/Engine/Source/Runtime/NavigationSystem/Private/NavMesh/PImplRecastNavMesh.cpp b/Engine/Source/Runtime/NavigationSystem/Private/NavMesh/PImplRecastNavMesh.cpp
--- a/Engine/Source/Runtime/NavigationSystem/Private/NavMesh/PImplRecastNavMesh.cpp
+++ b/Engine/Source/Runtime/NavigationSystem/Private/NavMesh/PImplRecastNavMesh.cpp
@@ -3339,56 +3339,6 @@ FBox FPImplRecastNavMesh::GetNavMeshBounds() const
 	return Bbox;
 }
 
-// Deprecated
-FBox FPImplRecastNavMesh::GetNavMeshTileBounds(int32 TileIndex) const
-{
-	FBox Bbox(ForceInit);
-
-	if (DetourNavMesh && TileIndex >= 0 && TileIndex < DetourNavMesh->getMaxTiles())
-	{
-		// workaround for privacy issue in the recast API
-		dtNavMesh const* const ConstRecastNavMesh = DetourNavMesh;
-
-		dtMeshTile const* const Tile = ConstRecastNavMesh->getTile(TileIndex);
-		if (Tile)
-		{
-			dtMeshHeader const* const Header = Tile->header;
-			if (Header)
-			{
-				Bbox = Recast2UnrealBox(Header->bmin, Header->bmax);
-			}
-		}
-	}
-
-	return Bbox;
-}
-
-// Deprecated
-/** Retrieves XY coordinates of tile specified by index */
-bool FPImplRecastNavMesh::GetNavMeshTileXY(int32 TileIndex, int32& OutX, int32& OutY, int32& OutLayer) const
-{
-	if (DetourNavMesh && TileIndex >= 0 && TileIndex < DetourNavMesh->getMaxTiles())
-	{
-		// workaround for privacy issue in the recast API
-		dtNavMesh const* const ConstRecastNavMesh = DetourNavMesh;
-
-		dtMeshTile const* const Tile = ConstRecastNavMesh->getTile(TileIndex);
-		if (Tile)
-		{
-			dtMeshHeader const* const Header = Tile->header;
-			if (Header)
-			{
-				OutX = Header->x;
-				OutY = Header->y;
-				OutLayer = Header->layer;
-				return true;
-			}
-		}
-	}
-
-	return false;
-}
-
 bool FPImplRecastNavMesh::GetNavMeshTileXY(const FVector& Point, int32& OutX, int32& OutY) const
 {
 	if (DetourNavMesh)
@@ -3409,42 +3359,6 @@ bool FPImplRecastNavMesh::GetNavMeshTileXY(const FVector& Point, int32& OutX, in
 	return false;
 }
 
-// Deprecated
-void FPImplRecastNavMesh::GetNavMeshTilesAt(int32 TileX, int32 TileY, TArray<int32>& Indices) const
-{
-	if (DetourNavMesh)
-	{
-		// workaround for privacy issue in the recast API
-		dtNavMesh const* const ConstRecastNavMesh = DetourNavMesh;
-
-		const int32 MaxTiles = ConstRecastNavMesh->getTileCountAt(TileX, TileY);
-		TArray<const dtMeshTile*> Tiles;
-		Tiles.AddZeroed(MaxTiles);
-
-		const int32 NumTiles = ConstRecastNavMesh->getTilesAt(TileX, TileY, Tiles.GetData(), MaxTiles);
-		for (int32 i = 0; i < NumTiles; i++)
-		{
-			dtTileRef TileRef = ConstRecastNavMesh->getTileRef(Tiles[i]);
-			if (TileRef)
-			{
-				const int32 TileIndex = (int32)ConstRecastNavMesh->decodePolyIdTile(TileRef);
-				Indices.Add(TileIndex);
-			}
-		}
-	}
-}
-
-// Deprecated
-void FPImplRecastNavMesh::GetNavMeshTilesIn(const TArray<FBox>& InclusionBounds, TArray<int32>& Indices) const
-{
-	TArray<FNavTileRef> Refs;
-	GetNavMeshTilesIn(InclusionBounds, Refs);
-
-	TArray<uint32> UnsignedIndices;
-	FNavTileRef::DeprecatedGetTileIdsFromNavTileRefs(this, Refs, UnsignedIndices);
-	Indices.Append(UnsignedIndices);
-}
-
 void FPImplRecastNavMesh::GetNavMeshTilesIn(const TArray<FBox>& InclusionBounds, TArray<FNavTileRef>& OutRefs) const
 {
 	if (DetourNavMesh)
diff --git a/Engine/Source/Runtime/NavigationSystem/Private/NavMesh/RecastNavMeshDataChunk.cpp b/Engine/Source/Runtime/NavigationSystem/Private/NavMesh/RecastNavMeshDataChunk.cpp
--- a/Engine/Source/Runtime/NavigationSystem/Private/NavMesh/RecastNavMeshDataChunk.cpp
+++ b/Engine/Source/Runtime/NavigationSystem/Private/NavMesh/RecastNavMeshDataChunk.cpp
@@ -581,20 +581,6 @@ void URecastNavMeshDataChunk::GetTiles(const FPImplRecastNavMesh* NavMeshImpl, c
 }
 #endif // WITH_RECAST
 
-// Deprecated
-void URecastNavMeshDataChunk::GetTilesBounds(const FPImplRecastNavMesh& NavMeshImpl, const TArray<int32>& TileIndices, FBox& OutBounds) const
-{
-	OutBounds.Init();
-#if WITH_RECAST
-	TArray<uint32> TileUnsignedIndices;
-	TileUnsignedIndices.Append(TileIndices);
-
-	TArray<FNavTileRef> TileRefs;
-	FNavTileRef::DeprecatedMakeTileRefsFromTileIds(&NavMeshImpl, TileUnsignedIndices, TileRefs);
-	GetTilesBounds(NavMeshImpl, TileRefs, OutBounds);
-#endif // WITH_RECAST
-}
-
 #if WITH_RECAST
 void URecastNavMeshDataChunk::GetTilesBounds(const FPImplRecastNavMesh& NavMeshImpl, const TArray<FNavTileRef>& TileRefs, FBox& OutBounds) const
 {
diff --git a/Engine/Source/Runtime/NavigationSystem/Private/NavMesh/RecastNavMeshGenerator.cpp b/Engine/Source/Runtime/NavigationSystem/Private/NavMesh/RecastNavMeshGenerator.cpp
--- a/Engine/Source/Runtime/NavigationSystem/Private/NavMesh/RecastNavMeshGenerator.cpp
+++ b/Engine/Source/Runtime/NavigationSystem/Private/NavMesh/RecastNavMeshGenerator.cpp
@@ -2242,36 +2242,6 @@ void FRecastTileGenerator::GatherNavigationDataGeometry(const TSharedRef<FNaviga
 	}
 }
 
-// Deprecated
-bool FRecastTileGenerator::CreateHeightField(FNavMeshBuildContext& BuildContext, FTileRasterizationContext& RasterContext)
-{
-	return CreateHeightField(BuildContext);
-}
-
-// Deprecated
-void FRecastTileGenerator::GenerateRecastFilter(FNavMeshBuildContext& BuildContext, FTileRasterizationContext& RasterContext)
-{
-	GenerateRecastFilter(BuildContext);
-}
-
-// Deprecated
-ETimeSliceWorkResult FRecastTileGenerator::GenerateRecastFilterTimeSliced(FNavMeshBuildContext& BuildContext, FTileRasterizationContext& RasterContext)
-{
-	return GenerateRecastFilterTimeSliced(BuildContext);
-}
-
-// Deprecated
-bool FRecastTileGenerator::BuildCompactHeightField(FNavMeshBuildContext& BuildContext, FTileRasterizationContext& RasterContext)
-{
-	return BuildCompactHeightField(BuildContext);
-}
-
-// Deprecated
-bool FRecastTileGenerator::RecastErodeWalkable(FNavMeshBuildContext& BuildContext, FTileRasterizationContext& RasterContext)
-{
-	return RecastErodeWalkable(BuildContext);
-}
-
 void FRecastTileGenerator::ApplyVoxelFilter(rcHeightfield* HF, FVector::FReal WalkableRadius)
 {
 	SCOPE_CYCLE_COUNTER(STAT_Navigation_TileVoxelFilteringAsync);
@@ -2389,12 +2359,6 @@ void FRecastTileGenerator::InitRasterizationMaskArray(const rcHeightfield* InSol
 	FMemory::Memset(OutRasterizationMasks.GetData(), AllowAllFlags, CellCount*sizeof(TInlineMaskArray::ElementType));
 }
 
-// Deprecated
-void FRecastTileGenerator::PrepareVoxelCache(const TNavStatArray<uint8>& RawCollisionCache, const FCompositeNavModifier& InModifier, TNavStatArray<rcSpanCache>& SpanData)
-{
-	return PrepareVoxelCache(RawCollisionCache, FBox(ForceInit), InModifier, SpanData);
-}
-
 // Deprecated
 void FRecastTileGenerator::MarkDynamicAreas(dtTileCacheLayer& Layer)
 {
@@ -2583,11 +2547,6 @@ void FRecastTileGenerator::ValidateAndAppendGeometry(const FNavigationRelevantDa
 	}
 }
 
-void FRecastTileGenerator::ValidateAndAppendGeometry(const TSharedRef<FNavigationRelevantData, ESPMode::ThreadSafe>& ElementData, const FCompositeNavModifier& InModifier)
-{
-	ValidateAndAppendGeometry(ElementData.Get(), InModifier);
-}
-
 void FRecastTileGenerator::AppendGeometry(const FNavigationRelevantData& ElementData, const FCompositeNavModifier& InModifier, const FNavDataPerInstanceTransformDelegate& InTransformsDelegate)
 {	
 	const TNavStatArray<uint8>& RawCollisionCache = ElementData.CollisionData;
@@ -2813,12 +2772,6 @@ bool FRecastTileGenerator::CreateHeightField(FNavMeshBuildContext& BuildContext)
 	return true;
 }
 
-// Deprecated
-ETimeSliceWorkResult FRecastTileGenerator::RasterizeGeometryRecastTimeSliced(FNavMeshBuildContext& BuildContext, const TArray<FVector::FReal>& Coords, const TArray<int32>& Indices, const rcRasterizationFlags RasterizationFlags, FTileRasterizationContext& RasterContext)
-{
-	return RasterizeGeometryRecastTimeSliced(BuildContext, Coords, Indices, FBox(ForceInit), RasterizationFlags, RasterContext);
-}
-
 ETimeSliceWorkResult FRecastTileGenerator::RasterizeGeometryRecastTimeSliced(FNavMeshBuildContext& BuildContext, const TArray<FVector::FReal>& Coords, const TArray<int32>& Indices, const FBox& CoordsBounds, const rcRasterizationFlags RasterizationFlags, FTileRasterizationContext& RasterContext)
 {
 	QUICK_SCOPE_CYCLE_COUNTER(STAT_Navigation_RasterizeGeometryRecast);
@@ -2886,12 +2839,6 @@ ETimeSliceWorkResult FRecastTileGenerator::RasterizeGeometryRecastTimeSliced(FNa
 	return ETimeSliceWorkResult::Succeeded;
 }
 
-// Deprecated
-void FRecastTileGenerator::RasterizeGeometryRecast(FNavMeshBuildContext& BuildContext, const TArray<FVector::FReal>& Coords, const TArray<int32>& Indices, const rcRasterizationFlags RasterizationFlags, FTileRasterizationContext& RasterContext)
-{
-	return RasterizeGeometryRecast(BuildContext, Coords, Indices, FBox(ForceInit), RasterizationFlags, RasterContext);
-}
-
 void FRecastTileGenerator::RasterizeGeometryRecast(FNavMeshBuildContext& BuildContext, const TArray<FVector::FReal>& Coords, const TArray<int32>& Indices, const FBox& CoordsBounds, const rcRasterizationFlags RasterizationFlags, FTileRasterizationContext& RasterContext)
 {
 	QUICK_SCOPE_CYCLE_COUNTER(STAT_Navigation_RasterizeGeometryRecast);
@@ -2941,12 +2888,6 @@ void FRecastTileGenerator::RasterizeGeometryRecast(FNavMeshBuildContext& BuildCo
 	RasterizeGeomRecastTriAreas.Reset();
 }
 
-// Deprecated
-void FRecastTileGenerator::RasterizeGeometryTransformCoordsAndFlipIndices(const TArray<FVector::FReal>& Coords, const TArray<int32>& Indices, const FTransform& LocalToWorld)
-{
-	return RasterizeGeometryTransformCoordsAndFlipIndices(Coords, Indices, FBox(ForceInit), LocalToWorld);
-}
-
 void FRecastTileGenerator::RasterizeGeometryTransformCoordsAndFlipIndices(const TArray<FVector::FReal>& Coords, const TArray<int32>& Indices, const FBox& CoordsBounds, const FTransform& LocalToWorld)
 {
 	QUICK_SCOPE_CYCLE_COUNTER(STAT_Navigation_RasterizeGeometryTransformCoordsAndFlipIndices);
@@ -2988,34 +2929,6 @@ void FRecastTileGenerator::RasterizeGeometryTransformCoordsAndFlipIndices(const
 	}
 }
 
-// Deprecated
-void FRecastTileGenerator::RasterizeGeometryTransformCoords(const TArray<FVector::FReal>& Coords, const FTransform& LocalToWorld)
-{
-	QUICK_SCOPE_CYCLE_COUNTER(STAT_Navigation_RasterizeGeometryTransformCoords);
-
-	RasterizeGeometryWorldRecastCoords.SetNumUninitialized(Coords.Num(), EAllowShrinking::No);
-	RasterizeGeometryWorldRecastCoordsBounds = FBox(ForceInit);
-
-	FMatrix LocalToRecastWorld = LocalToWorld.ToMatrixWithScale()*Unreal2RecastMatrix();
-
-	// Convert geometry to recast world space
-	for (int32 i = 0; i < Coords.Num(); i+=3)
-	{
-		// collision cache stores coordinates in recast space, convert them to unreal and transform to recast world space
-		FVector WorldRecastCoord = LocalToRecastWorld.TransformPosition(Recast2UnrealPoint(&Coords[i]));
-
-		RasterizeGeometryWorldRecastCoords[i+0] = WorldRecastCoord.X;
-		RasterizeGeometryWorldRecastCoords[i+1] = WorldRecastCoord.Y;
-		RasterizeGeometryWorldRecastCoords[i+2] = WorldRecastCoord.Z;
-	}
-}
-
-// Deprecated
-ETimeSliceWorkResult FRecastTileGenerator::RasterizeGeometryTimeSliced(FNavMeshBuildContext& BuildContext, const TArray<FVector::FReal>& Coords, const TArray<int32>& Indices, const FTransform& LocalToWorld, const rcRasterizationFlags RasterizationFlags, FTileRasterizationContext& RasterContext)
-{
-	return RasterizeGeometryTimeSliced(BuildContext, Coords, Indices, FBox(ForceInit), LocalToWorld, RasterizationFlags, RasterContext);
-}
-
 ETimeSliceWorkResult FRecastTileGenerator::RasterizeGeometryTimeSliced(FNavMeshBuildContext& BuildContext, const TArray<FVector::FReal>& Coords, const TArray<int32>& Indices, const FBox& CoordsBounds, const FTransform& LocalToWorld, const rcRasterizationFlags RasterizationFlags, FTileRasterizationContext& RasterContext)
 {
 	QUICK_SCOPE_CYCLE_COUNTER(STAT_Navigation_RasterizeGeometry);
@@ -3060,11 +2973,6 @@ ETimeSliceWorkResult FRecastTileGenerator::RasterizeGeometryTimeSliced(FNavMeshB
 	return WorkResult;
 }
 
-void FRecastTileGenerator::RasterizeGeometry(FNavMeshBuildContext& BuildContext, const TArray<FVector::FReal>& Coords, const TArray<int32>& Indices, const FTransform& LocalToWorld, const rcRasterizationFlags RasterizationFlags, FTileRasterizationContext& RasterContext)
-{
-	return RasterizeGeometry(BuildContext, Coords, Indices, FBox(ForceInit), LocalToWorld, RasterizationFlags, RasterContext);
-}
-
 void FRecastTileGenerator::RasterizeGeometry(FNavMeshBuildContext& BuildContext, const TArray<FVector::FReal>& Coords, const TArray<int32>& Indices, const FBox& CoordsBounds, const FTransform& LocalToWorld, const rcRasterizationFlags RasterizationFlags, FTileRasterizationContext& RasterContext)
 {
 	QUICK_SCOPE_CYCLE_COUNTER(STAT_Navigation_RasterizeGeometry);
@@ -3662,34 +3570,6 @@ ETimeSliceWorkResult FRecastTileGenerator::GenerateCompressedLayersTimeSliced(FN
 	return ETimeSliceWorkResult::Succeeded;
 }
 
-// Deprecated
-bool FRecastTileGenerator::GenerateCompressedLayers(FNavMeshBuildContext& BuildContext)
-{
-	dtLinkBuilderData LinBuilderData;
-	return GenerateCompressedLayers(BuildContext, LinBuilderData);
-}
-
-// Deprecated
-bool FRecastTileGenerator::GenerateNavigationDataLayer(FNavMeshBuildContext& BuildContext, FTileCacheCompressor& TileCompressor, FTileCacheAllocator& GenNavAllocator, FTileGenerationContext& GenerationContext, int32 LayerIdx)
-{
-	dtLinkBuilderData LinBuilderData;
-	return GenerateNavigationDataLayer(BuildContext, TileCompressor, GenNavAllocator, GenerationContext, LinBuilderData, LayerIdx);
-}
-
-// Deprecated
-ETimeSliceWorkResult FRecastTileGenerator::GenerateNavigationDataTimeSliced(FNavMeshBuildContext& BuildContext)
-{
-	dtLinkBuilderData LinBuilderData;
-	return GenerateNavigationDataTimeSliced(BuildContext, LinBuilderData);
-}
-
-// Deprecated
-bool FRecastTileGenerator::GenerateNavigationData(FNavMeshBuildContext& BuildContext)
-{
-	dtLinkBuilderData LinBuilderData;
-	return GenerateNavigationData(BuildContext, LinBuilderData);
-}
-
 namespace UE::NavMesh::Private
 {
 	// Make sure LinkSpillDistance and CellSize has been computed before computing borders.
@@ -7559,13 +7439,6 @@ void FRecastNavMeshGenerator::GetDebugGeometry(const FNavigationRelevantData& En
 }
 #endif // !UE_BUILD_SHIPPING
 
-// Deprecated
-void FRecastNavMeshGenerator::ExportNavRelevantObjectGeometry(INavRelevantInterface& InOutNavRelevantInterface, FNavigationRelevantData& OutData)
-{
-	const TSharedRef<const FNavigationElement> TmpElement = FNavigationElement::CreateFromNavRelevantInterface(InOutNavRelevantInterface);
-	FRecastGeometryExport::ExportElementGeometry(TmpElement.Get(), OutData);
-}
-
 void FRecastGeometryExport::ExportElementGeometry(const FNavigationElement& InElement, FNavigationRelevantData& OutData)
 {
 	FRecastGeometryExport GeomExport(OutData);
@@ -7579,12 +7452,6 @@ void FRecastGeometryExport::ExportElementGeometry(const FNavigationElement& InEl
 	RecastGeometryExport::StoreCollisionCache(GeomExport);
 }
 
-// Deprecated
-void FRecastNavMeshGenerator::ExportVertexSoupGeometry(const TArray<FVector>& InVerts, FNavigationRelevantData& OutData)
-{
-	FRecastGeometryExport::ExportVertexSoupGeometry(InVerts, OutData);
-}
-
 void FRecastGeometryExport::ExportVertexSoupGeometry(const TArray<FVector>& InVerts, FNavigationRelevantData& OutData)
 {
 	FRecastGeometryExport GeomExport(OutData);
@@ -7597,18 +7464,6 @@ void FRecastGeometryExport::ExportVertexSoupGeometry(const TArray<FVector>& InVe
 	RecastGeometryExport::StoreCollisionCache(GeomExport);
 }
 
-
-// Deprecated
-void FRecastNavMeshGenerator::ExportRigidBodyGeometry(
-	UBodySetup& InOutBodySetup,
-	TNavStatArray<FVector>& OutVertexBuffer,
-	TNavStatArray<int32>& OutIndexBuffer,
-	FBox& OutBounds,
-	const FTransform& LocalToWorld)
-{
-	FRecastGeometryExport::ExportRigidBodyGeometry(InOutBodySetup, OutVertexBuffer, OutIndexBuffer, OutBounds, LocalToWorld);
-}
-
 void FRecastGeometryExport::ExportRigidBodyGeometry(UBodySetup& InOutBodySetup, TNavStatArray<FVector>& OutVertexBuffer, TNavStatArray<int32>& OutIndexBuffer, FBox& OutBounds, const FTransform& LocalToWorld)
 {
 	TNavStatArray<FVector::FReal> VertCoords;
@@ -7621,29 +7476,6 @@ void FRecastGeometryExport::ExportRigidBodyGeometry(UBodySetup& InOutBodySetup,
 	}
 }
 
-
-// Deprecated
-void FRecastNavMeshGenerator::ExportRigidBodyGeometry(
-	UBodySetup& InOutBodySetup,
-	TNavStatArray<FVector>& OutTriMeshVertexBuffer,
-	TNavStatArray<int32>& OutTriMeshIndexBuffer,
-	TNavStatArray<FVector>& OutConvexVertexBuffer,
-	TNavStatArray<int32>& OutConvexIndexBuffer,
-	TNavStatArray<int32>& OutShapeBuffer,
-	FBox& OutBounds,
-	const FTransform& LocalToWorld)
-{
-	FRecastGeometryExport::ExportRigidBodyGeometry(
-		InOutBodySetup,
-		OutTriMeshVertexBuffer,
-		OutTriMeshIndexBuffer,
-		OutConvexVertexBuffer,
-		OutConvexIndexBuffer,
-		OutShapeBuffer,
-		OutBounds,
-		LocalToWorld);
-}
-
 void FRecastGeometryExport::ExportRigidBodyGeometry(UBodySetup& InOutBodySetup, TNavStatArray<FVector>& OutTriMeshVertexBuffer, TNavStatArray<int32>& OutTriMeshIndexBuffer, TNavStatArray<FVector>& OutConvexVertexBuffer, TNavStatArray<int32>& OutConvexIndexBuffer, TNavStatArray<int32>& OutShapeBuffer, FBox& OutBounds, const FTransform& LocalToWorld)
 {
 	InOutBodySetup.CreatePhysicsMeshes();
@@ -7671,18 +7503,6 @@ void FRecastGeometryExport::ExportRigidBodyGeometry(UBodySetup& InOutBodySetup,
 	}
 }
 
-// Deprecated
-void FRecastNavMeshGenerator::ExportAggregatedGeometry(
-	const FKAggregateGeom& AggGeom,
-	TNavStatArray<FVector>& OutConvexVertexBuffer,
-	TNavStatArray<int32>& OutConvexIndexBuffer,
-	TNavStatArray<int32>& OutShapeBuffer,
-	FBox& OutBounds,
-	const FTransform& LocalToWorld)
-{
-	FRecastGeometryExport::ExportAggregatedGeometry(AggGeom, OutConvexVertexBuffer, OutConvexIndexBuffer, OutShapeBuffer, OutBounds, LocalToWorld);
-}
-
 void FRecastGeometryExport::ExportAggregatedGeometry(const FKAggregateGeom& AggGeom, TNavStatArray<FVector>& OutConvexVertexBuffer, TNavStatArray<int32>& OutConvexIndexBuffer, TNavStatArray<int32>& OutShapeBuffer, FBox& OutBounds, const FTransform& LocalToWorld)
 {
 	TNavStatArray<FVector::FReal> VertCoords;
diff --git a/Engine/Source/Runtime/NavigationSystem/Public/NavMesh/PImplRecastNavMesh.h b/Engine/Source/Runtime/NavigationSystem/Public/NavMesh/PImplRecastNavMesh.h
--- a/Engine/Source/Runtime/NavigationSystem/Public/NavMesh/PImplRecastNavMesh.h
+++ b/Engine/Source/Runtime/NavigationSystem/Public/NavMesh/PImplRecastNavMesh.h
@@ -73,25 +73,9 @@ public:
 	/** Returns bounding box for the whole navmesh. */
 	NAVIGATIONSYSTEM_API FBox GetNavMeshBounds() const;
 
-	/** Returns bounding box for a given navmesh tile. */
-	UE_DEPRECATED(5.5, "Use the version of this function on ARecastNavMesh that takes a FNavTileRef instead.")
-	NAVIGATIONSYSTEM_API FBox GetNavMeshTileBounds(int32 TileIndex) const;
-
-	/** Retrieves XY and layer coordinates of tile specified by index */
-	UE_DEPRECATED(5.5, "Use the version of this function on ARecastNavMesh that takes a FNavTileRef instead.")
-	NAVIGATIONSYSTEM_API bool GetNavMeshTileXY(int32 TileIndex, int32& OutX, int32& OutY, int32& OutLayer) const;
-
 	/** Retrieves XY coordinates of tile specified by position */
 	NAVIGATIONSYSTEM_API bool GetNavMeshTileXY(const FVector& Point, int32& OutX, int32& OutY) const;
 
-	/** Retrieves all tile indices at matching XY coordinates */
-	UE_DEPRECATED(5.5, "Use the version of this function on ARecastNavMesh that takes an array of FNavTileRefs instead.")
-	NAVIGATIONSYSTEM_API void GetNavMeshTilesAt(int32 TileX, int32 TileY, TArray<int32>& Indices) const;
-
-	/** Retrieves list of tiles that intersect specified bounds */
-	UE_DEPRECATED(5.5, "Use the version of the function that takes an array of FNavTileRefs instead")
-	NAVIGATIONSYSTEM_API void GetNavMeshTilesIn(const TArray<FBox>& InclusionBounds, TArray<int32>& Indices) const;
-
 	/** Retrieves list of tiles that intersect specified bounds */
 	NAVIGATIONSYSTEM_API void GetNavMeshTilesIn(const TArray<FBox>& InclusionBounds, TArray<FNavTileRef>& OutRefs) const;
 
diff --git a/Engine/Source/Runtime/NavigationSystem/Public/NavMesh/RecastNavMeshDataChunk.h b/Engine/Source/Runtime/NavigationSystem/Public/NavMesh/RecastNavMeshDataChunk.h
--- a/Engine/Source/Runtime/NavigationSystem/Public/NavMesh/RecastNavMeshDataChunk.h
+++ b/Engine/Source/Runtime/NavigationSystem/Public/NavMesh/RecastNavMeshDataChunk.h
@@ -99,10 +99,6 @@ class URecastNavMeshDataChunk : public UNavigationDataChunk
 	/** Const accessor to the list of tiles in the data chunk. */
 	const TArray<FRecastTileData>& GetTiles() const { return Tiles; }
 
-	/** Returns the AABB for the given tiles. */
-	UE_DEPRECATED(5.5, "Use the version of this function that takes an array of FNavTileRefs instead.")
-	NAVIGATIONSYSTEM_API void GetTilesBounds(const FPImplRecastNavMesh& NavMeshImpl, const TArray<int32>& TileIndices, FBox& OutBounds) const;
-
 #if WITH_RECAST
 	/** Returns the AABB for the given tiles. */
 	NAVIGATIONSYSTEM_API void GetTilesBounds(const FPImplRecastNavMesh& NavMeshImpl, const TArray<FNavTileRef>& TileRefs, FBox& OutBounds) const;
diff --git a/Engine/Source/Runtime/NavigationSystem/Public/NavMesh/RecastNavMeshGenerator.h b/Engine/Source/Runtime/NavigationSystem/Public/NavMesh/RecastNavMeshGenerator.h
--- a/Engine/Source/Runtime/NavigationSystem/Public/NavMesh/RecastNavMeshGenerator.h
+++ b/Engine/Source/Runtime/NavigationSystem/Public/NavMesh/RecastNavMeshGenerator.h
@@ -512,40 +512,6 @@ protected:
 #endif
 
 	/** Deprecated functions */
-	UE_DEPRECATED(5.5, "Use the new version without RasterContext instead.")
-	NAVIGATIONSYSTEM_API bool CreateHeightField(FNavMeshBuildContext& BuildContext, FTileRasterizationContext& RasterContext);
-	UE_DEPRECATED(5.5, "Use the new version without RasterContext instead.")
-	NAVIGATIONSYSTEM_API void GenerateRecastFilter(FNavMeshBuildContext& BuildContext, FTileRasterizationContext& RasterContext);
-	UE_DEPRECATED(5.5, "Use the new version without RasterContext instead.")
-	NAVIGATIONSYSTEM_API ETimeSliceWorkResult GenerateRecastFilterTimeSliced(FNavMeshBuildContext& BuildContext, FTileRasterizationContext& RasterContext);
-	UE_DEPRECATED(5.5, "Use the new version without RasterContext instead.")
-	NAVIGATIONSYSTEM_API bool BuildCompactHeightField(FNavMeshBuildContext& BuildContext, FTileRasterizationContext& RasterContext);
-	UE_DEPRECATED(5.5, "Use the new version without RasterContext instead.")
-	NAVIGATIONSYSTEM_API bool RecastErodeWalkable(FNavMeshBuildContext& BuildContext, FTileRasterizationContext& RasterContext);
-	UE_DEPRECATED(5.5, "This function was not handling the indices order correctly and is replaced by RasterizeGeometryTransformCoordsAndFlipIndices.")
-	NAVIGATIONSYSTEM_API void RasterizeGeometryTransformCoords(const TArray<FVector::FReal>& Coords, const FTransform& LocalToWorld);
-	UE_DEPRECATED(5.5, "Use the overload with dtLinkBuilderData instead.")
-	NAVIGATIONSYSTEM_API virtual bool GenerateCompressedLayers(FNavMeshBuildContext& BuildContext);
-	UE_DEPRECATED(5.5, "Use the overload with dtLinkBuilderData instead.")
-	NAVIGATIONSYSTEM_API bool GenerateNavigationDataLayer(FNavMeshBuildContext& BuildContext, FTileCacheCompressor& TileCompressor, FTileCacheAllocator& GenNavAllocator, FTileGenerationContext& GenerationContext, int32 LayerIdx);
-	UE_DEPRECATED(5.5, "Use the overload with dtLinkBuilderData instead.")
-	NAVIGATIONSYSTEM_API ETimeSliceWorkResult GenerateNavigationDataTimeSliced(FNavMeshBuildContext& BuildContext);
-	UE_DEPRECATED(5.5, "Use the overload with dtLinkBuilderData instead.")
-	NAVIGATIONSYSTEM_API bool GenerateNavigationData(FNavMeshBuildContext& BuildContext);
-	UE_DEPRECATED(5.5, "Use the version taking a const reference on FNavigationRelevantData.")
-	NAVIGATIONSYSTEM_API void ValidateAndAppendGeometry(const TSharedRef<FNavigationRelevantData, ESPMode::ThreadSafe>& ElementData, const FCompositeNavModifier& InModifier);
-	UE_DEPRECATED(5.6, "Please use the overload with the added CoordBounds.")
-	NAVIGATIONSYSTEM_API ETimeSliceWorkResult RasterizeGeometryRecastTimeSliced(FNavMeshBuildContext& BuildContext, const TArray<FVector::FReal>& Coords, const TArray<int32>& Indices, const rcRasterizationFlags RasterizationFlags, FTileRasterizationContext& RasterContext);
-	UE_DEPRECATED(5.6, "Please use the overload with the added CoordBounds.")
-	NAVIGATIONSYSTEM_API void RasterizeGeometryRecast(FNavMeshBuildContext& BuildContext, const TArray<FVector::FReal>& Coords, const TArray<int32>& Indices, const rcRasterizationFlags RasterizationFlags, FTileRasterizationContext& RasterContext);
-	UE_DEPRECATED(5.6, "Please use the overload with the added CoordBounds.")
-	NAVIGATIONSYSTEM_API void RasterizeGeometryTransformCoordsAndFlipIndices(const TArray<FVector::FReal>& Coords, const TArray<int32>& Indices, const FTransform& LocalToWorld);
-	UE_DEPRECATED(5.6, "Please use the overload with the added CoordBounds.")
-	NAVIGATIONSYSTEM_API ETimeSliceWorkResult RasterizeGeometryTimeSliced(FNavMeshBuildContext& BuildContext, const TArray<FVector::FReal>& Coords, const TArray<int32>& Indices, const FTransform& LocalToWorld, const rcRasterizationFlags RasterizationFlags, FTileRasterizationContext& RasterContext);
-	UE_DEPRECATED(5.6, "Please use the overload with the added CoordBounds.")
-	NAVIGATIONSYSTEM_API void RasterizeGeometry(FNavMeshBuildContext& BuildContext, const TArray<FVector::FReal>& Coords, const TArray<int32>& Indices, const FTransform& LocalToWorld, const rcRasterizationFlags RasterizationFlags, FTileRasterizationContext& RasterContext);
-	UE_DEPRECATED(5.6, "Please use the overload with the added RecastBounds.")
-	NAVIGATIONSYSTEM_API void PrepareVoxelCache(const TNavStatArray<uint8>& RawCollisionCache, const FCompositeNavModifier& InModifier, TNavStatArray<rcSpanCache>& SpanData);
 	UE_DEPRECATED(5.8, "Please use the overload with the added BuildContext.")
 	NAVIGATIONSYSTEM_API void MarkDynamicAreas(dtTileCacheLayer& Layer);
 	UE_DEPRECATED(5.8, "Please use the overload with the added BuildContext.")
@@ -869,38 +835,6 @@ public:
 	NAVIGATIONSYSTEM_API virtual void GrabDebugSnapshot(struct FVisualLogEntry* Snapshot, const FBox& BoundingBox, const FName& CategoryName, ELogVerbosity::Type Verbosity) const override;
 #endif
 
-	UE_DEPRECATED(5.5, "Use FRecastGeometryExport::ExportElementGeometry.")
-	static NAVIGATIONSYSTEM_API void ExportNavRelevantObjectGeometry(INavRelevantInterface& InOutNavRelevantInterface, FNavigationRelevantData& OutData);
-	UE_DEPRECATED(5.5, "Use FRecastGeometryExport::ExportVertexSoupGeometry.")
-	static NAVIGATIONSYSTEM_API void ExportVertexSoupGeometry(const TArray<FVector>& InVerts, FNavigationRelevantData& OutData);
-
-	UE_DEPRECATED(5.5, "Use FRecastGeometryExport::ExportRigidBodyGeometry.")
-	static NAVIGATIONSYSTEM_API void ExportRigidBodyGeometry(UBodySetup& InOutBodySetup,
-		TNavStatArray<FVector>& OutVertexBuffer,
-		TNavStatArray<int32>& OutIndexBuffer,
-		FBox& OutBounds,
-		const FTransform& LocalToWorld = FTransform::Identity);
-
-	UE_DEPRECATED(5.5, "Use FRecastGeometryExport::ExportRigidBodyGeometry.")
-	static NAVIGATIONSYSTEM_API void ExportRigidBodyGeometry(
-		UBodySetup& InOutBodySetup,
-		TNavStatArray<FVector>& OutTriMeshVertexBuffer,
-		TNavStatArray<int32>& OutTriMeshIndexBuffer,
-		TNavStatArray<FVector>& OutConvexVertexBuffer,
-		TNavStatArray<int32>& OutConvexIndexBuffer,
-		TNavStatArray<int32>& OutShapeBuffer,
-		FBox& OutBounds,
-		const FTransform& LocalToWorld = FTransform::Identity);
-
-	UE_DEPRECATED(5.5, "Use FRecastGeometryExport::ExportAggregatedGeometry.")
-	static NAVIGATIONSYSTEM_API void ExportAggregatedGeometry(
-		const FKAggregateGeom& AggGeom,
-		TNavStatArray<FVector>& OutConvexVertexBuffer,
-		TNavStatArray<int32>& OutConvexIndexBuffer,
-		TNavStatArray<int32>& OutShapeBuffer,
-		FBox& OutBounds,
-		const FTransform& LocalToWorld = FTransform::Identity);
-
 #if UE_ENABLE_DEBUG_DRAWING
 	/** Converts data encoded in EncodedData.CollisionData to FNavDebugMeshData format */
 	static NAVIGATIONSYSTEM_API void GetDebugGeometry(const FNavigationRelevantData& EncodedData, FNavDebugMeshData& DebugMeshData);
@@ -1034,14 +968,6 @@ protected:
 		return TileGenerator;
 	}
 
-	UE_DEPRECATED(5.5, "Use BBoxGrowthLow and BBoxGrowthHigh properties instead.")
-	void SetBBoxGrowth(const FVector& InBBox)
-	{
-PRAGMA_DISABLE_DEPRECATION_WARNINGS
-		BBoxGrowth = InBBox;
-PRAGMA_ENABLE_DEPRECATION_WARNINGS
-	}
-
 	FBox GrowDirtyBounds(const FBox& BBox, bool bIncludeAgentHeight) const;
 
 	//----------------------------------------------------------------------//
