Minor fix in GetEdgesForPathCorridorImpl + typo fix

#Synced-CL 49739737

diff --git a/Engine/Source/Runtime/NavigationSystem/Private/NavMesh/PImplRecastNavMesh.cpp b/Engine/Source/Runtime/NavigationSystem/Private/NavMesh/PImplRecastNavMesh.cpp
--- a/Engine/Source/Runtime/NavigationSystem/Private/NavMesh/PImplRecastNavMesh.cpp
+++ b/Engine/Source/Runtime/NavigationSystem/Private/NavMesh/PImplRecastNavMesh.cpp
@@ -2602,10 +2602,10 @@ bool FPImplRecastNavMesh::GetClusterBounds(NavNodeRef ClusterRef, FBox& OutBound
 
 void FPImplRecastNavMesh::GetEdgesForPathCorridorImpl(const TArray<NavNodeRef>* PathCorridor, TArray<FNavigationPortalEdge>* PathCorridorEdges, const dtNavMeshQuery& NavQuery) const
 {
-	const int32 CorridorLenght = PathCorridor->Num();
+	const int32 CorridorLength = PathCorridor->Num();
 
-	PathCorridorEdges->Empty(CorridorLenght - 1);
-	for (int32 i = 0; i < CorridorLenght - 1; ++i)
+	PathCorridorEdges->Empty(FMath::Max(CorridorLength - 1, 0));
+	for (int32 i = 0; i < CorridorLength - 1; ++i)
 	{
 		unsigned char FromType = 0, ToType = 0;
 		FVector::FReal Left[3] = {0.f}, Right[3] = {0.f};
