Fixing navigation data not being removed when using external actors.

#Synced-CL 49739737

diff --git a/Engine/Source/Runtime/NavigationSystem/NavigationSystem.Build.cs b/Engine/Source/Runtime/NavigationSystem/NavigationSystem.Build.cs
--- a/Engine/Source/Runtime/NavigationSystem/NavigationSystem.Build.cs
+++ b/Engine/Source/Runtime/NavigationSystem/NavigationSystem.Build.cs
@@ -74,6 +74,7 @@ namespace UnrealBuildTool.Rules
                 // @todo api: Only public because of WITH_EDITOR and UNREALED_API
 				PublicDependencyModuleNames.Add("EditorFramework");
                 PublicDependencyModuleNames.Add("UnrealEd");
+                PrivateDependencyModuleNames.Add("SourceControl");
                 CircularlyReferencedDependentModules.Add("UnrealEd");
             }
         }
diff --git a/Engine/Source/Runtime/NavigationSystem/Private/NavigationData.cpp b/Engine/Source/Runtime/NavigationSystem/Private/NavigationData.cpp
--- a/Engine/Source/Runtime/NavigationSystem/Private/NavigationData.cpp
+++ b/Engine/Source/Runtime/NavigationSystem/Private/NavigationData.cpp
@@ -547,6 +547,21 @@ void ANavigationData::ApplyWorldOffset(const FVector& InOffset, bool bWorldShift
 	// do nothing, will be handled by NavigationSystem
 }
 
+void ANavigationData::DestroyActor()
+{
+	UWorld* World = GetWorld();
+	if (World && World->IsEditorWorld() && !World->IsPlayInEditor())
+	{
+		// Editor-aware destroy (notifies editor, handles actors)
+		World->EditorDestroyActor(this, /*bShouldModifyLevel*/true);
+	}
+	else if (World)
+	{
+		// Runtime destroy
+		Destroy(/*bNetForce*/false, /*bShouldModifyLevel*/false);
+	}
+}
+
 void ANavigationData::CleanUpAndMarkPendingKill()
 {
 	CleanUp();
diff --git a/Engine/Source/Runtime/NavigationSystem/Private/NavigationSystem.cpp b/Engine/Source/Runtime/NavigationSystem/Private/NavigationSystem.cpp
--- a/Engine/Source/Runtime/NavigationSystem/Private/NavigationSystem.cpp
+++ b/Engine/Source/Runtime/NavigationSystem/Private/NavigationSystem.cpp
@@ -49,12 +49,15 @@
 #endif // WITH_RECAST
 
 #if WITH_EDITOR
+#include "Editor.h"
 #include "AssetCompilingManager.h"
 #include "EditorModeManager.h"
 #include "EditorModes.h"
 #include "LevelEditor.h"
 #include "Misc/MessageDialog.h"
 #include "WorldPartition/WorldPartition.h"
+#include "SourceControlHelpers.h"
+#include "HAL/PlatformFileManager.h"
 #endif // WITH_EDITOR
 
 
@@ -1170,14 +1173,17 @@ void UNavigationSystemV1::PostInitProperties()
 		uint8 UseLockFlags = InitialNavBuildingLockFlags;
 
 		AddNavigationBuildLock(UseLockFlags);
-
-		// register for any actor move change
+		
 #if WITH_EDITOR
-		if ( GIsEditor )
+		if (GIsEditor && GEngine)
 		{
+			// register for any actor move change
 			GEngine->OnActorMoved().AddUObject(this, &UNavigationSystemV1::OnActorMoved);
+
+			FEditorDelegates::OnMapOpened.AddUObject(this, &UNavigationSystemV1::OnMapOpened);
 		}
-#endif
+#endif // WITH_EDITOR
+
 		FCoreUObjectDelegates::PostLoadMapWithWorld.AddUObject(this, &UNavigationSystemV1::OnPostLoadMap);
 		UNavigationSystemV1::NavigationDirtyEvent.AddUObject(this, &UNavigationSystemV1::OnNavigationDirtied);
 
@@ -2757,7 +2763,17 @@ void UNavigationSystemV1::ProcessRegistrationCandidates()
 
 			if (Result != RegistrationSuccessful && Result != RegistrationFailed_DataPendingKill)
 			{
-				NavDataPtr->Destroy();
+#if WITH_EDITOR				
+				if (GEditor)
+				{
+					NavigationDataToRemove.Add(NavDataPtr);	
+				}
+				else
+#endif // WITH_EDITOR					
+				{
+					NavDataPtr->Destroy();
+				}
+
 				if (NavDataPtr == MainNavData)
 				{
 					MainNavData = nullptr;
@@ -5009,6 +5025,63 @@ void UNavigationSystemV1::OnPostLoadMap(UWorld* LoadedWorld)
 }
 
 #if WITH_EDITOR
+void UNavigationSystemV1::OnMapOpened(const FString& Filename, bool bAsTemplate)
+{
+	// Now that the maps as loaded, remove unnecessary packages (and actors for OFPA).
+	for (TWeakObjectPtr<ANavigationData> NavigationDataPtr : NavigationDataToRemove)
+	{
+		if (NavigationDataPtr.IsValid())
+		{
+			bool bIsUsingExternalActor = false;
+			if (const UWorld* World = GetWorld())
+			{
+				const ULevel* Level = World->PersistentLevel;
+				bIsUsingExternalActor = Level && Level->IsUsingExternalActors();
+			}
+
+			UPackage* Package = nullptr;
+			FString PackageName;
+			if (bIsUsingExternalActor)
+			{
+				Package = NavigationDataPtr->GetPackage();
+				if (Package)
+				{
+				    PackageName = SourceControlHelpers::PackageFilename(Package);
+				}	
+			}
+			
+			UE_LOG(LogNavigation, Verbose, TEXT("%hs Destroying %s because it failed to register."), __FUNCTION__, *NavigationDataPtr->GetFullName());
+			NavigationDataPtr->DestroyActor();
+
+			if (bIsUsingExternalActor && Package)
+			{
+				bool bDeleted = false;
+
+				FPackageSourceControlHelper PackageHelper;
+				if (PackageHelper.UseSourceControl())
+				{
+					// Delete with source control
+					bDeleted = PackageHelper.Delete(Package);
+				}
+				else
+				{
+					// Direct file delete
+					if (!PackageName.IsEmpty())
+					{
+						bDeleted = FPlatformFileManager::Get().GetPlatformFile().DeleteFile(*PackageName);
+					}
+				}
+
+				if (bDeleted)
+				{
+					UE_LOG(LogNavigation, Verbose, TEXT("%hs File deleted (%s)."), __FUNCTION__, *PackageName);
+				}
+			}
+		}
+	}
+	NavigationDataToRemove.Reset();
+}
+
 void UNavigationSystemV1::OnActorMoved(AActor* Actor)
 {
 	if (Cast<ANavMeshBoundsVolume>(Actor))
@@ -5055,6 +5128,8 @@ void UNavigationSystemV1::CleanUp(FNavigationSystem::ECleanupMode Mode)
 #if WITH_EDITOR
 	if (GIsEditor && GEngine)
 	{
+		FEditorDelegates::OnMapOpened.RemoveAll(this);
+		
 		GEngine->OnActorMoved().RemoveAll(this);
 	}
 
diff --git a/Engine/Source/Runtime/NavigationSystem/Public/NavigationData.h b/Engine/Source/Runtime/NavigationSystem/Public/NavigationData.h
--- a/Engine/Source/Runtime/NavigationSystem/Public/NavigationData.h
+++ b/Engine/Source/Runtime/NavigationSystem/Public/NavigationData.h
@@ -622,6 +622,7 @@ private:
 	NAVIGATIONSYSTEM_API void UnregisterAndCleanUp();
 
 public:
+	NAVIGATIONSYSTEM_API void DestroyActor();
 	NAVIGATIONSYSTEM_API virtual void CleanUpAndMarkPendingKill();
 
 	inline bool IsRegistered() const { return bRegistered; }
diff --git a/Engine/Source/Runtime/NavigationSystem/Public/NavigationSystem.h b/Engine/Source/Runtime/NavigationSystem/Public/NavigationSystem.h
--- a/Engine/Source/Runtime/NavigationSystem/Public/NavigationSystem.h
+++ b/Engine/Source/Runtime/NavigationSystem/Public/NavigationSystem.h
@@ -1269,6 +1269,8 @@ protected:
 	UPROPERTY(transient)
 	TObjectPtr<UNavigationObjectRepository> Repository = nullptr;
 
+	TArray<TWeakObjectPtr<ANavigationData>> NavigationDataToRemove;
+
 	// async queries
 	FCriticalSection NavDataRegistrationSection;
 
@@ -1311,10 +1313,14 @@ protected:
 
 	/** delegate handler for PostLoadMap event */
 	NAVIGATIONSYSTEM_API void OnPostLoadMap(UWorld* LoadedWorld);
+	
 #if WITH_EDITOR
+	void OnMapOpened(const FString& Filename, bool bAsTemplate);
+	
 	/** delegate handler for ActorMoved events */
 	NAVIGATIONSYSTEM_API void OnActorMoved(AActor* Actor);
-#endif
+#endif // WITH_EDITOR
+
 	/** delegate handler called when navigation is dirtied*/
 	NAVIGATIONSYSTEM_API void OnNavigationDirtied(const FBox& Bounds);
 
