[Navigation] Prevent crash when exporting geometry with to many triangles.

#Synced-CL 49739737

diff --git a/Engine/Source/Runtime/NavigationSystem/Private/NavMesh/RecastNavMeshGenerator.cpp b/Engine/Source/Runtime/NavigationSystem/Private/NavMesh/RecastNavMeshGenerator.cpp
--- a/Engine/Source/Runtime/NavigationSystem/Private/NavMesh/RecastNavMeshGenerator.cpp
+++ b/Engine/Source/Runtime/NavigationSystem/Private/NavMesh/RecastNavMeshGenerator.cpp
@@ -86,7 +86,7 @@ namespace UE::NavMesh::Private
 	static FAutoConsoleVariableRef CVarUseTightBoundExpansion(TEXT("ai.nav.UseTightBoundExpansion"), bUseTightBoundExpansion, TEXT("Active by default. Use an expansion of one AgentRadius. Set to false to revert to the previous behavior (2 AgentRadius)."), ECVF_Default);
 
 	static bool bUseAsymetricBorderSizes = false;
-	static FAutoConsoleVariableRef CVarUseAsymetricBorderSizes(TEXT("ai.nav.UseAsymetricBorderSizes"), bUseAsymetricBorderSizes, TEXT("Active by default. When generating links, use asymetric tile border sizes to improve generation speed."), ECVF_Default);
+	static FAutoConsoleVariableRef CVarUseAsymetricBorderSizes(TEXT("ai.nav.UseAsymetricBorderSizes"), bUseAsymetricBorderSizes, TEXT("Disabled by default. When generating links, use asymetric tile border sizes to improve generation speed."), ECVF_Default);
 
 	static bool bAllowLinkGeneration = true;
 	static FAutoConsoleVariableRef CVarAllowLinkGeneration(TEXT("ai.nav.AllowLinkGeneration"), bAllowLinkGeneration, TEXT("Set to false to force disabling link generation."), ECVF_Default);
@@ -333,11 +333,37 @@ static void StoreCollisionCache(FRecastGeometryExport& GeomExport)
 	HeaderInfo.SlopeOverride = GeomExport.SlopeOverride;
 
 	// allocate memory
-	const int32 HeaderSize = sizeof(FRecastGeometryCache);
-	const int32 CoordsSize = sizeof(FVector::FReal) * 3 * NumVerts;
-	const int32 IndicesSize = sizeof(int32) * 3 * NumFaces;
-	const int32 CacheSize = HeaderSize + CoordsSize + IndicesSize;
+	constexpr int32 HeaderSize = sizeof(FRecastGeometryCache);
+	const int64 CoordsSize = sizeof(FVector::FReal) * 3 * NumVerts;
+	const int64 IndicesSize = sizeof(int32) * 3 * NumFaces;
+	const int64 CacheSize64 = HeaderSize + CoordsSize + IndicesSize;
+
+	if (!IntFitsIn<int32>(CacheSize64))
+	{
+		UE_LOG(LogNavigation, Warning, TEXT("%hs Can't cache collision data (expected size would be %lld). "
+			"See other warnings, like logs about collision geometry triangle count."),
+			__FUNCTION__, CacheSize64);
+
+		// Fill with a empty header
+		HeaderInfo.Validation.DataSize = 0;
+		HeaderInfo.NumFaces = 0;
+		HeaderInfo.NumVerts = 0;
+		HeaderInfo.SlopeOverride = GeomExport.SlopeOverride;
+
+		// empty + add combo to allocate exact amount (without any overhead/slack)
+		GeomExport.Data->CollisionData.Empty(HeaderSize);
+		GeomExport.Data->CollisionData.AddUninitialized(HeaderSize);
+		
+		uint8* RawMemory = GeomExport.Data->CollisionData.GetData();
+		FRecastGeometryCache* CacheMemory = (FRecastGeometryCache*)RawMemory;
+		
+		CacheMemory->Header = HeaderInfo;
+		CacheMemory->Verts = nullptr;
+		CacheMemory->Indices = nullptr;
+		return;
+	}
 
+	const int32 CacheSize = (int32)CacheSize64;
 	HeaderInfo.Validation.DataSize = CacheSize;
 
 	// empty + add combo to allocate exact amount (without any overhead/slack)
@@ -348,8 +374,8 @@ static void StoreCollisionCache(FRecastGeometryExport& GeomExport)
 	uint8* RawMemory = GeomExport.Data->CollisionData.GetData();
 	FRecastGeometryCache* CacheMemory = (FRecastGeometryCache*)RawMemory;
 	CacheMemory->Header = HeaderInfo;
-	CacheMemory->Verts = 0;
-	CacheMemory->Indices = 0;
+	CacheMemory->Verts = nullptr;
+	CacheMemory->Indices = nullptr;
 
 	FMemory::Memcpy(RawMemory + HeaderSize, GeomExport.VertexBuffer.GetData(), CoordsSize);
 	FMemory::Memcpy(RawMemory + HeaderSize + CoordsSize, GeomExport.IndexBuffer.GetData(), IndicesSize);
diff --git a/Engine/Source/Runtime/Navmesh/Private/Recast/Recast.cpp b/Engine/Source/Runtime/Navmesh/Private/Recast/Recast.cpp
--- a/Engine/Source/Runtime/Navmesh/Private/Recast/Recast.cpp
+++ b/Engine/Source/Runtime/Navmesh/Private/Recast/Recast.cpp
@@ -436,7 +436,7 @@ int rcGetHeightFieldSpanCount(rcContext* /*ctx*/, rcHeightfield& hf)
 /// @par
 ///
 /// This is just the beginning of the process of fully building a compact heightfield.
-/// Various filters may be applied applied, then the distance field and regions built.
+/// Various filters may be applied, then the distance field and regions built.
 /// E.g: #rcBuildDistanceField and #rcBuildRegions
 ///
 /// See the #rcConfig documentation for more information on the configuration parameters.
