Fix to make sure generated NavLinks are applied to the correct layers when there are directional links with different areas.

#Synced-CL 50022663

diff --git a/Engine/Source/Runtime/NavigationSystem/Private/NavMesh/RecastNavMeshGenerator.cpp b/Engine/Source/Runtime/NavigationSystem/Private/NavMesh/RecastNavMeshGenerator.cpp
--- a/Engine/Source/Runtime/NavigationSystem/Private/NavMesh/RecastNavMeshGenerator.cpp
+++ b/Engine/Source/Runtime/NavigationSystem/Private/NavMesh/RecastNavMeshGenerator.cpp
@@ -1425,7 +1425,7 @@ struct FOffMeshData
 		}
 	}
 	
-	void AddLinks(const TArray<FGeneratedNavigationLink>& Links, const FTransform& LocalToWorld, int32 AgentIndex, float DefaultSnapHeight)
+	void AddLinks(const TInlineGeneratedNavLinkArray& Links, const FTransform& LocalToWorld, int32 AgentIndex, float DefaultSnapHeight)
 	{
 		for (int32 LinkIndex = 0; LinkIndex < Links.Num(); ++LinkIndex)
 		{
@@ -3783,26 +3783,41 @@ struct FTileGenerationContext
 };
 
 static void AddGeneratedLinks(
-	TArray<FGeneratedNavigationLink>& OutGeneratedLinks,
+	TInlineGeneratedNavLinkArray& OutGeneratedLinks,
 	const dtLinkBuilderConfig& linkBuilderConfig,
-	const dtNavLinkBuilder& linkBuilder)
+	const dtNavLinkBuilder& linkBuilder,
+	const TArray<FBox, TInlineAllocator<16>>& LayerBounds,
+	const float SnapHeight)
 {
-	auto AddLinkSharedLambda = [&OutGeneratedLinks](const dtNavLinkBuilderJumpConfig& config) -> FGeneratedNavigationLink&
+	auto FindLayerForZ = [&](const FVector::FReal Z) -> int16 {
+		for (int32 i = 0; i < LayerBounds.Num(); i++) {
+			if ((Z + SnapHeight) >= LayerBounds[i].Min.Z &&
+				(Z - SnapHeight) <= LayerBounds[i].Max.Z)
+				return (int16)i;
+		}
+		return INDEX_NONE;
+	};
+
+	auto AddLinkSharedLambda = [&OutGeneratedLinks, &FindLayerForZ](const dtNavLinkBuilderJumpConfig& config, const FVector& LeftPos) -> FGeneratedNavigationLink&
 	{
 		FGeneratedNavigationLink& NewLink = OutGeneratedLinks.Emplace_GetRef();
 		NewLink.bIsGenerated = true;
 		NewLink.NavLinkId = FNavLinkId(config.linkUserId);
+		NewLink.TargetLayerIndex = FindLayerForZ(LeftPos.Z);
 		return NewLink;
 	};
 
 	// Used when we're provided with identical up and down areas. In this situation, a single bidirectional link is sufficient for our needs
 	auto AddBidirectionalLinkLambda = [&AddLinkSharedLambda](const dtNavLinkBuilderJumpConfig& config, const dtReal* posA, const dtReal* posB)
 	{
-		FGeneratedNavigationLink& NewLink = AddLinkSharedLambda(config);
+		const FVector LeftPos = Recast2UnrealPoint(posA);
+		const FVector RightPos = Recast2UnrealPoint(posB);
+
+		FGeneratedNavigationLink& NewLink = AddLinkSharedLambda(config, LeftPos);
 		NewLink.generatedLinkArea = config.downDirArea;
 		NewLink.generatedLinkPolyFlag = config.downDirPolyFlag;
-		NewLink.Left = Recast2UnrealPoint(posA);
-		NewLink.Right = Recast2UnrealPoint(posB);
+		NewLink.Left = LeftPos;
+		NewLink.Right = RightPos;
 	};
 
 	// Used when we're provided with different up and down areas. In this situation, we need to create two different one-way link representations in order to meet our needs
@@ -3818,7 +3833,7 @@ static void AddGeneratedLinks(
 
 		if (config.downDirArea != RECAST_NULL_AREA)
 		{
-			FGeneratedNavigationLink& DownLink = AddLinkSharedLambda(config);
+			FGeneratedNavigationLink& DownLink = AddLinkSharedLambda(config, DownLinkLeft);
 			DownLink.generatedLinkArea = config.downDirArea;
 			DownLink.generatedLinkPolyFlag = config.downDirPolyFlag;
 			DownLink.Left = DownLinkLeft;
@@ -3829,7 +3844,7 @@ static void AddGeneratedLinks(
 		if (config.upDirArea != RECAST_NULL_AREA)
 		{
 			// Note that in order for this link to be recognized by the navigation system, it needs to have unique left/right points, which is why the left and right points are flipped and we're marking this link as left-to-right
-			FGeneratedNavigationLink& UpLink = AddLinkSharedLambda(config);
+			FGeneratedNavigationLink& UpLink = AddLinkSharedLambda(config, DownLinkRight);
 			UpLink.generatedLinkArea = config.upDirArea;
 			UpLink.generatedLinkPolyFlag = config.upDirPolyFlag;
 			UpLink.Left = DownLinkRight;
@@ -3887,7 +3902,7 @@ static void AddGeneratedLinks(
 }
 
 dtStatus FRecastTileGenerator::BuildTileCacheLinks(FNavMeshBuildContext& BuildContext, dtTileCacheAlloc* alloc, const dtTileCacheLayer& layer,
-                                                   const dtTileCacheContourSet& lcset, TArray<FGeneratedNavigationLink>& OutGeneratedLinks) const
+                                                   const dtTileCacheContourSet& lcset, const float DefaultSnapHeight, TInlineGeneratedNavLinkArray& OutGeneratedLinks) const
 {
 	TRACE_CPUPROFILER_EVENT_SCOPE(FRecastTileGenerator::BuildTileCacheLinks);
 	BuildContext.log(RC_LOG_PROGRESS, "Building Links:");
@@ -4001,14 +4016,23 @@ dtStatus FRecastTileGenerator::BuildTileCacheLinks(FNavMeshBuildContext& BuildCo
 #endif // RECAST_INTERNAL_DEBUG_DATA
 	}
 
-	AddGeneratedLinks(OutGeneratedLinks, linkBuilderConfig, linkBuilder);
+	// Extract layer bounds
+	TArray<FBox, TInlineAllocator<16>> LayerBounds;
+	LayerBounds.Reserve(CompressedLayers.Num());
+	for (const FNavMeshTileData& LayerData : CompressedLayers)
+	{
+		LayerBounds.Add(LayerData.LayerBBox);
+	}
+
+	AddGeneratedLinks(OutGeneratedLinks, linkBuilderConfig, linkBuilder, LayerBounds, DefaultSnapHeight);
 
 	LogOnExit();
 	return DT_SUCCESS;
 }
 
 bool FRecastTileGenerator::GenerateNavigationDataLayer(FNavMeshBuildContext& BuildContext, FTileCacheCompressor& TileCompressor,
-	FTileCacheAllocator& GenNavAllocator, FTileGenerationContext& GenerationContext, const dtLinkBuilderData& InLinkBuilderData, int32 LayerIdx)
+	FTileCacheAllocator& GenNavAllocator, FTileGenerationContext& GenerationContext, const dtLinkBuilderData& InLinkBuilderData, int32 LayerIdx, const TInlineGeneratedNavLinkArray* InAdditionalLinks,
+	TInlineGeneratedNavLinkArray* OutGeneratedLinks)
 {
 	SCOPE_CYCLE_COUNTER(STAT_Navigation_GenerateNavigationDataLayer)
 		
@@ -4203,21 +4227,44 @@ bool FRecastTileGenerator::GenerateNavigationDataLayer(FNavMeshBuildContext& Bui
 #endif
 	}
 
+	const float DefaultSnapHeight = static_cast<float>(TileConfig.walkableClimb) * static_cast<float>(TileConfig.ch);
+
 	// Build Links
-	TArray<FGeneratedNavigationLink> GeneratedLinks;
+	TInlineGeneratedNavLinkArray GeneratedLinks;
 	if (InLinkBuilderData.generatingLinks)
 	{
 		SCOPE_CYCLE_COUNTER(STAT_Navigation_RecastBuildLinks);
 
 		if (SolidHF && CompactHF)
 		{
-			status = BuildTileCacheLinks(BuildContext, &GenNavAllocator, *GenerationContext.Layer, *GenerationContext.ContourSet, GeneratedLinks);	
+			TInlineGeneratedNavLinkArray NewlyGeneratedLinks;
+			status = BuildTileCacheLinks(BuildContext, &GenNavAllocator, *GenerationContext.Layer, *GenerationContext.ContourSet, DefaultSnapHeight, NewlyGeneratedLinks);
+
+			if (dtStatusFailed(status))
+			{
+				BuildContext.log(RC_LOG_ERROR, "GenerateNavigationDataLayer: Failed to generate links (0x%08X).", status);
+				return false;
+			}
+
+			for (const FGeneratedNavigationLink& NewLink : NewlyGeneratedLinks)
+			{
+				if (NewLink.TargetLayerIndex == LayerIdx)
+				{
+					// Belongs to this layer - add to GeneratedLinks for nav data
+					GeneratedLinks.Add(NewLink);
+				}
+				else if (OutGeneratedLinks != nullptr && NewLink.TargetLayerIndex != INDEX_NONE)
+				{
+					// Belongs to a different layer - output for deferral
+					OutGeneratedLinks->Add(NewLink);
+				}
+			}
 		}
 
-		if (dtStatusFailed(status))
+		// Add links from previous layers. Needs to happen regardless of valid height fields.
+		if (InAdditionalLinks != nullptr)
 		{
-			BuildContext.log(RC_LOG_ERROR, "GenerateNavigationDataLayer: Failed to generate links (0x%08X).", status);
-			return false;
+			GeneratedLinks.Append(*InAdditionalLinks);
 		}
 	}
 
@@ -4235,8 +4282,6 @@ bool FRecastTileGenerator::GenerateNavigationDataLayer(FNavMeshBuildContext& Bui
 			return false;
 		}
 
-		const float DefaultSnapHeight = static_cast<float>(TileConfig.walkableClimb) * static_cast<float>(TileConfig.ch);
-		
 		// if we didn't fail already then it's high time we created data for off-mesh links
 		FOffMeshData OffMeshData;
 		if (!OffmeshLinks.IsEmpty() || !GeneratedLinks.IsEmpty())
@@ -4325,6 +4370,58 @@ bool FRecastTileGenerator::GenerateNavigationDataLayer(FNavMeshBuildContext& Bui
 	return true;
 }
 
+void FRecastTileGenerator::PrepareZOrderedLayerIndices(TArray<int32>& OutLayerIndices) const
+{
+	const int32 NumLayers = FMath::Min(CompressedLayers.Num(), DirtyLayers.Num());
+
+	// Step 1: Create array for indices
+	OutLayerIndices.Reserve(NumLayers);
+
+	// Step 2: Add valid layer indices
+	for (int32 Index = 0; Index < NumLayers; Index++)
+	{
+		if (DirtyLayers[Index] && CompressedLayers[Index].IsValid())
+		{
+			OutLayerIndices.Add(Index);
+		}
+	}
+
+	// Step 3: Sort by Z descending (highest first)
+	OutLayerIndices.Sort([this](int32 A, int32 B)
+	{
+		return CompressedLayers[A].LayerBBox.Max.Z > CompressedLayers[B].LayerBBox.Max.Z;
+	});
+}
+
+bool FRecastTileGenerator::GenerateNavigationDataLayerWithDeferredLinks(FNavMeshBuildContext& BuildContext, FTileCacheCompressor& TileCompressor, FTileCacheAllocator& GenNavAllocator, FTileGenerationContext& GenerationContext, 
+	const dtLinkBuilderData& InLinkBuilderData, int32 LayerIdx,	TInlineGeneratedNavLinkArray& InOutDeferredLinks)
+{
+	// Collect deferred links for THIS layer
+	TInlineGeneratedNavLinkArray LinksFromHigherLayers;
+	for (int32 Index = InOutDeferredLinks.Num() - 1; Index >= 0; Index--)
+	{
+		if (InOutDeferredLinks[Index].TargetLayerIndex == LayerIdx)
+		{
+			LinksFromHigherLayers.Add(InOutDeferredLinks[Index]);
+			InOutDeferredLinks.RemoveAtSwap(Index);
+		}
+	}
+
+	// Do the regular navigation data layer generation, passing in matching deferred links and retrieving new ones.
+
+	TInlineGeneratedNavLinkArray NewlyGeneratedLinks;
+	const bool bGenDataLayerSuccess = GenerateNavigationDataLayer(BuildContext, TileCompressor, GenNavAllocator, GenerationContext, InLinkBuilderData, LayerIdx, &LinksFromHigherLayers, &NewlyGeneratedLinks);
+
+	if (bGenDataLayerSuccess)
+	{
+		// Defer links for lower layers
+		InOutDeferredLinks.Append(NewlyGeneratedLinks);
+	}
+
+	return bGenDataLayerSuccess;
+}
+
+
 ETimeSliceWorkResult FRecastTileGenerator::GenerateNavigationDataTimeSliced(FNavMeshBuildContext& BuildContext, const dtLinkBuilderData& InLinkBuilderData)
 {
 	SCOPE_CYCLE_COUNTER(STAT_Navigation_RecastBuildNavigation);
@@ -4345,37 +4442,78 @@ ETimeSliceWorkResult FRecastTileGenerator::GenerateNavigationDataTimeSliced(FNav
 		GenNavDataTimeSlicedAllocator = MakeUnique<FTileCacheAllocator>();
 		GenNavDataTimeSlicedGenerationContext = MakeUnique<FTileGenerationContext>(GenNavDataTimeSlicedAllocator.Get());
 		GenNavDataTimeSlicedGenerationContext->NavigationData.Reserve(CompressedLayers.Num());
+		GenNavDataLayerIndices.Empty();
+		GenNavDataDeferredLinks.Empty();
 		GenerateNavDataTimeSlicedState = EGenerateNavDataTimeSlicedState::GenerateLayers;
 	}//fall through to next state
 	case EGenerateNavDataTimeSlicedState::GenerateLayers:
 	{
-		for (; GenNavDataLayerTimeSlicedIdx < CompressedLayers.Num() && GenNavDataLayerTimeSlicedIdx < DirtyLayers.Num(); GenNavDataLayerTimeSlicedIdx++)
+		const int32 NumLayers = FMath::Min(CompressedLayers.Num(), DirtyLayers.Num());
+
+		if (InLinkBuilderData.generatingLinks && NumLayers > 1)
 		{
-			if (DirtyLayers[GenNavDataLayerTimeSlicedIdx] == false || !CompressedLayers[GenNavDataLayerTimeSlicedIdx].IsValid())
+			if (GenNavDataLayerIndices.IsEmpty())
 			{
-				// skip layers not marked for rebuild
-				continue;
+				PrepareZOrderedLayerIndices(GenNavDataLayerIndices);
 			}
 
-			if (TimeSliceManager->GetTimeSlicer().IsTimeSliceFinishedCached())
+			// Process layers in sorted order
+
+			for (; GenNavDataLayerTimeSlicedIdx < GenNavDataLayerIndices.Num(); GenNavDataLayerTimeSlicedIdx++)
 			{
-				WorkResult = ETimeSliceWorkResult::CallAgainNextTimeSlice;
-				break;
-			}
+				if (TimeSliceManager->GetTimeSlicer().IsTimeSliceFinishedCached())
+				{
+					WorkResult = ETimeSliceWorkResult::CallAgainNextTimeSlice;
+					break;
+				}
+
+				int32 LayerIdx = GenNavDataLayerIndices[GenNavDataLayerTimeSlicedIdx];
 
-			FTileCacheCompressor TileCompressor;
-			const bool bGenDataLayer = GenerateNavigationDataLayer(BuildContext, TileCompressor, *GenNavDataTimeSlicedAllocator,
-				*GenNavDataTimeSlicedGenerationContext, InLinkBuilderData, GenNavDataLayerTimeSlicedIdx);
+				FTileCacheCompressor TileCompressor;
+				const bool bGenDataLayerSuccess = GenerateNavigationDataLayerWithDeferredLinks(BuildContext, TileCompressor, *GenNavDataTimeSlicedAllocator, *GenNavDataTimeSlicedGenerationContext, InLinkBuilderData, LayerIdx, GenNavDataDeferredLinks);
 
-			MARK_TIMESLICE_SECTION_DEBUG(TimeSliceManager->GetTimeSlicer(), GenerateLayers);
+				MARK_TIMESLICE_SECTION_DEBUG(TimeSliceManager->GetTimeSlicer(), GenerateLayers);
 
-			//carry on iterating but don't do any more work if the time slice is finished (as we may not need to in which case we can avoid calling this function again)
-			TimeSliceManager->GetTimeSlicer().TestTimeSliceFinished();
+				//carry on iterating but don't do any more work if the time slice is finished (as we may not need to in which case we can avoid calling this function again)
+				TimeSliceManager->GetTimeSlicer().TestTimeSliceFinished();
 
-			if (!bGenDataLayer)
+				if (!bGenDataLayerSuccess)
+				{
+					WorkResult = ETimeSliceWorkResult::Failed;
+					break;
+				}
+			}
+		}
+		else
+		{
+			for (; GenNavDataLayerTimeSlicedIdx < CompressedLayers.Num() && GenNavDataLayerTimeSlicedIdx < DirtyLayers.Num(); GenNavDataLayerTimeSlicedIdx++)
 			{
-				WorkResult = ETimeSliceWorkResult::Failed;
-				break;
+				if (DirtyLayers[GenNavDataLayerTimeSlicedIdx] == false || !CompressedLayers[GenNavDataLayerTimeSlicedIdx].IsValid())
+				{
+					// skip layers not marked for rebuild
+					continue;
+				}
+
+				if (TimeSliceManager->GetTimeSlicer().IsTimeSliceFinishedCached())
+				{
+					WorkResult = ETimeSliceWorkResult::CallAgainNextTimeSlice;
+					break;
+				}
+
+				FTileCacheCompressor TileCompressor;
+				const bool bGenDataLayerSuccess = GenerateNavigationDataLayer(BuildContext, TileCompressor, *GenNavDataTimeSlicedAllocator,
+					*GenNavDataTimeSlicedGenerationContext, InLinkBuilderData, GenNavDataLayerTimeSlicedIdx);
+
+				MARK_TIMESLICE_SECTION_DEBUG(TimeSliceManager->GetTimeSlicer(), GenerateLayers);
+
+				//carry on iterating but don't do any more work if the time slice is finished (as we may not need to in which case we can avoid calling this function again)
+				TimeSliceManager->GetTimeSlicer().TestTimeSliceFinished();
+
+				if (!bGenDataLayerSuccess)
+				{
+					WorkResult = ETimeSliceWorkResult::Failed;
+					break;
+				}
 			}
 		}
 
@@ -4411,26 +4549,50 @@ bool FRecastTileGenerator::GenerateNavigationData(FNavMeshBuildContext& BuildCon
 	FTileGenerationContext GenerationContext(&GenNavAllocator);
 	GenerationContext.NavigationData.Reserve(CompressedLayers.Num());
 	FTileCacheCompressor TileCompressor;
-	bool bGenDataLayer = true;
-	dtStatus status = DT_SUCCESS;
+	bool bGenDataLayerSuccess = true;
 
-	for (int32 LayerIdx = 0; LayerIdx < CompressedLayers.Num() && LayerIdx < DirtyLayers.Num(); LayerIdx++)
+	const int32 NumLayers = FMath::Min(CompressedLayers.Num(), DirtyLayers.Num());
+
+	if (InLinkBuilderData.generatingLinks && NumLayers > 1)
 	{
-		if (DirtyLayers[LayerIdx] == false || !CompressedLayers[LayerIdx].IsValid())
-		{
-			// skip layers not marked for rebuild
-			continue;
-		}
+		// Make list of layer indices sorted by Z.
+		
+		TArray<int32> LayerIndices;
+		PrepareZOrderedLayerIndices(LayerIndices);
+
+		// Process the layers in Z order.
 
-		bGenDataLayer = GenerateNavigationDataLayer(BuildContext, TileCompressor, GenNavAllocator, GenerationContext, InLinkBuilderData, LayerIdx);
+		// DeferredLinks are populated with each call and results are passed onto subsequent calls as an input-output variable.
+		TInlineGeneratedNavLinkArray DeferredLinks;
 
-		if (!bGenDataLayer)
+		for (const int32 LayerIdx : LayerIndices)
 		{
-			break;
+			bGenDataLayerSuccess = GenerateNavigationDataLayerWithDeferredLinks(BuildContext, TileCompressor, GenNavAllocator, GenerationContext, InLinkBuilderData, LayerIdx, DeferredLinks);
+			if (!bGenDataLayerSuccess)
+			{
+				break;
+			}
+		}
+	}
+	else
+	{
+		for (int32 LayerIdx = 0; LayerIdx < CompressedLayers.Num() && LayerIdx < DirtyLayers.Num(); LayerIdx++)
+		{
+			if (DirtyLayers[LayerIdx] == false || !CompressedLayers[LayerIdx].IsValid())
+			{
+				// skip layers not marked for rebuild
+				continue;
+			}
+
+			bGenDataLayerSuccess = GenerateNavigationDataLayer(BuildContext, TileCompressor, GenNavAllocator, GenerationContext, InLinkBuilderData, LayerIdx);
+			if (!bGenDataLayerSuccess)
+			{
+				break;
+			}
 		}
 	}
 
-	if (bGenDataLayer)
+	if (bGenDataLayerSuccess)
 	{
 		NavigationData = MoveTemp(GenerationContext.NavigationData);
 	}
@@ -4440,7 +4602,7 @@ bool FRecastTileGenerator::GenerateNavigationData(FNavMeshBuildContext& BuildCon
 		GenerationContext.ResetIntermediateData();
 	}
 
-	return bGenDataLayer;
+	return bGenDataLayerSuccess;
 }
 
 void FRecastTileGenerator::ComputeRasterizationMasks(FNavMeshBuildContext& BuildContext, FTileRasterizationContext& RasterContext)
@@ -5657,6 +5819,15 @@ void FRecastNavMeshGenerator::EnsureBuildCompletion()
 	const bool bDoAsyncDataGathering = (GatherGeometryOnGameThread() == false);
 	do 
 	{
+#if !RECAST_ASYNC_REBUILDING
+		// When time slicing is enabled, we need to setup the time slice before processing tiles.
+		// Use a large duration since EnsureBuildCompletion wants to complete all work immediately.
+		if (SyncTimeSlicedData.TimeSliceManager && SyncTimeSlicedData.bTimeSliceRegenActive)
+		{
+			SyncTimeSlicedData.TimeSliceManager->GetTimeSlicer().SetupTimeSlice(TNumericLimits<double>::Max());
+		}
+#endif // !RECAST_ASYNC_REBUILDING
+
 		const int32 NumTasksToProcess = (bDoAsyncDataGathering ? 1 : MaxTileGeneratorTasks) - RunningDirtyTiles.Num();
 		ProcessTileTasksAndGetUpdatedTiles(NumTasksToProcess);
 		
diff --git a/Engine/Source/Runtime/NavigationSystem/Public/NavMesh/RecastNavMeshGenerator.h b/Engine/Source/Runtime/NavigationSystem/Public/NavMesh/RecastNavMeshGenerator.h
--- a/Engine/Source/Runtime/NavigationSystem/Public/NavMesh/RecastNavMeshGenerator.h
+++ b/Engine/Source/Runtime/NavigationSystem/Public/NavMesh/RecastNavMeshGenerator.h
@@ -353,7 +353,12 @@ struct FGeneratedNavigationLink : public FNavigationLink
 
 	/// User defined area ids assigned to the off-mesh connections
 	unsigned char generatedLinkArea = 0;
+
+	// What layer the link will be associated with
+	int16 TargetLayerIndex = INDEX_NONE;
 };
+typedef TArray<FGeneratedNavigationLink, TInlineAllocator<64>> TInlineGeneratedNavLinkArray;
+
 
 /**
  * Class handling generation of a single tile, caching data that can speed up subsequent tile generations
@@ -410,6 +415,12 @@ public:
 		
 	typedef TArray<int32, TInlineAllocator<4096>> TInlineMaskArray;
 
+private:
+
+	void PrepareZOrderedLayerIndices(TArray<int32>& OutLayerIndices) const;
+	bool GenerateNavigationDataLayerWithDeferredLinks(FNavMeshBuildContext& BuildContext, FTileCacheCompressor& TileCompressor, FTileCacheAllocator& GenNavAllocator, FTileGenerationContext& GenerationContext, const dtLinkBuilderData& InLinkBuilderData, int32 LayerIdx,
+		TInlineGeneratedNavLinkArray& OutDeferredLinks);
+
 protected:
 	/** Does the actual TimeSliced tile generation. 
 	 *	@note always trigger tile generation only via DoWorkTimeSliced(). This is a worker function
@@ -467,7 +478,8 @@ protected:
 	NAVIGATIONSYSTEM_API virtual bool GenerateCompressedLayers(FNavMeshBuildContext& BuildContext, const dtLinkBuilderData& InLinkBuilderData);
 
 	/** Builds a navigation data layer */
-	NAVIGATIONSYSTEM_API bool GenerateNavigationDataLayer(FNavMeshBuildContext& BuildContext, FTileCacheCompressor& TileCompressor, FTileCacheAllocator& GenNavAllocator, FTileGenerationContext& GenerationContext, const dtLinkBuilderData& InLinkBuilderData, int32 LayerIdx);
+	NAVIGATIONSYSTEM_API bool GenerateNavigationDataLayer(FNavMeshBuildContext& BuildContext, FTileCacheCompressor& TileCompressor, FTileCacheAllocator& GenNavAllocator, FTileGenerationContext& GenerationContext, const dtLinkBuilderData& InLinkBuilderData, int32 LayerIdx, 
+		const TInlineGeneratedNavLinkArray* InAdditionalLinks = nullptr, TInlineGeneratedNavLinkArray* OutGeneratedLinks = nullptr);
 
 	/** Builds NavigationData array (layers + obstacles) time sliced */
 	NAVIGATIONSYSTEM_API ETimeSliceWorkResult GenerateNavigationDataTimeSliced(FNavMeshBuildContext& BuildContext, const dtLinkBuilderData& InLinkBuilderData);
@@ -477,7 +489,7 @@ protected:
 
 	/** Builds navigation links */
 	dtStatus BuildTileCacheLinks(FNavMeshBuildContext& BuildContext, struct dtTileCacheAlloc* alloc, const dtTileCacheLayer& layer,
-		const struct dtTileCacheContourSet& lcset, TArray<FGeneratedNavigationLink>& OutGeneratedLinks) const;
+		const struct dtTileCacheContourSet& lcset, const float DefaultSnapHeight, TInlineGeneratedNavLinkArray& OutGeneratedLinks) const;
 	
 	NAVIGATIONSYSTEM_API virtual void ApplyVoxelFilter(struct rcHeightfield* SolidHF, FVector::FReal WalkableRadius);
 
@@ -547,6 +559,9 @@ protected:
 	TUniquePtr<struct FTileGenerationContext> GenNavDataTimeSlicedGenerationContext;
 	TUniquePtr<struct FTileRasterizationContext> GenCompressedlayersTimeSlicedRasterContext;
 
+	TArray<int32> GenNavDataLayerIndices;
+	TInlineGeneratedNavLinkArray GenNavDataDeferredLinks;
+
 	FRecastTileTimeSliceSettings TileTimeSliceSettings;
 	/** End time slicing variables */
 
diff --git a/Engine/Source/Runtime/Navmesh/Public/Detour/DetourNavMeshBuilder.h b/Engine/Source/Runtime/Navmesh/Public/Detour/DetourNavMeshBuilder.h
--- a/Engine/Source/Runtime/Navmesh/Public/Detour/DetourNavMeshBuilder.h
+++ b/Engine/Source/Runtime/Navmesh/Public/Detour/DetourNavMeshBuilder.h
@@ -128,7 +128,7 @@ struct dtNavMeshCreateParams
 
 	unsigned int userId;	///< The user defined id of the tile.
 	int tileX;				///< The tile's x-grid location within the multi-tile destination mesh. (Along the x-axis.)
-	int tileY;				///< The tile's y-grid location within the multi-tile desitation mesh. (Along the z-axis.)
+	int tileY;				///< The tile's y-grid location within the multi-tile destination mesh. (Along the z-axis.)
 	int tileLayer;			///< The tile's layer within the layered destination mesh. (Along the y-axis.) [Limit: >= 0]
 	dtReal bmin[3];			///< The minimum bounds of the tile. [(x, y, z)] [Unit: wu]
 	dtReal bmax[3];			///< The maximum bounds of the tile. [(x, y, z)] [Unit: wu]
