Exposing GetLengthFromPosition to Blueprints.

#Synced-CL 50139823

diff --git a/Engine/Source/Runtime/NavigationSystem/Private/NavigationPath.cpp b/Engine/Source/Runtime/NavigationSystem/Private/NavigationPath.cpp
--- a/Engine/Source/Runtime/NavigationSystem/Private/NavigationPath.cpp
+++ b/Engine/Source/Runtime/NavigationSystem/Private/NavigationPath.cpp
@@ -539,6 +539,24 @@ double UNavigationPath::GetPathLength() const
 	return !!bIsValid ? SharedPath->GetLength() : -1.;
 }
 
+double UNavigationPath::GetLengthFromPosition(FVector Position, int32 NextPathPointIndex) const
+{
+	check((SharedPath.IsValid() && SharedPath->IsValid()) == !!bIsValid);
+
+	if (!bIsValid)
+	{
+		return -1.;
+	}
+
+	if (NextPathPointIndex < 0)
+	{
+		UE_LOG(LogNavigation, Warning, TEXT("Negative PathPointIndex passed in (%d)."), NextPathPointIndex);
+		return 0.;
+	}
+
+	return SharedPath->GetLengthFromPosition(Position, NextPathPointIndex);
+}
+
 double UNavigationPath::GetPathCost() const
 {
 	check((SharedPath.IsValid() && SharedPath->IsValid()) == !!bIsValid);
diff --git a/Engine/Source/Runtime/NavigationSystem/Public/NavigationPath.h b/Engine/Source/Runtime/NavigationSystem/Public/NavigationPath.h
--- a/Engine/Source/Runtime/NavigationSystem/Public/NavigationPath.h
+++ b/Engine/Source/Runtime/NavigationSystem/Public/NavigationPath.h
@@ -66,6 +66,10 @@ public:
 	UFUNCTION(BlueprintCallable, Category = "AI|Navigation")
 	NAVIGATIONSYSTEM_API double GetPathLength() const;
 
+	/** returns the remaining length of the path from Position, including the distance from Position to the NextPathPoint */
+	UFUNCTION(BlueprintCallable, Category = "AI|Navigation")
+	NAVIGATIONSYSTEM_API double GetLengthFromPosition(FVector Position, int32 NextPathPointIndex) const;
+
 	UFUNCTION(BlueprintCallable, Category = "AI|Navigation")
 	NAVIGATIONSYSTEM_API double GetPathCost() const;
 
