Addition of the navdata name when logging navdata registration.
Addition a log when trying to build navigation when it's prevented by a navigation lock.
#Synced-CL 50036996

diff --git a/Engine/Source/Runtime/NavigationSystem/Private/NavigationSystem.cpp b/Engine/Source/Runtime/NavigationSystem/Private/NavigationSystem.cpp
--- a/Engine/Source/Runtime/NavigationSystem/Private/NavigationSystem.cpp
+++ b/Engine/Source/Runtime/NavigationSystem/Private/NavigationSystem.cpp
@@ -1423,7 +1423,7 @@ void UNavigationSystemV1::OnWorldInitDone(FNavigationSystemRunMode Mode)
 					if (NavData != nullptr)
 					{
 						const ERegistrationResult Result = RegisterNavData(NavData);
-						LogNavDataRegistrationResult(Result);
+						LogNavDataRegistrationResult(Result, NavData);
 
 						if (Result == RegistrationSuccessful)
 						{
@@ -2519,7 +2519,7 @@ ANavigationData* UNavigationSystemV1::GetDefaultNavDataInstance(FNavigationSyste
 		if (MainNavData != nullptr)
 		{
 			const ERegistrationResult Result = RegisterNavData(MainNavData);
-			LogNavDataRegistrationResult(Result);
+			LogNavDataRegistrationResult(Result, MainNavData);
 		}
 	}
 
@@ -2759,10 +2759,13 @@ void UNavigationSystemV1::ProcessRegistrationCandidates()
 		if (OwningLevel && OwningLevel->bIsVisible)
 		{
 			const ERegistrationResult Result = RegisterNavData(NavDataPtr);
-			LogNavDataRegistrationResult(Result);
+			LogNavDataRegistrationResult(Result, NavDataPtr);
 
 			if (Result != RegistrationSuccessful && Result != RegistrationFailed_DataPendingKill)
 			{
+				UE_LOG(LogNavigation, Warning, TEXT("%hs %s will be removed. This warning will disapear once the map is saved."), 
+					__FUNCTION__, *GetNameSafe(NavDataPtr));
+				
 #if WITH_EDITOR				
 				if (GEditor)
 				{
@@ -4427,6 +4430,7 @@ void UNavigationSystemV1::Build()
 	const bool bLockedIgnoreEditor = (NavBuildingLockFlags & ~ENavigationBuildLock::NoUpdateInEditor) != 0;
 	if (!bHasWork || bLockedIgnoreEditor)
 	{
+		UE_LOG(LogNavigation, Warning, TEXT("%hs Navigation NOT building because navigation build is locked (flags: 0x%x)."), __FUNCTION__, NavBuildingLockFlags);
 		return;
 	}
 
@@ -5259,30 +5263,43 @@ ERuntimeGenerationType UNavigationSystemV1::GetRuntimeGenerationType() const
 	return RuntimeGenerationType;
 }
 
+// Deprecated
 void UNavigationSystemV1::LogNavDataRegistrationResult(ERegistrationResult InResult)
+{
+	LogNavDataRegistrationResult(InResult, nullptr);
+}
+
+void UNavigationSystemV1::LogNavDataRegistrationResult(ERegistrationResult InResult, const ANavigationData* NavigationData) const
 {
 	switch (InResult)
 	{
 	case UNavigationSystemV1::RegistrationError:
-		UE_VLOG_UELOG(this, LogNavigation, Warning, TEXT("NavData RegistrationError, could not be registered."));
+		UE_VLOG_UELOG(this, LogNavigation, Warning,
+			TEXT("NavData RegistrationError, could not be registered (%s)."), *GetNameSafe(NavigationData));
 		break;
 	case UNavigationSystemV1::RegistrationFailed_DataPendingKill:
-		UE_VLOG_UELOG(this, LogNavigation, Warning, TEXT("NavData RegistrationFailed_DataPendingKill."));
+		UE_VLOG_UELOG(this, LogNavigation, Warning,
+			TEXT("NavData RegistrationFailed_DataPendingKill (%s)."), *GetNameSafe(NavigationData));
 		break;
 	case UNavigationSystemV1::RegistrationFailed_AgentAlreadySupported:
-		UE_VLOG_UELOG(this, LogNavigation, Warning, TEXT("NavData RegistrationFailed_AgentAlreadySupported, specified agent type already has its navmesh implemented."));
+		UE_VLOG_UELOG(this, LogNavigation, Warning,
+			TEXT("NavData RegistrationFailed_AgentAlreadySupported, specified agent type already has its navmesh implemented (%s)."), *GetNameSafe(NavigationData));
 		break;
 	case UNavigationSystemV1::RegistrationFailed_AgentNotValid:
-		UE_VLOG_UELOG(this, LogNavigation, Warning, TEXT("NavData RegistrationFailed_AgentNotValid, NavData instance contains navmesh that doesn't support any of expected agent types."));
+		UE_VLOG_UELOG(this, LogNavigation, Warning,
+			TEXT("NavData RegistrationFailed_AgentNotValid, NavData instance contains navmesh that doesn't support any of expected agent types (%s)."), *GetNameSafe(NavigationData));
 		break;
 	case UNavigationSystemV1::RegistrationFailed_NotSuitable:
-		UE_VLOG_UELOG(this, LogNavigation, Warning, TEXT("NavData RegistrationFailed_NotSuitable."));
+		UE_VLOG_UELOG(this, LogNavigation, Warning,
+			TEXT("NavData RegistrationFailed_NotSuitable (%s)."), *GetNameSafe(NavigationData));
 		break;
 	case UNavigationSystemV1::RegistrationSuccessful:
-		UE_VLOG_UELOG(this, LogNavigation, Verbose, TEXT("NavData RegistrationSuccessful."));
+		UE_VLOG_UELOG(this, LogNavigation, Verbose,
+			TEXT("NavData RegistrationSuccessful (%s)."), *GetNameSafe(NavigationData));
 		break;
 	default:
-		UE_VLOG_UELOG(this, LogNavigation, Warning, TEXT("Registration not successful default warning."));
+		UE_VLOG_UELOG(this, LogNavigation, Warning,
+			TEXT("Registration not successful default warning (%s)."), *GetNameSafe(NavigationData));
 		break;
 	}
 }
diff --git a/Engine/Source/Runtime/NavigationSystem/Public/NavigationSystem.h b/Engine/Source/Runtime/NavigationSystem/Public/NavigationSystem.h
--- a/Engine/Source/Runtime/NavigationSystem/Public/NavigationSystem.h
+++ b/Engine/Source/Runtime/NavigationSystem/Public/NavigationSystem.h
@@ -1441,7 +1441,10 @@ protected:
 	 */
 	NAVIGATIONSYSTEM_API ERuntimeGenerationType GetRuntimeGenerationType() const;
 
+	UE_DEPRECATED(5.8, "Use LogNavDataRegistrationResult receiving ANavigationData* instead.")
 	NAVIGATIONSYSTEM_API void LogNavDataRegistrationResult(ERegistrationResult);
+	
+	NAVIGATIONSYSTEM_API void LogNavDataRegistrationResult(ERegistrationResult, const ANavigationData* NavigationData) const;
 
 	/** Whether Navigation System is allowed to rebuild the navmesh
 	 *  Depends on runtime generation settings of each navigation data, always true in the editor
