[Navigation] Addition of logs to record AsyncLoadLock.
Increased verbosity on navigation test to help identify hard to reproduce issue.

#Synced-CL 50036996

diff --git a/Engine/Source/Runtime/NavigationSystem/Private/NavigationSystem.cpp b/Engine/Source/Runtime/NavigationSystem/Private/NavigationSystem.cpp
--- a/Engine/Source/Runtime/NavigationSystem/Private/NavigationSystem.cpp
+++ b/Engine/Source/Runtime/NavigationSystem/Private/NavigationSystem.cpp
@@ -1048,6 +1048,8 @@ void UNavigationSystemV1::DoInitialSetup()
 	{
 		return;
 	}
+
+	UE_LOG(LogNavigation, Verbose, TEXT("%hs"), __FUNCTION__);
 	
 	UpdateAbstractNavData();
 	CreateCrowdManager();
@@ -1057,6 +1059,7 @@ void UNavigationSystemV1::DoInitialSetup()
 #if WITH_EDITOR		
 	if (bWaitForAsyncLoadingBeforeBuildingNavigationAutomatically && GetIsAutoUpdateEnabled() && OperationMode == FNavigationSystemRunMode::EditorMode)
 	{
+		UE_LOG(LogNavigation, Verbose, TEXT("%hs    Adding ENavigationBuildLock::AsyncLoadLock."), __FUNCTION__);
 		AddNavigationBuildLock(ENavigationBuildLock::AsyncLoadLock);
 
 		FAssetCompilingManager::Get().OnAssetPostCompileEvent().AddUObject(this, &UNavigationSystemV1::UpdateDelayedUnlockeRegistration);
@@ -5840,6 +5843,8 @@ void UNavigationSystemV1::UnregisterFromRepositoryDelegates() const
 #if WITH_EDITOR
 void UNavigationSystemV1::UpdateDelayedUnlockeRegistration(const TArray<FAssetCompileData>& AssetCompileData)
 {
+	UE_LOG(LogNavigation, Verbose, TEXT("%hs Assests remaining %i."), __FUNCTION__, FAssetCompilingManager::Get().GetNumRemainingAssets());
+	
 	// Register a delayed unlock navigation when all assets are finished compiling.
 	if (FAssetCompilingManager::Get().GetNumRemainingAssets() == 0)
 	{
@@ -5849,17 +5854,19 @@ void UNavigationSystemV1::UpdateDelayedUnlockeRegistration(const TArray<FAssetCo
 	}
 	else
 	{
-		// We're not idle anymore, unregister unlocink until we get to 0 assets again
+		// We're not idle anymore, unregister unlocking until we get to 0 assets again
 		UnregisterDelayedUnlock();
 	}
 }
 
 void UNavigationSystemV1::RegisterDelayedUnlock(float DelayInSeconds)
 {
+	UE_LOG(LogNavigation, Verbose, TEXT("%hs"), __FUNCTION__);
+	
 	// Unregister the previous one and register a new one.
 	UnregisterDelayedUnlock();
 
-	// The editor might be unresponsive during heavy asset compilation so we 
+	// The editor might be unresponsive during heavy asset compilation, so we 
 	// not only need a delay, but also a minimum amount of frames
 	// to pass until we call the actual refresh.
 	DelayedUnlockHandle = FTSTicker::GetCoreTicker().AddTicker(
@@ -5867,9 +5874,13 @@ void UNavigationSystemV1::RegisterDelayedUnlock(float DelayInSeconds)
 		0.f,
 		[this, FireInTickCount = 16, DelayInSeconds](float DeltaTime) mutable
 		{
+			UE_LOG(LogNavigation, Verbose, TEXT("%hs    Time remaining %f."), __FUNCTION__, DelayInSeconds);
+				
 			DelayInSeconds -= DeltaTime;
 			if (--FireInTickCount <= 0 && DelayInSeconds <= 0.f && FAssetCompilingManager::Get().GetNumRemainingAssets() == 0)
 			{
+				UE_LOG(LogNavigation, Verbose, TEXT("%hs    Done with assets. Removing ENavigationBuildLock::AsyncLoadLock."), __FUNCTION__);
+				
 				FAssetCompilingManager::Get().OnAssetPostCompileEvent().RemoveAll(this);
 				
 				// Cleared initial async load dirty areas and remove AsyncLoadLock
@@ -5888,6 +5899,8 @@ void UNavigationSystemV1::UnregisterDelayedUnlock()
 {
 	if (DelayedUnlockHandle.IsValid())
 	{
+		UE_LOG(LogNavigation, Verbose, TEXT("%hs"), __FUNCTION__);
+		
 		FTSTicker::RemoveTicker(DelayedUnlockHandle);
 		DelayedUnlockHandle.Reset();
 	}
