Update navigation tests waiting for the map to load to only wait for the AsyncLoadLock (else they can wait indefinitely on the NoUpdateInEditor lock if the setting is disabled).
Added logs and debug display information
#rnx

#Synced-CL 50193655

diff --git a/Engine/Source/Runtime/NavigationSystem/Private/NavigationSystem.cpp b/Engine/Source/Runtime/NavigationSystem/Private/NavigationSystem.cpp
--- a/Engine/Source/Runtime/NavigationSystem/Private/NavigationSystem.cpp
+++ b/Engine/Source/Runtime/NavigationSystem/Private/NavigationSystem.cpp
@@ -4742,7 +4742,9 @@ void UNavigationSystemV1::AddNavigationBuildLock(uint8 Flags)
 	NavBuildingLockFlags |= Flags;
 
 	const bool bIsLocked = IsNavigationBuildingLocked();
-	UE_LOG(LogNavigation, Verbose, TEXT("UNavigationSystemV1::AddNavigationBuildLock WasLocked=%s IsLocked=%s"), *LexToString(bWasLocked), *LexToString(bIsLocked));
+	UE_LOG(LogNavigation, Verbose, TEXT("UNavigationSystemV1::AddNavigationBuildLock WasLocked=%s IsLocked=%s (flags: 0x%x)"),
+		*LexToString(bWasLocked), *LexToString(bIsLocked), NavBuildingLockFlags);
+	
 	if (!bWasLocked && bIsLocked)
 	{
 		DefaultDirtyAreasController.OnNavigationBuildLocked();
@@ -4756,7 +4758,8 @@ void UNavigationSystemV1::RemoveNavigationBuildLock(uint8 Flags, const ELockRemo
 	NavBuildingLockFlags &= ~Flags;
 
 	const bool bIsLocked = IsNavigationBuildingLocked();
-	UE_LOG(LogNavigation, Verbose, TEXT("UNavigationSystemV1::RemoveNavigationBuildLock WasLocked=%s IsLocked=%s"), *LexToString(bWasLocked), *LexToString(bIsLocked));
+	UE_LOG(LogNavigation, Verbose, TEXT("UNavigationSystemV1::RemoveNavigationBuildLock WasLocked=%s IsLocked=%s (flags: 0x%x)"),
+		*LexToString(bWasLocked), *LexToString(bIsLocked), NavBuildingLockFlags);
 	if (bWasLocked && !bIsLocked)
 	{
 		DefaultDirtyAreasController.OnNavigationBuildUnlocked();
@@ -5140,6 +5143,7 @@ void UNavigationSystemV1::CleanUp(FNavigationSystem::ECleanupMode Mode)
 		GEngine->OnActorMoved().RemoveAll(this);
 	}
 
+	UE_LOG(LogNavigation, Verbose, TEXT("%hs Removing functions from OnAssetPostCompileEvent delegate."), __FUNCTION__);
 	FAssetCompilingManager::Get().OnAssetPostCompileEvent().RemoveAll(this);
 	UnregisterDelayedUnlock();
 #endif // WITH_EDITOR
@@ -5879,11 +5883,11 @@ void UNavigationSystemV1::RegisterDelayedUnlock(float DelayInSeconds)
 			DelayInSeconds -= DeltaTime;
 			if (--FireInTickCount <= 0 && DelayInSeconds <= 0.f && FAssetCompilingManager::Get().GetNumRemainingAssets() == 0)
 			{
-				UE_LOG(LogNavigation, Verbose, TEXT("%hs    Done with assets. Removing ENavigationBuildLock::AsyncLoadLock."), __FUNCTION__);
-				
+				UE_LOG(LogNavigation, Verbose, TEXT("%hs    Removing functions from OnAssetPostCompileEvent delegate."), __FUNCTION__);
 				FAssetCompilingManager::Get().OnAssetPostCompileEvent().RemoveAll(this);
 				
 				// Cleared initial async load dirty areas and remove AsyncLoadLock
+				UE_LOG(LogNavigation, Verbose, TEXT("%hs    Done with assets. Removing ENavigationBuildLock::AsyncLoadLock."), __FUNCTION__);
 				DefaultDirtyAreasController.Reset();
 				RemoveNavigationBuildLock(ENavigationBuildLock::AsyncLoadLock, ELockRemovalRebuildAction::NoRebuild);
 				
