Fixed automatic navlink proxies not getting registered. This is a fix for an issue introduced in CL 49450039.

#rnx

#Synced-CL 50434970

diff --git a/Engine/Source/Runtime/NavigationSystem/Private/NavMesh/RecastNavMesh.cpp b/Engine/Source/Runtime/NavigationSystem/Private/NavMesh/RecastNavMesh.cpp
--- a/Engine/Source/Runtime/NavigationSystem/Private/NavMesh/RecastNavMesh.cpp
+++ b/Engine/Source/Runtime/NavigationSystem/Private/NavMesh/RecastNavMesh.cpp
@@ -61,6 +61,9 @@
 
 namespace UE::NavMesh::Private
 {
+	static bool bAllowLinkGeneration = true;
+	static FAutoConsoleVariableRef CVarAllowLinkGeneration(TEXT("ai.nav.AllowLinkGeneration"), bAllowLinkGeneration, TEXT("Set to false to force disabling link generation."), ECVF_Default);
+
 	// Max tile size in voxels. Larger than this tiles will start to get slow to build.
 	constexpr int32 ArbitraryMaxTileSizeVoxels = 1024;
 	// Min tile size on voxels. Smaller tiles than this waste computation during voxelization because the border are will be larger than usable area.
@@ -741,6 +744,11 @@ void ARecastNavMesh::CheckToDiscardSubLevelNavData(const UNavigationSystemBase&
 	}
 }
 
+bool ARecastNavMesh::IsGeneratingLinks() const
+{
+	return bGenerateNavLinks && UE::NavMesh::Private::bAllowLinkGeneration;
+}
+
 void ARecastNavMesh::RegisterGeneratedLinksProxy()
 {
 	for (FNavLinkGenerationJumpConfig& NavLinkJumpConfig : NavLinkJumpConfigs)
@@ -838,8 +846,7 @@ void ARecastNavMesh::PostRegisterAllComponents()
 
 	// Create and register link proxy if enabled.
 
-	const FRecastNavMeshGenerator* MyGenerator = static_cast<FRecastNavMeshGenerator*>(GetGenerator());
-	if (MyGenerator && MyGenerator->IsGeneratingLinks())
+	if (IsGeneratingLinks())
 	{
 		for (FNavLinkGenerationJumpConfig& NavLinkJumpConfig : NavLinkJumpConfigs)
 		{
diff --git a/Engine/Source/Runtime/NavigationSystem/Private/NavMesh/RecastNavMeshGenerator.cpp b/Engine/Source/Runtime/NavigationSystem/Private/NavMesh/RecastNavMeshGenerator.cpp
--- a/Engine/Source/Runtime/NavigationSystem/Private/NavMesh/RecastNavMeshGenerator.cpp
+++ b/Engine/Source/Runtime/NavigationSystem/Private/NavMesh/RecastNavMeshGenerator.cpp
@@ -87,10 +87,6 @@ namespace UE::NavMesh::Private
 
 	static bool bUseAsymetricBorderSizes = false;
 	static FAutoConsoleVariableRef CVarUseAsymetricBorderSizes(TEXT("ai.nav.UseAsymetricBorderSizes"), bUseAsymetricBorderSizes, TEXT("Disabled by default. When generating links, use asymetric tile border sizes to improve generation speed."), ECVF_Default);
-
-	static bool bAllowLinkGeneration = true;
-	static FAutoConsoleVariableRef CVarAllowLinkGeneration(TEXT("ai.nav.AllowLinkGeneration"), bAllowLinkGeneration, TEXT("Set to false to force disabling link generation."), ECVF_Default);
-
 }
 
 static FOodleDataCompression::ECompressor GNavmeshTileCacheCompressor = FOodleDataCompression::ECompressor::Mermaid;
@@ -5703,7 +5699,7 @@ void FRecastNavMeshGenerator::CalcPolyRefBits(ARecastNavMesh* NavMeshOwner, int3
 
 bool FRecastNavMeshGenerator::IsGeneratingLinks() const
 {
-	return DestNavMesh && DestNavMesh->bGenerateNavLinks && UE::NavMesh::Private::bAllowLinkGeneration; 
+	return DestNavMesh && DestNavMesh->IsGeneratingLinks(); 
 }
 
 void FRecastNavMeshGenerator::ResolveGeneratedLinkAreas(FRecastBuildConfig& OutConfig)
diff --git a/Engine/Source/Runtime/NavigationSystem/Public/NavMesh/RecastNavMesh.h b/Engine/Source/Runtime/NavigationSystem/Public/NavMesh/RecastNavMesh.h
--- a/Engine/Source/Runtime/NavigationSystem/Public/NavMesh/RecastNavMesh.h
+++ b/Engine/Source/Runtime/NavigationSystem/Public/NavMesh/RecastNavMesh.h
@@ -1518,6 +1518,7 @@ protected:
 
 	NAVIGATIONSYSTEM_API void CheckToDiscardSubLevelNavData(const UNavigationSystemBase& NavSys);
 	
+	bool IsGeneratingLinks() const;
 	void RegisterGeneratedLinksProxy();
 	void UnregisterGeneratedLinksProxy();
 
