[Navigation] Unregister delegate using a handle instead of removing them all.

#rnx

#Synced-CL 50642529

diff --git a/Engine/Source/Runtime/NavigationSystem/Private/NavigationSystem.cpp b/Engine/Source/Runtime/NavigationSystem/Private/NavigationSystem.cpp
--- a/Engine/Source/Runtime/NavigationSystem/Private/NavigationSystem.cpp
+++ b/Engine/Source/Runtime/NavigationSystem/Private/NavigationSystem.cpp
@@ -1062,7 +1062,7 @@ void UNavigationSystemV1::DoInitialSetup()
 		UE_LOG(LogNavigation, Verbose, TEXT("%hs    Adding ENavigationBuildLock::AsyncLoadLock."), __FUNCTION__);
 		AddNavigationBuildLock(ENavigationBuildLock::AsyncLoadLock);
 
-		FAssetCompilingManager::Get().OnAssetPostCompileEvent().AddUObject(this, &UNavigationSystemV1::UpdateDelayedUnlockeRegistration);
+		OnAssetPostCompileEventHandle = FAssetCompilingManager::Get().OnAssetPostCompileEvent().AddUObject(this, &UNavigationSystemV1::UpdateDelayedUnlockeRegistration);
 
 		// Run it at least once to make sure we don't stay locked.
 		UpdateDelayedUnlockeRegistration({});
@@ -5144,7 +5144,8 @@ void UNavigationSystemV1::CleanUp(FNavigationSystem::ECleanupMode Mode)
 	}
 
 	UE_LOG(LogNavigation, Verbose, TEXT("%hs Removing functions from OnAssetPostCompileEvent delegate."), __FUNCTION__);
-	FAssetCompilingManager::Get().OnAssetPostCompileEvent().RemoveAll(this);
+	FAssetCompilingManager::Get().OnAssetPostCompileEvent().Remove(OnAssetPostCompileEventHandle);
+	OnAssetPostCompileEventHandle.Reset();
 	UnregisterDelayedUnlock();
 #endif // WITH_EDITOR
 
diff --git a/Engine/Source/Runtime/NavigationSystem/Public/NavigationSystem.h b/Engine/Source/Runtime/NavigationSystem/Public/NavigationSystem.h
--- a/Engine/Source/Runtime/NavigationSystem/Public/NavigationSystem.h
+++ b/Engine/Source/Runtime/NavigationSystem/Public/NavigationSystem.h
@@ -495,6 +495,7 @@ private:
 	void RegisterDelayedUnlock(float DelayInSeconds);
 	void UnregisterDelayedUnlock();
 
+	FDelegateHandle OnAssetPostCompileEventHandle;
 	FTSTicker::FDelegateHandle DelayedUnlockHandle;
 #endif // WITH_EDITOR
 
