[Navigation navlinks] Fix multi-config registration for generated links.

#Synced-CL 50642529

diff --git a/Engine/Source/Runtime/NavigationSystem/Private/NavMesh/RecastNavMesh.cpp b/Engine/Source/Runtime/NavigationSystem/Private/NavMesh/RecastNavMesh.cpp
--- a/Engine/Source/Runtime/NavigationSystem/Private/NavMesh/RecastNavMesh.cpp
+++ b/Engine/Source/Runtime/NavigationSystem/Private/NavMesh/RecastNavMesh.cpp
@@ -749,47 +749,85 @@ bool ARecastNavMesh::IsGeneratingLinks() const
 	return bGenerateNavLinks && UE::NavMesh::Private::bAllowLinkGeneration;
 }
 
+void ARecastNavMesh::RegisterGeneratedLinksProxy(FNavLinkGenerationJumpConfig& InOutNavLinkJumpConfig)
+{
+	if (!InOutNavLinkJumpConfig.bLinkProxyRegistered && InOutNavLinkJumpConfig.LinkProxy)
+	{
+		if (UWorld* World = GetWorld())
+		{
+			UNavigationSystemV1* NavSys = FNavigationSystem::GetCurrent<UNavigationSystemV1>(World);
+			if (NavSys)
+			{
+				UE_LOG(LogNavLink, Log, TEXT("RegisterLinkProxy id: %llu ptr: 0x%p"), InOutNavLinkJumpConfig.LinkProxyId.GetId(), InOutNavLinkJumpConfig.LinkProxy.Get());
+
+				NavSys->RegisterCustomLink(*InOutNavLinkJumpConfig.LinkProxy);
+				InOutNavLinkJumpConfig.bLinkProxyRegistered = true;
+			}
+		}
+	}
+}
+
+void ARecastNavMesh::UnregisterGeneratedLinksProxy(FNavLinkGenerationJumpConfig& InOutNavLinkJumpConfig)
+{
+	if (InOutNavLinkJumpConfig.bLinkProxyRegistered && InOutNavLinkJumpConfig.LinkProxy)
+	{
+		if (UWorld* World = GetWorld())
+		{
+			UNavigationSystemV1* NavSys = FNavigationSystem::GetCurrent<UNavigationSystemV1>(World);
+			if (NavSys)
+			{
+				UE_LOG(LogNavLink, Log, TEXT("UnregisterLinkProxy id: %llu ptr: 0x%p"), InOutNavLinkJumpConfig.LinkProxyId.GetId(), InOutNavLinkJumpConfig.LinkProxy.Get());
+
+				NavSys->UnregisterCustomLink(*InOutNavLinkJumpConfig.LinkProxy);
+				InOutNavLinkJumpConfig.bLinkProxyRegistered = false;
+			}
+		}
+	}
+}
+
+PRAGMA_DISABLE_DEPRECATION_WARNINGS
+// Deprecated
 void ARecastNavMesh::RegisterGeneratedLinksProxy()
 {
-	for (FNavLinkGenerationJumpConfig& NavLinkJumpConfig : NavLinkJumpConfigs)
+#if	WITH_EDITOR	
+	if (!NavLinkJumpDownConfig.bLinkProxyRegistered && NavLinkJumpDownConfig.LinkProxy)
 	{
-		if (!NavLinkJumpConfig.bLinkProxyRegistered && NavLinkJumpConfig.LinkProxy)
+		if (UWorld* World = GetWorld())
 		{
-			if (UWorld* World = GetWorld())
+			UNavigationSystemV1* NavSys = FNavigationSystem::GetCurrent<UNavigationSystemV1>(World);
+			if (NavSys)
 			{
-				UNavigationSystemV1* NavSys = FNavigationSystem::GetCurrent<UNavigationSystemV1>(World);
-				if (NavSys)
-				{
-					UE_LOG(LogNavLink, Log, TEXT("RegisterLinkProxy id: %llu ptr: 0x%p"), NavLinkJumpConfig.LinkProxyId.GetId(), NavLinkJumpConfig.LinkProxy.Get());
+				UE_LOG(LogNavLink, Log, TEXT("RegisterLinkProxy id: %llu ptr: 0x%p"), NavLinkJumpDownConfig.LinkProxyId.GetId(), NavLinkJumpDownConfig.LinkProxy.Get());
 
-					NavSys->RegisterCustomLink(*NavLinkJumpConfig.LinkProxy);
-					NavLinkJumpConfig.bLinkProxyRegistered = true;
-				}
+				NavSys->RegisterCustomLink(*NavLinkJumpDownConfig.LinkProxy);
+				NavLinkJumpDownConfig.bLinkProxyRegistered = true;
 			}
 		}
 	}
+#endif // WITH_EDITOR
 }
 
+// Deprecated
 void ARecastNavMesh::UnregisterGeneratedLinksProxy()
 {
-	for (FNavLinkGenerationJumpConfig& NavLinkJumpConfig : NavLinkJumpConfigs)
+#if	WITH_EDITOR
+	if (NavLinkJumpDownConfig.bLinkProxyRegistered && NavLinkJumpDownConfig.LinkProxy)
 	{
-		if (NavLinkJumpConfig.bLinkProxyRegistered && NavLinkJumpConfig.LinkProxy)
+		if (UWorld* World = GetWorld())
 		{
-			if (UWorld* World = GetWorld())
+			UNavigationSystemV1* NavSys = FNavigationSystem::GetCurrent<UNavigationSystemV1>(World);
+			if (NavSys)
 			{
-				UNavigationSystemV1* NavSys = FNavigationSystem::GetCurrent<UNavigationSystemV1>(World);
-				if (NavSys)
-				{
-					UE_LOG(LogNavLink, Log, TEXT("UnregisterLinkProxy id: %llu ptr: 0x%p"), NavLinkJumpConfig.LinkProxyId.GetId(), NavLinkJumpConfig.LinkProxy.Get());
+				UE_LOG(LogNavLink, Log, TEXT("UnregisterLinkProxy id: %llu ptr: 0x%p"), NavLinkJumpDownConfig.LinkProxyId.GetId(), NavLinkJumpDownConfig.LinkProxy.Get());
 
-					NavSys->UnregisterCustomLink(*NavLinkJumpConfig.LinkProxy);
-					NavLinkJumpConfig.bLinkProxyRegistered = false;
-				}
+				NavSys->UnregisterCustomLink(*NavLinkJumpDownConfig.LinkProxy);
+				NavLinkJumpDownConfig.bLinkProxyRegistered = false;
 			}
 		}
 	}
+#endif // WITH_EDITOR	
 }
+PRAGMA_ENABLE_DEPRECATION_WARNINGS
 
 // Deprecated
 void ARecastNavMesh::CreateAndRegisterJumpDownLinksProxy(const FNavLinkId LinkProxyId /*= FNavLinkId::GenerateUniqueId()*/)
@@ -831,7 +869,7 @@ void ARecastNavMesh::CreateAndRegisterJumpLinksProxy(FNavLinkGenerationJumpConfi
 		LinkProxy->SetOwner(this);
 		InOutNavLinkJumpConfig.LinkProxy = LinkProxy;
 
-		RegisterGeneratedLinksProxy();
+		RegisterGeneratedLinksProxy(InOutNavLinkJumpConfig);
 	}
 }
 
@@ -864,10 +902,15 @@ void ARecastNavMesh::PostUnregisterAllComponents()
 {
 	Super::PostUnregisterAllComponents();
 
+PRAGMA_DISABLE_DEPRECATION_WARNINGS
+	// Deprecated
 	UnregisterGeneratedLinksProxy();
+PRAGMA_ENABLE_DEPRECATION_WARNINGS
 
 	for (FNavLinkGenerationJumpConfig& NavLinkJumpConfig : NavLinkJumpConfigs)
 	{
+		UnregisterGeneratedLinksProxy(NavLinkJumpConfig);
+		
 		NavLinkJumpConfig.LinkProxy = nullptr;
 	}
 }
@@ -3945,11 +3988,17 @@ bool ARecastNavMesh::IsUsingActiveTilesGeneration(const UNavigationSystemV1& Nav
 
 void ARecastNavMesh::PostLoadPreRebuild()
 {
+PRAGMA_DISABLE_DEPRECATION_WARNINGS
+	// Deprecated
 	// If any, unregister previous generated links proxy.
 	UnregisterGeneratedLinksProxy();
+PRAGMA_ENABLE_DEPRECATION_WARNINGS
 
 	for (FNavLinkGenerationJumpConfig& NavLinkJumpConfig : NavLinkJumpConfigs)
 	{
+		// If any, unregister previous generated links proxy.
+		UnregisterGeneratedLinksProxy(NavLinkJumpConfig);
+		
 		NavLinkJumpConfig.LinkProxy = nullptr;
 
 		// Register new one if needed.
diff --git a/Engine/Source/Runtime/NavigationSystem/Public/NavMesh/RecastNavMesh.h b/Engine/Source/Runtime/NavigationSystem/Public/NavMesh/RecastNavMesh.h
--- a/Engine/Source/Runtime/NavigationSystem/Public/NavMesh/RecastNavMesh.h
+++ b/Engine/Source/Runtime/NavigationSystem/Public/NavMesh/RecastNavMesh.h
@@ -1519,7 +1519,13 @@ protected:
 	NAVIGATIONSYSTEM_API void CheckToDiscardSubLevelNavData(const UNavigationSystemBase& NavSys);
 	
 	bool IsGeneratingLinks() const;
+	void RegisterGeneratedLinksProxy(FNavLinkGenerationJumpConfig& InOutNavLinkJumpConfig);
+	void UnregisterGeneratedLinksProxy(FNavLinkGenerationJumpConfig& InOutNavLinkJumpConfig);
+
+	UE_DEPRECATED(5.8, "Use the override using an FNavLinkGenerationJumpConfig instead.")
 	void RegisterGeneratedLinksProxy();
+
+	UE_DEPRECATED(5.8, "Use the override using an FNavLinkGenerationJumpConfig instead.")
 	void UnregisterGeneratedLinksProxy();
 
 	/* Create and register links proxy. It's expected to be called on load or when the navmesh is rebuilt. */
