Removing old changes to handle generated NavLinks between layers due to lack of robustness and inefficiency.

#Synced-CL 50671260

diff --git a/Engine/Source/Runtime/NavigationSystem/Private/NavMesh/RecastNavMeshGenerator.cpp b/Engine/Source/Runtime/NavigationSystem/Private/NavMesh/RecastNavMeshGenerator.cpp
--- a/Engine/Source/Runtime/NavigationSystem/Private/NavMesh/RecastNavMeshGenerator.cpp
+++ b/Engine/Source/Runtime/NavigationSystem/Private/NavMesh/RecastNavMeshGenerator.cpp
@@ -3781,25 +3781,13 @@ struct FTileGenerationContext
 static void AddGeneratedLinks(
 	TInlineGeneratedNavLinkArray& OutGeneratedLinks,
 	const dtLinkBuilderConfig& linkBuilderConfig,
-	const dtNavLinkBuilder& linkBuilder,
-	const TArray<FBox, TInlineAllocator<16>>& LayerBounds,
-	const float SnapHeight)
+	const dtNavLinkBuilder& linkBuilder)
 {
-	auto FindLayerForZ = [&](const FVector::FReal Z) -> int16 {
-		for (int32 i = 0; i < LayerBounds.Num(); i++) {
-			if ((Z + SnapHeight) >= LayerBounds[i].Min.Z &&
-				(Z - SnapHeight) <= LayerBounds[i].Max.Z)
-				return (int16)i;
-		}
-		return INDEX_NONE;
-	};
-
-	auto AddLinkSharedLambda = [&OutGeneratedLinks, &FindLayerForZ](const dtNavLinkBuilderJumpConfig& config, const FVector& LeftPos) -> FGeneratedNavigationLink&
+	auto AddLinkSharedLambda = [&OutGeneratedLinks](const dtNavLinkBuilderJumpConfig& config) -> FGeneratedNavigationLink&
 	{
 		FGeneratedNavigationLink& NewLink = OutGeneratedLinks.Emplace_GetRef();
 		NewLink.bIsGenerated = true;
 		NewLink.NavLinkId = FNavLinkId(config.linkUserId);
-		NewLink.TargetLayerIndex = FindLayerForZ(LeftPos.Z);
 		return NewLink;
 	};
 
@@ -3809,7 +3797,7 @@ static void AddGeneratedLinks(
 		const FVector LeftPos = Recast2UnrealPoint(posA);
 		const FVector RightPos = Recast2UnrealPoint(posB);
 
-		FGeneratedNavigationLink& NewLink = AddLinkSharedLambda(config, LeftPos);
+		FGeneratedNavigationLink& NewLink = AddLinkSharedLambda(config);
 		NewLink.generatedLinkArea = config.downDirArea;
 		NewLink.generatedLinkPolyFlag = config.downDirPolyFlag;
 		NewLink.Left = LeftPos;
@@ -3829,7 +3817,7 @@ static void AddGeneratedLinks(
 
 		if (config.downDirArea != RECAST_NULL_AREA)
 		{
-			FGeneratedNavigationLink& DownLink = AddLinkSharedLambda(config, DownLinkLeft);
+			FGeneratedNavigationLink& DownLink = AddLinkSharedLambda(config);
 			DownLink.generatedLinkArea = config.downDirArea;
 			DownLink.generatedLinkPolyFlag = config.downDirPolyFlag;
 			DownLink.Left = DownLinkLeft;
@@ -3840,7 +3828,7 @@ static void AddGeneratedLinks(
 		if (config.upDirArea != RECAST_NULL_AREA)
 		{
 			// Note that in order for this link to be recognized by the navigation system, it needs to have unique left/right points, which is why the left and right points are flipped and we're marking this link as left-to-right
-			FGeneratedNavigationLink& UpLink = AddLinkSharedLambda(config, DownLinkRight);
+			FGeneratedNavigationLink& UpLink = AddLinkSharedLambda(config);
 			UpLink.generatedLinkArea = config.upDirArea;
 			UpLink.generatedLinkPolyFlag = config.upDirPolyFlag;
 			UpLink.Left = DownLinkRight;
@@ -4012,23 +4000,14 @@ dtStatus FRecastTileGenerator::BuildTileCacheLinks(FNavMeshBuildContext& BuildCo
 #endif // RECAST_INTERNAL_DEBUG_DATA
 	}
 
-	// Extract layer bounds
-	TArray<FBox, TInlineAllocator<16>> LayerBounds;
-	LayerBounds.Reserve(CompressedLayers.Num());
-	for (const FNavMeshTileData& LayerData : CompressedLayers)
-	{
-		LayerBounds.Add(LayerData.LayerBBox);
-	}
-
-	AddGeneratedLinks(OutGeneratedLinks, linkBuilderConfig, linkBuilder, LayerBounds, DefaultSnapHeight);
+	AddGeneratedLinks(OutGeneratedLinks, linkBuilderConfig, linkBuilder);
 
 	LogOnExit();
 	return DT_SUCCESS;
 }
 
 bool FRecastTileGenerator::GenerateNavigationDataLayer(FNavMeshBuildContext& BuildContext, FTileCacheCompressor& TileCompressor,
-	FTileCacheAllocator& GenNavAllocator, FTileGenerationContext& GenerationContext, const dtLinkBuilderData& InLinkBuilderData, int32 LayerIdx, const TInlineGeneratedNavLinkArray* InAdditionalLinks,
-	TInlineGeneratedNavLinkArray* OutGeneratedLinks)
+	FTileCacheAllocator& GenNavAllocator, FTileGenerationContext& GenerationContext, const dtLinkBuilderData& InLinkBuilderData, int32 LayerIdx)
 {
 	SCOPE_CYCLE_COUNTER(STAT_Navigation_GenerateNavigationDataLayer)
 		
@@ -4233,34 +4212,13 @@ bool FRecastTileGenerator::GenerateNavigationDataLayer(FNavMeshBuildContext& Bui
 
 		if (SolidHF && CompactHF)
 		{
-			TInlineGeneratedNavLinkArray NewlyGeneratedLinks;
-			status = BuildTileCacheLinks(BuildContext, &GenNavAllocator, *GenerationContext.Layer, *GenerationContext.ContourSet, DefaultSnapHeight, NewlyGeneratedLinks);
+			status = BuildTileCacheLinks(BuildContext, &GenNavAllocator, *GenerationContext.Layer, *GenerationContext.ContourSet, DefaultSnapHeight, GeneratedLinks);
 
 			if (dtStatusFailed(status))
 			{
 				BuildContext.log(RC_LOG_ERROR, "GenerateNavigationDataLayer: Failed to generate links (0x%08X).", status);
 				return false;
 			}
-
-			for (const FGeneratedNavigationLink& NewLink : NewlyGeneratedLinks)
-			{
-				if (NewLink.TargetLayerIndex == LayerIdx)
-				{
-					// Belongs to this layer - add to GeneratedLinks for nav data
-					GeneratedLinks.Add(NewLink);
-				}
-				else if (OutGeneratedLinks != nullptr && NewLink.TargetLayerIndex != INDEX_NONE)
-				{
-					// Belongs to a different layer - output for deferral
-					OutGeneratedLinks->Add(NewLink);
-				}
-			}
-		}
-
-		// Add links from previous layers. Needs to happen regardless of valid height fields.
-		if (InAdditionalLinks != nullptr)
-		{
-			GeneratedLinks.Append(*InAdditionalLinks);
 		}
 	}
 
@@ -4366,58 +4324,6 @@ bool FRecastTileGenerator::GenerateNavigationDataLayer(FNavMeshBuildContext& Bui
 	return true;
 }
 
-void FRecastTileGenerator::PrepareZOrderedLayerIndices(TArray<int32>& OutLayerIndices) const
-{
-	const int32 NumLayers = FMath::Min(CompressedLayers.Num(), DirtyLayers.Num());
-
-	// Step 1: Create array for indices
-	OutLayerIndices.Reserve(NumLayers);
-
-	// Step 2: Add valid layer indices
-	for (int32 Index = 0; Index < NumLayers; Index++)
-	{
-		if (DirtyLayers[Index] && CompressedLayers[Index].IsValid())
-		{
-			OutLayerIndices.Add(Index);
-		}
-	}
-
-	// Step 3: Sort by Z descending (highest first)
-	OutLayerIndices.Sort([this](int32 A, int32 B)
-	{
-		return CompressedLayers[A].LayerBBox.Max.Z > CompressedLayers[B].LayerBBox.Max.Z;
-	});
-}
-
-bool FRecastTileGenerator::GenerateNavigationDataLayerWithDeferredLinks(FNavMeshBuildContext& BuildContext, FTileCacheCompressor& TileCompressor, FTileCacheAllocator& GenNavAllocator, FTileGenerationContext& GenerationContext, 
-	const dtLinkBuilderData& InLinkBuilderData, int32 LayerIdx,	TInlineGeneratedNavLinkArray& InOutDeferredLinks)
-{
-	// Collect deferred links for THIS layer
-	TInlineGeneratedNavLinkArray LinksFromHigherLayers;
-	for (int32 Index = InOutDeferredLinks.Num() - 1; Index >= 0; Index--)
-	{
-		if (InOutDeferredLinks[Index].TargetLayerIndex == LayerIdx)
-		{
-			LinksFromHigherLayers.Add(InOutDeferredLinks[Index]);
-			InOutDeferredLinks.RemoveAtSwap(Index);
-		}
-	}
-
-	// Do the regular navigation data layer generation, passing in matching deferred links and retrieving new ones.
-
-	TInlineGeneratedNavLinkArray NewlyGeneratedLinks;
-	const bool bGenDataLayerSuccess = GenerateNavigationDataLayer(BuildContext, TileCompressor, GenNavAllocator, GenerationContext, InLinkBuilderData, LayerIdx, &LinksFromHigherLayers, &NewlyGeneratedLinks);
-
-	if (bGenDataLayerSuccess)
-	{
-		// Defer links for lower layers
-		InOutDeferredLinks.Append(NewlyGeneratedLinks);
-	}
-
-	return bGenDataLayerSuccess;
-}
-
-
 ETimeSliceWorkResult FRecastTileGenerator::GenerateNavigationDataTimeSliced(FNavMeshBuildContext& BuildContext, const dtLinkBuilderData& InLinkBuilderData)
 {
 	SCOPE_CYCLE_COUNTER(STAT_Navigation_RecastBuildNavigation);
@@ -4438,78 +4344,37 @@ ETimeSliceWorkResult FRecastTileGenerator::GenerateNavigationDataTimeSliced(FNav
 		GenNavDataTimeSlicedAllocator = MakeUnique<FTileCacheAllocator>();
 		GenNavDataTimeSlicedGenerationContext = MakeUnique<FTileGenerationContext>(GenNavDataTimeSlicedAllocator.Get());
 		GenNavDataTimeSlicedGenerationContext->NavigationData.Reserve(CompressedLayers.Num());
-		GenNavDataLayerIndices.Empty();
-		GenNavDataDeferredLinks.Empty();
 		GenerateNavDataTimeSlicedState = EGenerateNavDataTimeSlicedState::GenerateLayers;
 	}//fall through to next state
 	case EGenerateNavDataTimeSlicedState::GenerateLayers:
 	{
-		const int32 NumLayers = FMath::Min(CompressedLayers.Num(), DirtyLayers.Num());
-
-		if (InLinkBuilderData.generatingLinks && NumLayers > 1)
+		for (; GenNavDataLayerTimeSlicedIdx < CompressedLayers.Num() && GenNavDataLayerTimeSlicedIdx < DirtyLayers.Num(); GenNavDataLayerTimeSlicedIdx++)
 		{
-			if (GenNavDataLayerIndices.IsEmpty())
+			if (DirtyLayers[GenNavDataLayerTimeSlicedIdx] == false || !CompressedLayers[GenNavDataLayerTimeSlicedIdx].IsValid())
 			{
-				PrepareZOrderedLayerIndices(GenNavDataLayerIndices);
+				// skip layers not marked for rebuild
+				continue;
 			}
 
-			// Process layers in sorted order
-
-			for (; GenNavDataLayerTimeSlicedIdx < GenNavDataLayerIndices.Num(); GenNavDataLayerTimeSlicedIdx++)
+			if (TimeSliceManager->GetTimeSlicer().IsTimeSliceFinishedCached())
 			{
-				if (TimeSliceManager->GetTimeSlicer().IsTimeSliceFinishedCached())
-				{
-					WorkResult = ETimeSliceWorkResult::CallAgainNextTimeSlice;
-					break;
-				}
-
-				int32 LayerIdx = GenNavDataLayerIndices[GenNavDataLayerTimeSlicedIdx];
-
-				FTileCacheCompressor TileCompressor;
-				const bool bGenDataLayerSuccess = GenerateNavigationDataLayerWithDeferredLinks(BuildContext, TileCompressor, *GenNavDataTimeSlicedAllocator, *GenNavDataTimeSlicedGenerationContext, InLinkBuilderData, LayerIdx, GenNavDataDeferredLinks);
-
-				MARK_TIMESLICE_SECTION_DEBUG(TimeSliceManager->GetTimeSlicer(), GenerateLayers);
-
-				//carry on iterating but don't do any more work if the time slice is finished (as we may not need to in which case we can avoid calling this function again)
-				TimeSliceManager->GetTimeSlicer().TestTimeSliceFinished();
-
-				if (!bGenDataLayerSuccess)
-				{
-					WorkResult = ETimeSliceWorkResult::Failed;
-					break;
-				}
+				WorkResult = ETimeSliceWorkResult::CallAgainNextTimeSlice;
+				break;
 			}
-		}
-		else
-		{
-			for (; GenNavDataLayerTimeSlicedIdx < CompressedLayers.Num() && GenNavDataLayerTimeSlicedIdx < DirtyLayers.Num(); GenNavDataLayerTimeSlicedIdx++)
-			{
-				if (DirtyLayers[GenNavDataLayerTimeSlicedIdx] == false || !CompressedLayers[GenNavDataLayerTimeSlicedIdx].IsValid())
-				{
-					// skip layers not marked for rebuild
-					continue;
-				}
 
-				if (TimeSliceManager->GetTimeSlicer().IsTimeSliceFinishedCached())
-				{
-					WorkResult = ETimeSliceWorkResult::CallAgainNextTimeSlice;
-					break;
-				}
-
-				FTileCacheCompressor TileCompressor;
-				const bool bGenDataLayerSuccess = GenerateNavigationDataLayer(BuildContext, TileCompressor, *GenNavDataTimeSlicedAllocator,
-					*GenNavDataTimeSlicedGenerationContext, InLinkBuilderData, GenNavDataLayerTimeSlicedIdx);
+			FTileCacheCompressor TileCompressor;
+			const bool bGenDataLayerSuccess = GenerateNavigationDataLayer(BuildContext, TileCompressor, *GenNavDataTimeSlicedAllocator,
+				*GenNavDataTimeSlicedGenerationContext, InLinkBuilderData, GenNavDataLayerTimeSlicedIdx);
 
-				MARK_TIMESLICE_SECTION_DEBUG(TimeSliceManager->GetTimeSlicer(), GenerateLayers);
+			MARK_TIMESLICE_SECTION_DEBUG(TimeSliceManager->GetTimeSlicer(), GenerateLayers);
 
-				//carry on iterating but don't do any more work if the time slice is finished (as we may not need to in which case we can avoid calling this function again)
-				TimeSliceManager->GetTimeSlicer().TestTimeSliceFinished();
+			//carry on iterating but don't do any more work if the time slice is finished (as we may not need to in which case we can avoid calling this function again)
+			TimeSliceManager->GetTimeSlicer().TestTimeSliceFinished();
 
-				if (!bGenDataLayerSuccess)
-				{
-					WorkResult = ETimeSliceWorkResult::Failed;
-					break;
-				}
+			if (!bGenDataLayerSuccess)
+			{
+				WorkResult = ETimeSliceWorkResult::Failed;
+				break;
 			}
 		}
 
@@ -4547,44 +4412,18 @@ bool FRecastTileGenerator::GenerateNavigationData(FNavMeshBuildContext& BuildCon
 	FTileCacheCompressor TileCompressor;
 	bool bGenDataLayerSuccess = true;
 
-	const int32 NumLayers = FMath::Min(CompressedLayers.Num(), DirtyLayers.Num());
-
-	if (InLinkBuilderData.generatingLinks && NumLayers > 1)
+	for (int32 LayerIdx = 0; LayerIdx < CompressedLayers.Num() && LayerIdx < DirtyLayers.Num(); LayerIdx++)
 	{
-		// Make list of layer indices sorted by Z.
-		
-		TArray<int32> LayerIndices;
-		PrepareZOrderedLayerIndices(LayerIndices);
-
-		// Process the layers in Z order.
-
-		// DeferredLinks are populated with each call and results are passed onto subsequent calls as an input-output variable.
-		TInlineGeneratedNavLinkArray DeferredLinks;
-
-		for (const int32 LayerIdx : LayerIndices)
+		if (DirtyLayers[LayerIdx] == false || !CompressedLayers[LayerIdx].IsValid())
 		{
-			bGenDataLayerSuccess = GenerateNavigationDataLayerWithDeferredLinks(BuildContext, TileCompressor, GenNavAllocator, GenerationContext, InLinkBuilderData, LayerIdx, DeferredLinks);
-			if (!bGenDataLayerSuccess)
-			{
-				break;
-			}
+			// skip layers not marked for rebuild
+			continue;
 		}
-	}
-	else
-	{
-		for (int32 LayerIdx = 0; LayerIdx < CompressedLayers.Num() && LayerIdx < DirtyLayers.Num(); LayerIdx++)
-		{
-			if (DirtyLayers[LayerIdx] == false || !CompressedLayers[LayerIdx].IsValid())
-			{
-				// skip layers not marked for rebuild
-				continue;
-			}
 
-			bGenDataLayerSuccess = GenerateNavigationDataLayer(BuildContext, TileCompressor, GenNavAllocator, GenerationContext, InLinkBuilderData, LayerIdx);
-			if (!bGenDataLayerSuccess)
-			{
-				break;
-			}
+		bGenDataLayerSuccess = GenerateNavigationDataLayer(BuildContext, TileCompressor, GenNavAllocator, GenerationContext, InLinkBuilderData, LayerIdx);
+		if (!bGenDataLayerSuccess)
+		{
+			break;
 		}
 	}
 
diff --git a/Engine/Source/Runtime/NavigationSystem/Public/NavMesh/RecastNavMeshGenerator.h b/Engine/Source/Runtime/NavigationSystem/Public/NavMesh/RecastNavMeshGenerator.h
--- a/Engine/Source/Runtime/NavigationSystem/Public/NavMesh/RecastNavMeshGenerator.h
+++ b/Engine/Source/Runtime/NavigationSystem/Public/NavMesh/RecastNavMeshGenerator.h
@@ -353,9 +353,6 @@ struct FGeneratedNavigationLink : public FNavigationLink
 
 	/// User defined area ids assigned to the off-mesh connections
 	unsigned char generatedLinkArea = 0;
-
-	// What layer the link will be associated with
-	int16 TargetLayerIndex = INDEX_NONE;
 };
 typedef TArray<FGeneratedNavigationLink, TInlineAllocator<64>> TInlineGeneratedNavLinkArray;
 
@@ -415,12 +412,6 @@ public:
 		
 	typedef TArray<int32, TInlineAllocator<4096>> TInlineMaskArray;
 
-private:
-
-	void PrepareZOrderedLayerIndices(TArray<int32>& OutLayerIndices) const;
-	bool GenerateNavigationDataLayerWithDeferredLinks(FNavMeshBuildContext& BuildContext, FTileCacheCompressor& TileCompressor, FTileCacheAllocator& GenNavAllocator, FTileGenerationContext& GenerationContext, const dtLinkBuilderData& InLinkBuilderData, int32 LayerIdx,
-		TInlineGeneratedNavLinkArray& OutDeferredLinks);
-
 protected:
 	/** Does the actual TimeSliced tile generation. 
 	 *	@note always trigger tile generation only via DoWorkTimeSliced(). This is a worker function
@@ -478,8 +469,7 @@ protected:
 	NAVIGATIONSYSTEM_API virtual bool GenerateCompressedLayers(FNavMeshBuildContext& BuildContext, const dtLinkBuilderData& InLinkBuilderData);
 
 	/** Builds a navigation data layer */
-	NAVIGATIONSYSTEM_API bool GenerateNavigationDataLayer(FNavMeshBuildContext& BuildContext, FTileCacheCompressor& TileCompressor, FTileCacheAllocator& GenNavAllocator, FTileGenerationContext& GenerationContext, const dtLinkBuilderData& InLinkBuilderData, int32 LayerIdx, 
-		const TInlineGeneratedNavLinkArray* InAdditionalLinks = nullptr, TInlineGeneratedNavLinkArray* OutGeneratedLinks = nullptr);
+	NAVIGATIONSYSTEM_API bool GenerateNavigationDataLayer(FNavMeshBuildContext& BuildContext, FTileCacheCompressor& TileCompressor, FTileCacheAllocator& GenNavAllocator, FTileGenerationContext& GenerationContext, const dtLinkBuilderData& InLinkBuilderData, int32 LayerIdx);
 
 	/** Builds NavigationData array (layers + obstacles) time sliced */
 	NAVIGATIONSYSTEM_API ETimeSliceWorkResult GenerateNavigationDataTimeSliced(FNavMeshBuildContext& BuildContext, const dtLinkBuilderData& InLinkBuilderData);
@@ -559,9 +549,6 @@ protected:
 	TUniquePtr<struct FTileGenerationContext> GenNavDataTimeSlicedGenerationContext;
 	TUniquePtr<struct FTileRasterizationContext> GenCompressedlayersTimeSlicedRasterContext;
 
-	TArray<int32> GenNavDataLayerIndices;
-	TInlineGeneratedNavLinkArray GenNavDataDeferredLinks;
-
 	FRecastTileTimeSliceSettings TileTimeSliceSettings;
 	/** End time slicing variables */
 
