Fixed storing reversed NavLinks. Now directional up links are stored in reversed order, which makes sure they stay with the associated down link. This removes the need for deferring links being added per-layer.
Additionally, this now correctly handles the situation of asymmetric border sizes by making sure the owner position is the one that is contained by the tile.

#Synced-CL 50791735

#Synced-CL 50792019

diff --git a/Engine/Source/Runtime/NavigationSystem/Private/NavMesh/PImplRecastNavMesh.cpp b/Engine/Source/Runtime/NavigationSystem/Private/NavMesh/PImplRecastNavMesh.cpp
--- a/Engine/Source/Runtime/NavigationSystem/Private/NavMesh/PImplRecastNavMesh.cpp
+++ b/Engine/Source/Runtime/NavigationSystem/Private/NavMesh/PImplRecastNavMesh.cpp
@@ -2912,7 +2912,7 @@ void FPImplRecastNavMesh::GetTilePolyEdges(const dtMeshTile& Tile, bool bGatherI
 	}
 }
 
-uint8 GetValidEnds(const dtNavMesh& NavMesh, const dtMeshTile& Tile, const dtPoly& Poly)
+uint8 GetValidEnds(const dtNavMesh& NavMesh, const dtMeshTile& Tile, const dtPoly& Poly, const bool bIsReversed = false)
 {
 	if (Poly.getType() == DT_POLYTYPE_GROUND)
 	{
@@ -2929,11 +2929,11 @@ uint8 GetValidEnds(const dtNavMesh& NavMesh, const dtMeshTile& Tile, const dtPol
 
 		if (link.edge == 0)
 		{
-			ValidEnds |= FRecastDebugGeometry::OMLE_Left;
+			ValidEnds |= bIsReversed ? FRecastDebugGeometry::OMLE_Right : FRecastDebugGeometry::OMLE_Left;
 		}
 		if (link.edge == 1)
 		{
-			ValidEnds |= FRecastDebugGeometry::OMLE_Right;
+			ValidEnds |= bIsReversed ? FRecastDebugGeometry::OMLE_Left : FRecastDebugGeometry::OMLE_Right;
 		}
 	}
 
@@ -3226,7 +3226,7 @@ int32 FPImplRecastNavMesh::GetTilesDebugGeometry(const FRecastNavMeshGenerator*
 				, Recast2UnrVector(vb)
 				, LinkPoly->getArea()
 				, (uint8)OffMeshConnection->getBiDirectional()
-				, GetValidEnds(*DetourNavMesh, Tile, *LinkPoly)
+				, GetValidEnds(*DetourNavMesh, Tile, *LinkPoly, OffMeshConnection->getIsReversed())
 				, OffMeshConnection->getIsGenerated()
 				, UE_REAL_TO_FLOAT_CLAMPED_MAX(OffMeshConnection->rad)
 			};
@@ -3238,6 +3238,9 @@ int32 FPImplRecastNavMesh::GetTilesDebugGeometry(const FRecastNavMeshGenerator*
 	}
 
 #if WITH_NAVMESH_SEGMENT_LINKS
+
+	// Note: Segment Links don't currently support the notion of reversed links yet.
+
 	for (int32 i = 0; i < Header->offMeshSegConCount; ++i)
 	{
 		const dtOffMeshSegmentConnection* OffMeshSeg = &Tile.offMeshSeg[i];
diff --git a/Engine/Source/Runtime/NavigationSystem/Private/NavMesh/RecastNavMeshGenerator.cpp b/Engine/Source/Runtime/NavigationSystem/Private/NavMesh/RecastNavMeshGenerator.cpp
--- a/Engine/Source/Runtime/NavigationSystem/Private/NavMesh/RecastNavMeshGenerator.cpp
+++ b/Engine/Source/Runtime/NavigationSystem/Private/NavMesh/RecastNavMeshGenerator.cpp
@@ -1381,7 +1381,8 @@ struct FOffMeshData
 		NewInfo.type = DT_OFFMESH_CON_POINT | 
 			(Link.Direction == ENavLinkDirection::BothWays ? DT_OFFMESH_CON_BIDIR : 0) |
 			(Link.bSnapToCheapestArea ? DT_OFFMESH_CON_CHEAPAREA : 0) |
-			(Link.bIsGenerated ? DT_OFFMESH_CON_GENERATED : 0);
+			(Link.bIsGenerated ? DT_OFFMESH_CON_GENERATED : 0) |
+			(Link.Direction == ENavLinkDirection::RightToLeft ? DT_OFFMESH_CON_REVERSED : 0);
 
 		NewInfo.snapRadius = Link.SnapRadius;
 		NewInfo.snapHeight = Link.bUseSnapHeight ? Link.SnapHeight : DefaultSnapHeight;
@@ -3781,8 +3782,17 @@ struct FTileGenerationContext
 static void AddGeneratedLinks(
 	TInlineGeneratedNavLinkArray& OutGeneratedLinks,
 	const dtLinkBuilderConfig& linkBuilderConfig,
-	const dtNavLinkBuilder& linkBuilder)
+	const dtNavLinkBuilder& linkBuilder,
+	const dtReal* TileBoundsMin,
+	const dtReal* TileBoundsMax)
 {
+	// Helper to check if a point is inside the tile's XZ bounds
+	auto IsInsideTileXZ = [&TileBoundsMin, &TileBoundsMax](const dtReal* pos) -> bool
+	{
+		return pos[0] >= TileBoundsMin[0] && pos[0] <= TileBoundsMax[0]
+			&& pos[2] >= TileBoundsMin[2] && pos[2] <= TileBoundsMax[2];
+	};
+
 	auto AddLinkSharedLambda = [&OutGeneratedLinks](const dtNavLinkBuilderJumpConfig& config) -> FGeneratedNavigationLink&
 	{
 		FGeneratedNavigationLink& NewLink = OutGeneratedLinks.Emplace_GetRef();
@@ -3792,48 +3802,68 @@ static void AddGeneratedLinks(
 	};
 
 	// Used when we're provided with identical up and down areas. In this situation, a single bidirectional link is sufficient for our needs
-	auto AddBidirectionalLinkLambda = [&AddLinkSharedLambda](const dtNavLinkBuilderJumpConfig& config, const dtReal* posA, const dtReal* posB)
+	auto AddBidirectionalLinkLambda = [&IsInsideTileXZ , &AddLinkSharedLambda](const dtNavLinkBuilderJumpConfig& config, const dtReal* posA, const dtReal* posB)
 	{
-		const FVector LeftPos = Recast2UnrealPoint(posA);
-		const FVector RightPos = Recast2UnrealPoint(posB);
+		const FVector UnrealPosA = Recast2UnrealPoint(posA);
+		const FVector UnrealPosB = Recast2UnrealPoint(posB);
+
+		// Check if the endpoints are in the current tile and, if posA isn't inside this one, flip the direction.
+		// Always keep the direction from owner position to far position.
+
+		const bool bAInside = IsInsideTileXZ(posA);
+		const bool bBInside = IsInsideTileXZ(posB);
+		const bool bSwapForOwnership = !bAInside && bBInside;
+
+		const FVector OwnerPos = bSwapForOwnership ? UnrealPosB : UnrealPosA;
+		const FVector FarPos = bSwapForOwnership ? UnrealPosA : UnrealPosB;
 
 		FGeneratedNavigationLink& NewLink = AddLinkSharedLambda(config);
 		NewLink.generatedLinkArea = config.downDirArea;
 		NewLink.generatedLinkPolyFlag = config.downDirPolyFlag;
-		NewLink.Left = LeftPos;
-		NewLink.Right = RightPos;
+		NewLink.Left = OwnerPos;
+		NewLink.Right = FarPos;
 	};
 
 	// Used when we're provided with different up and down areas. In this situation, we need to create two different one-way link representations in order to meet our needs
-	auto AddDirectionalLinksLambda = [&AddLinkSharedLambda](const dtNavLinkBuilderJumpConfig& config, const dtReal* posA, const dtReal* posB)
+	auto AddDirectionalLinksLambda = [&IsInsideTileXZ, &AddLinkSharedLambda](const dtNavLinkBuilderJumpConfig& config, const dtReal* posA, const dtReal* posB)
 	{
 		const FVector UnrealPosA = Recast2UnrealPoint(posA);
 		const FVector UnrealPosB = Recast2UnrealPoint(posB);
 
-		const bool bABDown = UnrealPosA.Z > UnrealPosB.Z;
+		const bool bAInside = IsInsideTileXZ(posA);
+		const bool bBInside = IsInsideTileXZ(posB);
+		const bool bSwap = !bAInside && bBInside;
 
-		const FVector DownLinkLeft = bABDown ? UnrealPosA : UnrealPosB;
-		const FVector DownLinkRight = bABDown ? UnrealPosB : UnrealPosA;
+		const FVector OwnerPos = bSwap ? UnrealPosB : UnrealPosA;
+		const FVector FarPos = bSwap ? UnrealPosA : UnrealPosB;
+
+		// Both directions of the links share the same Left (owner) and Right (far) positions.
+		// They are distinguished by Direction: the "down" link uses LeftToRight, while the "up"
+		// link uses RightToLeft, which sets DT_OFFMESH_CON_REVERSED to swap traversal semantics.
+		// If ownership was swapped (bSwap), these directions are inverted.
 
 		if (config.downDirArea != RECAST_NULL_AREA)
 		{
 			FGeneratedNavigationLink& DownLink = AddLinkSharedLambda(config);
 			DownLink.generatedLinkArea = config.downDirArea;
 			DownLink.generatedLinkPolyFlag = config.downDirPolyFlag;
-			DownLink.Left = DownLinkLeft;
-			DownLink.Right = DownLinkRight;
-			DownLink.Direction = ENavLinkDirection::LeftToRight;
+			DownLink.Left = OwnerPos;
+			DownLink.Right = FarPos;
+			DownLink.Direction = bSwap
+				? ENavLinkDirection::RightToLeft
+				: ENavLinkDirection::LeftToRight;
 		}
 
 		if (config.upDirArea != RECAST_NULL_AREA)
 		{
-			// Note that in order for this link to be recognized by the navigation system, it needs to have unique left/right points, which is why the left and right points are flipped and we're marking this link as left-to-right
 			FGeneratedNavigationLink& UpLink = AddLinkSharedLambda(config);
 			UpLink.generatedLinkArea = config.upDirArea;
 			UpLink.generatedLinkPolyFlag = config.upDirPolyFlag;
-			UpLink.Left = DownLinkRight;
-			UpLink.Right = DownLinkLeft;
-			UpLink.Direction = ENavLinkDirection::LeftToRight;
+			UpLink.Left = OwnerPos;
+			UpLink.Right = FarPos;
+			UpLink.Direction = bSwap
+				? ENavLinkDirection::LeftToRight
+				: ENavLinkDirection::RightToLeft;
 		}
 	};
 
@@ -4000,7 +4030,11 @@ dtStatus FRecastTileGenerator::BuildTileCacheLinks(FNavMeshBuildContext& BuildCo
 #endif // RECAST_INTERNAL_DEBUG_DATA
 	}
 
-	AddGeneratedLinks(OutGeneratedLinks, linkBuilderConfig, linkBuilder);
+	const FBox RcTileBB = Unreal2RecastBox(TileBB);
+	const FVector RcTileMin(RcTileBB.Min);
+	const FVector RcTileMax(RcTileBB.Max);
+
+	AddGeneratedLinks(OutGeneratedLinks, linkBuilderConfig, linkBuilder, &RcTileMin.X, &RcTileMax.X);
 
 	LogOnExit();
 	return DT_SUCCESS;
diff --git a/Engine/Source/Runtime/Navmesh/Private/Detour/DetourNavMesh.cpp b/Engine/Source/Runtime/Navmesh/Private/Detour/DetourNavMesh.cpp
--- a/Engine/Source/Runtime/Navmesh/Private/Detour/DetourNavMesh.cpp
+++ b/Engine/Source/Runtime/Navmesh/Private/Detour/DetourNavMesh.cpp
@@ -1323,35 +1323,44 @@ void dtNavMesh::connectExtOffMeshLinks(dtMeshTile* tile, dtMeshTile* target, int
 		// Avoid linking back into the same ground poly
 		if (!ref || (targetLandPoly == ref))
 			continue;
-		// Make sure the location is on current mesh.
-		dtReal* v = &target->verts[targetPoly->verts[1]*3];
+
+		// For reversed links, the landing vertex is verts[0] instead of verts[1].
+		const bool bReversed = targetCon->getIsReversed();
+		const int landingVertIdx = bReversed ? 0 : 1;
+		dtReal* v = &target->verts[targetPoly->verts[landingVertIdx] * 3];
 		dtVcopy(v, nearestPt);
-		
-		unsigned char linkSide = oppositeSide | DT_LINK_FLAG_OFFMESH_CON | biDirFlag;
+
+		// Link off-mesh connection to landing poly.
+		// For reversed links, this is the backtracker direction.
+		unsigned char fwdLinkSide = oppositeSide | DT_LINK_FLAG_OFFMESH_CON | biDirFlag;
 		if (tile != target)
 		{
-			linkSide &= ~DT_CONNECTION_INTERNAL;
+			fwdLinkSide &= ~DT_CONNECTION_INTERNAL;
+		}
+		if (bReversed && biDirFlag == 0)
+		{
+			fwdLinkSide |= DT_LINK_FLAG_OFFMESH_CON_BACKTRACKER;
 		}
 
 		// Link off-mesh connection to target poly.
 		const unsigned int landPolyIdx = decodePolyIdPoly(ref);
-		linkOffMeshHelper(target, targetCon->poly, tile, landPolyIdx, linkSide, 1);
+		const unsigned char landingEdge = bReversed ? 0 : 1;
+
+		linkOffMeshHelper(target, targetCon->poly, tile, landPolyIdx, fwdLinkSide, landingEdge);
 
-		// Link target poly to off-mesh connection.
-		linkSide = (unsigned char)(side == -1 ? DT_CONNECTION_INTERNAL : side) | DT_LINK_FLAG_OFFMESH_CON | biDirFlag;
+		// Link landing poly back to off-mesh connection.
+		// For non-reversed unidirectional links, this is the backtracker direction.
+		unsigned char bckLinkSide = (unsigned char)(side == -1 ? DT_CONNECTION_INTERNAL : side) | DT_LINK_FLAG_OFFMESH_CON | biDirFlag;
 		if (tile != target)
 		{
-			linkSide &= ~DT_CONNECTION_INTERNAL;
+			bckLinkSide &= ~DT_CONNECTION_INTERNAL;
 		}
-
-		if (biDirFlag == 0)
+		if (!bReversed && biDirFlag == 0)
 		{
-			// if it's not a bi-directional link put it in anyway
-			// just annotate it accordingly
-			linkSide |= DT_LINK_FLAG_OFFMESH_CON_BACKTRACKER;
+			bckLinkSide |= DT_LINK_FLAG_OFFMESH_CON_BACKTRACKER;
 		}
 
-		linkOffMeshHelper(tile, landPolyIdx, target, targetCon->poly, linkSide, 0xff);
+		linkOffMeshHelper(tile, landPolyIdx, target, targetCon->poly, bckLinkSide, 0xff);
 
 #if WITH_NAVMESH_CLUSTER_LINKS
 		if (updateCLinks)
@@ -1366,10 +1375,12 @@ void dtNavMesh::connectExtOffMeshLinks(dtMeshTile* tile, dtMeshTile* target, int
 				unsigned int thisClusterIdx = tile->polyClusters[thisPolyIdx];
 				const bool bUniqueCheck = true;
 
-				const unsigned char flagsFwd = DT_CLINK_VALID_FWD | (biDirFlag ? DT_CLINK_VALID_BCK : 0);
-				const unsigned char flagsBck = DT_CLINK_VALID_BCK | (biDirFlag ? DT_CLINK_VALID_FWD : 0);
+				const unsigned char baseFwd = bReversed ? DT_CLINK_VALID_BCK : DT_CLINK_VALID_FWD;
+				const unsigned char baseBck = bReversed ? DT_CLINK_VALID_FWD : DT_CLINK_VALID_BCK;
+				const unsigned char flagsFwd = baseFwd | (biDirFlag ? baseBck : 0);
+				const unsigned char flagsBck = baseBck | (biDirFlag ? baseFwd : 0);
 
-				connectClusterLink(target, targetClusterIdx, tile, thisClusterIdx, flagsFwd, bUniqueCheck);				
+				connectClusterLink(target, targetClusterIdx, tile, thisClusterIdx, flagsFwd, bUniqueCheck);
 				connectClusterLink(tile, thisClusterIdx, target, targetClusterIdx, flagsBck, bUniqueCheck);
 			}
 		}
@@ -1457,17 +1468,25 @@ void dtNavMesh::baseOffMeshLinks(dtMeshTile* tile)
 		if (!ref) continue;
 
 		// Make sure the location is on current mesh.
-		dtReal* v = &tile->verts[poly->verts[0]*3];
+		const bool bReversed = con->getIsReversed();
+		const int baseVertIdx = bReversed ? 1 : 0;
+		dtReal* v = &tile->verts[poly->verts[baseVertIdx] * 3];
 		dtVcopy(v, nearestPt);
 
 		unsigned char sideFwd = DT_CONNECTION_INTERNAL | DT_LINK_FLAG_OFFMESH_CON | (con->getBiDirectional() ? DT_LINK_FLAG_OFFMESH_CON_BIDIR : 0);
 		unsigned char sideBck = sideFwd | DT_LINK_FLAG_OFFMESH_CON_BACKTRACKER;
 
-		// Link off-mesh connection to target poly.
-		linkOffMeshHelper(tile, con->poly, tile, decodePolyIdPoly(ref), sideBck, 0);
+		// For reversed links, swap forward/backward so that the normal traversal
+		// direction goes from pos[3] to pos[0] (i.e. B0 to A0).
+		const unsigned char sideToLocal = bReversed ? sideFwd : sideBck;
+		const unsigned char sideFromLocal = bReversed ? sideBck : sideFwd;
+		const int baseEdge = bReversed ? 1 : 0;
+
+		// Link off-mesh connection to local ground poly.
+		linkOffMeshHelper(tile, con->poly, tile, decodePolyIdPoly(ref), sideToLocal, baseEdge);
 
-		// Start end-point is always connect back to off-mesh connection.
-		linkOffMeshHelper(tile, decodePolyIdPoly(ref), tile, con->poly, sideFwd, 0xff);
+		// Link local ground poly back to off-mesh connection.
+		linkOffMeshHelper(tile, decodePolyIdPoly(ref), tile, con->poly, sideFromLocal, 0xff);
 	}
 }
 
diff --git a/Engine/Source/Runtime/Navmesh/Public/Detour/DetourNavMesh.h b/Engine/Source/Runtime/Navmesh/Public/Detour/DetourNavMesh.h
--- a/Engine/Source/Runtime/Navmesh/Public/Detour/DetourNavMesh.h
+++ b/Engine/Source/Runtime/Navmesh/Public/Detour/DetourNavMesh.h
@@ -115,6 +115,8 @@ static const unsigned char DT_OFFMESH_CON_POINT		= 0x02;		//UE
 static const unsigned char DT_OFFMESH_CON_SEGMENT	= 0x04;		//UE
 static const unsigned char DT_OFFMESH_CON_CHEAPAREA = 0x08;		//UE
 static const unsigned char DT_OFFMESH_CON_GENERATED = 0x10;		//UE
+/// A flag that indicates whether the direction of travel along the one-way off-mesh connection is actually reversed.
+static const unsigned char DT_OFFMESH_CON_REVERSED  = 0x20;		//UE
 
 /// The maximum number of user defined area ids.
 /// @ingroup detour
@@ -297,7 +299,7 @@ struct dtOffMeshSegmentConnection
 	/// Sets link flags
 	inline void setFlags(unsigned char conFlags)
 	{
-		flags = ((conFlags & DT_OFFMESH_CON_BIDIR) ? 0x80 : 0); 
+		flags = ((conFlags & DT_OFFMESH_CON_BIDIR) ? 0x80 : 0) ;
 	}
 
 	/// Gets the link direction
@@ -339,7 +341,8 @@ struct dtOffMeshConnection
 	{
 		flags = ((conTypeFlags & DT_OFFMESH_CON_BIDIR) ? 0x80 : 0) |
 				((conTypeFlags & DT_OFFMESH_CON_CHEAPAREA) ? 0x40 : 0) |
-				((conTypeFlags & DT_OFFMESH_CON_GENERATED) ? 0x20 : 0);
+				((conTypeFlags & DT_OFFMESH_CON_GENERATED) ? 0x20 : 0) | 
+				((conTypeFlags & DT_OFFMESH_CON_REVERSED) ? 0x10 : 0);
 	}
 
 	/// Gets the link direction
@@ -350,6 +353,10 @@ struct dtOffMeshConnection
 
 	/// Indicates if the link was automatically generated
 	inline bool getIsGenerated() const { return (flags & 0x20) != 0; }
+	
+	/// Indicates if the link is reversed
+	inline bool getIsReversed() const { return (flags & 0x10) != 0; }
+
 	//@UE END
 };
 
