Fix for directional generated links not being height tested.

#Synced-CL 50912775

diff --git a/Engine/Source/Runtime/NavigationSystem/Private/NavMesh/RecastNavMeshGenerator.cpp b/Engine/Source/Runtime/NavigationSystem/Private/NavMesh/RecastNavMeshGenerator.cpp
--- a/Engine/Source/Runtime/NavigationSystem/Private/NavMesh/RecastNavMeshGenerator.cpp
+++ b/Engine/Source/Runtime/NavigationSystem/Private/NavMesh/RecastNavMeshGenerator.cpp
@@ -3832,7 +3832,14 @@ static void AddGeneratedLinks(
 
 		const bool bAInside = IsInsideTileXZ(posA);
 		const bool bBInside = IsInsideTileXZ(posB);
+
+		// If posA is not inside the tile but posB is, swap directions for ownership.
 		const bool bSwap = !bAInside && bBInside;
+		
+		const bool bABDown = UnrealPosA.Z > UnrealPosB.Z;
+
+		// If directions are swapped but posA is above posB, reverse the directionality of the links. This is to preserve the downDirArea going down, and the upDirArea going up.
+		const bool bDownIsRightToLeft = (bSwap == bABDown);
 
 		const FVector OwnerPos = bSwap ? UnrealPosB : UnrealPosA;
 		const FVector FarPos = bSwap ? UnrealPosA : UnrealPosB;
@@ -3849,7 +3856,7 @@ static void AddGeneratedLinks(
 			DownLink.generatedLinkPolyFlag = config.downDirPolyFlag;
 			DownLink.Left = OwnerPos;
 			DownLink.Right = FarPos;
-			DownLink.Direction = bSwap
+			DownLink.Direction = bDownIsRightToLeft
 				? ENavLinkDirection::RightToLeft
 				: ENavLinkDirection::LeftToRight;
 		}
@@ -3861,7 +3868,7 @@ static void AddGeneratedLinks(
 			UpLink.generatedLinkPolyFlag = config.upDirPolyFlag;
 			UpLink.Left = OwnerPos;
 			UpLink.Right = FarPos;
-			UpLink.Direction = bSwap
+			UpLink.Direction = bDownIsRightToLeft
 				? ENavLinkDirection::LeftToRight
 				: ENavLinkDirection::RightToLeft;
 		}
