Fix for missing Link Proxy registration when NavSys is sometimes null.

#Synced-CL 50912775

diff --git a/Engine/Source/Runtime/NavigationSystem/Private/NavMesh/RecastNavMesh.cpp b/Engine/Source/Runtime/NavigationSystem/Private/NavMesh/RecastNavMesh.cpp
--- a/Engine/Source/Runtime/NavigationSystem/Private/NavMesh/RecastNavMesh.cpp
+++ b/Engine/Source/Runtime/NavigationSystem/Private/NavMesh/RecastNavMesh.cpp
@@ -1217,6 +1217,28 @@ void ARecastNavMesh::OnRegistered()
 	ensure(RecastNavMeshImpl->GetRecastMesh() == nullptr || RecastNavMeshImpl->GetRecastMesh()->getBVQuantFactor((uint8)ENavigationDataResolution::Default) != 0);
 
 	Super::OnRegistered();
+
+	// Re-register any generated link proxies that couldn't be registered
+	// when created (NavSys was null), or that need re-registration after
+	// NavSys was destroyed and recreated (flags reset in OnUnregistered).
+	for (FNavLinkGenerationJumpConfig& NavLinkJumpConfig : NavLinkJumpConfigs)
+	{
+		RegisterGeneratedLinksProxy(NavLinkJumpConfig);
+	}
+}
+
+void ARecastNavMesh::OnUnregistered()
+{
+	Super::OnUnregistered();
+
+	// Reset link proxy registration flags so they can be re-registered
+	// with a new NavSys instance in OnRegistered. The proxies themselves
+	// remain valid — only the registration state is cleared, since the
+	// NavSys they were registered with is being torn down.
+	for (FNavLinkGenerationJumpConfig& NavLinkJumpConfig : NavLinkJumpConfigs)
+	{
+		NavLinkJumpConfig.bLinkProxyRegistered = false;
+	}
 }
 
 void ARecastNavMesh::SerializeRecastNavMesh(FArchive& Ar, FPImplRecastNavMesh*& NavMesh, int32 InNavMeshVersion)
diff --git a/Engine/Source/Runtime/NavigationSystem/Public/NavMesh/RecastNavMesh.h b/Engine/Source/Runtime/NavigationSystem/Public/NavMesh/RecastNavMesh.h
--- a/Engine/Source/Runtime/NavigationSystem/Public/NavMesh/RecastNavMesh.h
+++ b/Engine/Source/Runtime/NavigationSystem/Public/NavMesh/RecastNavMesh.h
@@ -1092,6 +1092,7 @@ protected:
 	NAVIGATIONSYSTEM_API virtual void RestrictBuildingToActiveTiles(bool InRestrictBuildingToActiveTiles) override;
 
 	NAVIGATIONSYSTEM_API virtual void OnRegistered() override;
+	NAVIGATIONSYSTEM_API virtual void OnUnregistered() override;
 
 public:
 	/** Get the CellSize for the given resolution. */
diff --git a/Engine/Source/Runtime/NavigationSystem/Public/NavigationData.h b/Engine/Source/Runtime/NavigationSystem/Public/NavigationData.h
--- a/Engine/Source/Runtime/NavigationSystem/Public/NavigationData.h
+++ b/Engine/Source/Runtime/NavigationSystem/Public/NavigationData.h
@@ -627,7 +627,7 @@ public:
 
 	inline bool IsRegistered() const { return bRegistered; }
 	NAVIGATIONSYSTEM_API virtual void OnRegistered();
-	NAVIGATIONSYSTEM_API void OnUnregistered();
+	NAVIGATIONSYSTEM_API virtual void OnUnregistered();
 	
 	inline uint16 GetNavDataUniqueID() const { return NavDataUniqueID; }
 
