Fix reentrance in RegisterNavData

#Synced-CL 50899516

diff --git a/Engine/Source/Runtime/NavigationSystem/Private/NavigationSystem.cpp b/Engine/Source/Runtime/NavigationSystem/Private/NavigationSystem.cpp
--- a/Engine/Source/Runtime/NavigationSystem/Private/NavigationSystem.cpp
+++ b/Engine/Source/Runtime/NavigationSystem/Private/NavigationSystem.cpp
@@ -2517,9 +2517,10 @@ ANavigationData* UNavigationSystemV1::GetDefaultNavDataInstance(FNavigationSyste
 			MainNavData = GetWorld()->SpawnActor<ANavigationData>(ARecastNavMesh::StaticClass());
 		}
 #endif // WITH_RECAST
+		
 		// either way make sure it's registered. Registration stores unique
 		// navmeshes, so we have nothing to lose
-		if (MainNavData != nullptr)
+		if (MainNavData != nullptr && !MainNavData->IsRegistered())
 		{
 			const ERegistrationResult Result = RegisterNavData(MainNavData);
 			LogNavDataRegistrationResult(Result, MainNavData);
@@ -2827,12 +2828,13 @@ UNavigationSystemV1::ERegistrationResult UNavigationSystemV1::RegisterNavData(AN
 	{
 		return RegistrationError;
 	}
-	else if (!IsValid(NavData))
+
+	if (!IsValid(NavData))
 	{
 		return RegistrationFailed_DataPendingKill;
 	}
-	// still to be seen if this is really true, but feels right
-	else if (NavData->IsRegistered() == true)
+
+	if (NavData->IsRegistered() == true)
 	{
 		return RegistrationSuccessful;
 	}
