[Navigation] Don't spawn missing navigation data that is already pending for registration.

#Synced-CL 51086784

diff --git a/Engine/Source/Runtime/NavigationSystem/Private/NavigationSystem.cpp b/Engine/Source/Runtime/NavigationSystem/Private/NavigationSystem.cpp
--- a/Engine/Source/Runtime/NavigationSystem/Private/NavigationSystem.cpp
+++ b/Engine/Source/Runtime/NavigationSystem/Private/NavigationSystem.cpp
@@ -1535,6 +1535,8 @@ void UNavigationSystemV1::OnWorldInitDone(FNavigationSystemRunMode Mode)
 void UNavigationSystemV1::RegisterNavigationDataInstances()
 {
 	UWorld* World = GetWorld();
+	
+	UE_LOGF(LogNavigation, Verbose, "%s World: %ls", __FUNCTION__, *GetNameSafe(World));
 
 	bool bProcessRegistration = false;
 	for (TActorIterator<ANavigationData> It(World); It; ++It)
@@ -2737,11 +2739,12 @@ void UNavigationSystemV1::RequestRegistrationDeferred(ANavigationData& NavData)
 
 	if (NavDataRegistrationQueue.Num() < REGISTRATION_QUEUE_SIZE)
 	{
+		UE_LOGF(LogNavigation, Verbose, "%s Adding to registration queue: %ls.", __FUNCTION__, *GetFullNameSafe(&NavData));
 		NavDataRegistrationQueue.AddUnique(&NavData);
 	}
 	else
 	{
-		UE_LOG(LogNavigation, Warning, TEXT("Navigation System: registration queue full! System:%s NavData:%s"), *GetPathNameSafe(this), *GetPathNameSafe(&NavData));
+		UE_LOGF(LogNavigation, Warning, "%s Navigation System: registration queue full! System: %ls NavData: %ls.", __FUNCTION__, *GetPathNameSafe(this), *GetPathNameSafe(&NavData));
 	}
 }
 
@@ -2790,6 +2793,10 @@ void UNavigationSystemV1::ProcessRegistrationCandidates()
 			NumNavDataProcessed++;
 			NavDataRegistrationQueue.RemoveAtSwap(CandidateIndex);
 		}
+		else
+		{
+			UE_LOGF(LogNavigation, Verbose, "%s   Skipped %ls because owning level not visible.", __FUNCTION__, *GetFullNameSafe(NavDataPtr));
+		}
 	}	
 	
 	if (NumNavDataProcessed)
@@ -4240,6 +4247,8 @@ void UNavigationSystemV1::CheckToLimitNavigationBoundsToLoadedRegions(FNavigatio
 
 void UNavigationSystemV1::AddNavigationBoundsUpdateRequest(const FNavigationBoundsUpdateRequest& UpdateRequest)
 {
+	UE_LOGF(LogNavigation, Verbose, "%s Update request: %i", __FUNCTION__, UpdateRequest.UpdateRequest);
+	
 	const int32 ExistingIdx = PendingNavBoundsUpdates.IndexOfByPredicate([&](const FNavigationBoundsUpdateRequest& Element) {
 		return UpdateRequest.NavBounds.UniqueID == Element.NavBounds.UniqueID;
 	});
@@ -4589,6 +4598,8 @@ uint8 UNavigationSystemV1::FillInstantiatedDataMask(TBitArray<>& OutInstantiated
 void UNavigationSystemV1::SpawnMissingNavigationDataInLevel(const TBitArray<>& InInstantiatedMask, ULevel* InLevel/*=nullptr*/)
 {
 	UWorld* NavWorld = GetWorld();
+	
+	UE_LOGF(LogNavigation, Verbose, "%s Level: %ls, OuterWorld: %ls", __FUNCTION__, *GetNameSafe(InLevel), *GetNameSafe(NavWorld));
 
 	ensure(SupportedAgents.Num() == InInstantiatedMask.Num());
 	int32 AllSupportedAgentsCount = InInstantiatedMask.Num();
@@ -4596,10 +4607,13 @@ void UNavigationSystemV1::SpawnMissingNavigationDataInLevel(const TBitArray<>& I
 	for (int32 AgentIndex = 0; AgentIndex < AllSupportedAgentsCount; ++AgentIndex)
 	{
 		const FNavDataConfig& NavConfig = SupportedAgents[AgentIndex];
+		
 		if (InInstantiatedMask[AgentIndex] == false
 			&& SupportedAgentsMask.Contains(AgentIndex)
 			&& NavConfig.GetNavDataClass<ANavigationData>() != nullptr)
 		{
+			UE_LOGF(LogNavigation, Verbose, "%s   AgentIndex: %i is supported. NavConfig name: %ls.", __FUNCTION__, AgentIndex, *NavConfig.Name.ToString());
+			
 			const ANavigationData* NavDataCDO = NavConfig.GetNavDataClass<ANavigationData>()->GetDefaultObject<ANavigationData>();
 			if (NavDataCDO == nullptr || !NavDataCDO->CanSpawnOnRebuild())
 			{
@@ -4612,6 +4626,19 @@ void UNavigationSystemV1::SpawnMissingNavigationDataInLevel(const TBitArray<>& I
 				// to be static, then we don't want to create an instance					
 				UE_LOG(LogNavigation, Log, TEXT("Not spawning navigation data for %s since indicated NavigationData type is not configured for dynamic generation")
 					, *NavConfig.Name.ToString());
+				
+				continue;
+			}
+
+			// Check if it's already wating for registration
+			const bool bIsWaitingForRegistration = NavDataRegistrationQueue.ContainsByPredicate([&NavConfig](const TObjectPtr<ANavigationData>& NavData)
+			{
+				return NavData->GetConfig().Name == NavConfig.Name;
+			});
+			
+			if (bIsWaitingForRegistration)
+			{
+				UE_LOGF(LogNavigation, Verbose, "%s   Skipping %ls already waiting for registration.", __FUNCTION__, *NavConfig.Name.ToString());
 				continue;
 			}
 
@@ -4622,7 +4649,7 @@ void UNavigationSystemV1::SpawnMissingNavigationDataInLevel(const TBitArray<>& I
 			}
 			else
 			{
-				UE_LOG(LogNavigation, Warning, TEXT("Was not able to create navigation data for SupportedAgent[%d]: %s"), AgentIndex, *NavConfig.Name.ToString());
+				UE_LOGF(LogNavigation, Warning, "%s   Was not able to create navigation data for SupportedAgent[%d]: %ls", __FUNCTION__, AgentIndex, *NavConfig.Name.ToString());
 			}
 		}
 	}
@@ -4908,6 +4935,8 @@ void UNavigationSystemV1::OnLevelAddedToWorld(ULevel* InLevel, UWorld* InWorld)
 	{
 		return;
 	}
+	
+	UE_LOGF(LogNavigation, Verbose, "%s Level: %ls, World: %ls", __FUNCTION__, *GetNameSafe(InLevel), *GetNameSafe(InWorld));
 
 	if ((IsNavigationSystemStatic() == false))
 	{
