[Navigation] Skip BVTree creation for tiles with low polygon count

#Synced-CL 51086784

diff --git a/Engine/Source/Runtime/NavigationSystem/Private/NavMesh/RecastNavMeshGenerator.cpp b/Engine/Source/Runtime/NavigationSystem/Private/NavMesh/RecastNavMeshGenerator.cpp
--- a/Engine/Source/Runtime/NavigationSystem/Private/NavMesh/RecastNavMeshGenerator.cpp
+++ b/Engine/Source/Runtime/NavigationSystem/Private/NavMesh/RecastNavMeshGenerator.cpp
@@ -4336,7 +4336,12 @@ bool FRecastTileGenerator::GenerateNavigationDataLayer(FNavMeshBuildContext& Bui
 		Params.cs = TileConfig.cs;
 		Params.ch = TileConfig.ch;
 		Params.tileResolutionLevel = (unsigned char)TileConfig.TileResolution;
-		Params.buildBvTree = TileConfig.bGenerateBVTree;
+
+		// Skip BVtree for tiles having few polygons since linear scan is faster for low poly counts. 
+		// This saves sizeof(dtBVNode) * polyCount * 2 bytes per tile.
+		constexpr int32 PolyCountThreshold = 12;
+		const bool bShouldBuildBvTree = TileConfig.bGenerateBVTree && (Params.polyCount > PolyCountThreshold);
+		Params.buildBvTree = bShouldBuildBvTree;
 
 #if WITH_NAVMESH_CLUSTER_LINKS
 		Params.clusterCount = IntCastChecked<unsigned short>(GenerationContext.ClusterSet->nclusters);
