[Navigation] Exposed navmesh AverageLayersPerTile and reduced it's default value (saves memory on tile pool allocation).

#Synced-CL 51265516

diff --git a/Engine/Source/Runtime/NavigationSystem/Private/NavMesh/PImplRecastNavMesh.cpp b/Engine/Source/Runtime/NavigationSystem/Private/NavMesh/PImplRecastNavMesh.cpp
--- a/Engine/Source/Runtime/NavigationSystem/Private/NavMesh/PImplRecastNavMesh.cpp
+++ b/Engine/Source/Runtime/NavigationSystem/Private/NavMesh/PImplRecastNavMesh.cpp
@@ -728,7 +728,8 @@ void FPImplRecastNavMesh::Serialize( FArchive& Ar, int32 NavMeshVersion )
 					Status = DetourNavMesh->addTile(TileData, TileDataSize, DT_TILE_FREE_DATA, TileRef, NULL);
 					if (dtStatusDetail(Status, DT_OUT_OF_MEMORY))
 					{
-						UE_LOG(LogNavigation, Warning, TEXT("%hs Failed to add tile (%d,%d:%d), %d tile limit reached in %s. If using FixedTilePoolSize, try increasing the TilePoolSize or using bigger tiles."),
+						UE_LOGF(LogNavigation, Warning, 
+							"%s Failed to add tile (%d,%d:%d), %d tile limit reached in %ls. Increase AverageLayersPerTile or, if using FixedTilePoolSize, increase the TilePoolSize or use bigger tiles.",
 							__FUNCTION__, TileHeader->x, TileHeader->y, TileHeader->layer, DetourNavMesh->getMaxTiles(), *NavMeshOwner->GetFullName());
 					}
 
diff --git a/Engine/Source/Runtime/NavigationSystem/Private/NavMesh/RecastNavMesh.cpp b/Engine/Source/Runtime/NavigationSystem/Private/NavMesh/RecastNavMesh.cpp
--- a/Engine/Source/Runtime/NavigationSystem/Private/NavMesh/RecastNavMesh.cpp
+++ b/Engine/Source/Runtime/NavigationSystem/Private/NavMesh/RecastNavMesh.cpp
@@ -486,6 +486,7 @@ ARecastNavMesh::ARecastNavMesh(const FObjectInitializer& ObjectInitializer)
 	, bDistinctlyDrawTilesBeingBuilt(true)
 	, DrawOffset(10.f)
 	, TilePoolSize(1024)
+	, AverageLayersPerTile(3)
 	, MaxVerticalMergeError(INT_MAX) // By default, ignore vertical error
 	, MaxSimplificationError(1.3f)	// from RecastDemo
 	, SimplificationElevationRatio(0.f)	// By default, ignore contour simplification from elevation
diff --git a/Engine/Source/Runtime/NavigationSystem/Private/NavMesh/RecastNavMeshDataChunk.cpp b/Engine/Source/Runtime/NavigationSystem/Private/NavMesh/RecastNavMeshDataChunk.cpp
--- a/Engine/Source/Runtime/NavigationSystem/Private/NavMesh/RecastNavMeshDataChunk.cpp
+++ b/Engine/Source/Runtime/NavigationSystem/Private/NavMesh/RecastNavMeshDataChunk.cpp
@@ -267,8 +267,9 @@ TArray<FNavTileRef> URecastNavMeshDataChunk::AttachTiles(ARecastNavMesh& NavMesh
 				{
 					if (dtStatusDetail(status, DT_OUT_OF_MEMORY))
 					{
-						UE_LOG(LogNavigation, Warning, TEXT("%s> Failed to add tile (%d,%d:%d), %d tile limit reached! (from: %s). If using FixedTilePoolSize, try increasing the TilePoolSize or using bigger tiles."),
-							*NavMesh.GetName(), Header->x, Header->y, Header->layer, DetourNavMesh->getMaxTiles(), ANSI_TO_TCHAR(__FUNCTION__));
+						UE_LOGF(LogNavigation, Warning, 
+							"%ls> Failed to add tile (%d,%d:%d), %d tile limit reached! (from: %s). Increase AverageLayersPerTile or, if using FixedTilePoolSize, increase the TilePoolSize or use bigger tiles.",
+							*NavMesh.GetName(), Header->x, Header->y, Header->layer, DetourNavMesh->getMaxTiles(), __FUNCTION__);
 					}
 					
 					continue;
diff --git a/Engine/Source/Runtime/NavigationSystem/Private/NavMesh/RecastNavMeshGenerator.cpp b/Engine/Source/Runtime/NavigationSystem/Private/NavMesh/RecastNavMeshGenerator.cpp
--- a/Engine/Source/Runtime/NavigationSystem/Private/NavMesh/RecastNavMeshGenerator.cpp
+++ b/Engine/Source/Runtime/NavigationSystem/Private/NavMesh/RecastNavMeshGenerator.cpp
@@ -5129,7 +5129,6 @@ static bool IsGameStaticNavMesh(ARecastNavMesh* InNavMesh)
 FRecastNavMeshGenerator::FRecastNavMeshGenerator(ARecastNavMesh& InDestNavMesh)
 	: NumActiveTiles(0)
 	, MaxTileGeneratorTasks(1)
-	, AvgLayersPerTile(8.0f)
 	, DestNavMesh(&InDestNavMesh)
 	, bInitialized(false)
 	, bRestrictBuildingToActiveTiles(false)
@@ -5254,6 +5253,8 @@ void FRecastNavMeshGenerator::ConfigureBuildProperties(FRecastBuildConfig& OutCo
 	OutConfig.AgentMaxClimb = AgentMaxClimb;
 	OutConfig.AgentRadius = AgentRadius;
 
+	OutConfig.AverageLayersPerTile = FMath::Max(1.f, DestNavMesh->AverageLayersPerTile);
+
 	UE::NavMesh::Private::ComputeConfigBorderSizes(IsGeneratingLinks(), OutConfig);
 
 	OutConfig.maxEdgeLen = (int32)(1200.0f / CellSize);
@@ -5643,7 +5644,8 @@ void FRecastNavMeshGenerator::CalcNavMeshProperties(int32& MaxTiles, int32& MaxP
 	int32 MaxRequestedTiles = 0;
 	if (DestNavMesh->IsResizable())
 	{
-		MaxRequestedTiles = UE::NavMesh::Private::CalculateMaxTilesCount(InclusionBounds, Config.GetTileSizeUU(), AvgLayersPerTile, DestNavMesh->NavMeshVersion);
+		MaxRequestedTiles = UE::NavMesh::Private::CalculateMaxTilesCount(
+			InclusionBounds, Config.GetTileSizeUU(), Config.AverageLayersPerTile, DestNavMesh->NavMeshVersion);
 	}
 	else
 	{
@@ -5827,7 +5829,9 @@ void FRecastNavMeshGenerator::OnNavigationBoundsChanged()
 		if (NavSys && !NavSys->IsNavigationBuildingLocked())
 		{
 			// Check whether Navmesh size needs to be changed
-			const int32 MaxRequestedTiles = UE::NavMesh::Private::CalculateMaxTilesCount(InclusionBounds, Config.GetTileSizeUU(), AvgLayersPerTile, DestNavMesh->NavMeshVersion);
+			const int32 MaxRequestedTiles = UE::NavMesh::Private::CalculateMaxTilesCount(
+				InclusionBounds, Config.GetTileSizeUU(), Config.AverageLayersPerTile, DestNavMesh->NavMeshVersion);
+
 			if (DetourMesh->getMaxTiles() != MaxRequestedTiles)
 			{
 				UE_LOG(LogNavigation, Log, TEXT("%s> Navigation bounds changed, rebuilding navmesh"), *DestNavMesh->GetName());
@@ -6226,8 +6230,9 @@ void FRecastNavMeshGenerator::AddGeneratedTileLayer(int32 LayerIndex, FRecastTil
 			{
 				if (dtStatusDetail(status, DT_OUT_OF_MEMORY))
 				{
-					UE_LOG(LogNavigation, Error, TEXT("%s> Failed to add tile (%d,%d:%d), %d tile limit reached! (from %s). If using FixedTilePoolSize, try increasing the TilePoolSize or using bigger tiles."),
-						*DestNavMesh->GetName(), TileX, TileY, LayerIndex, DetourMesh->getMaxTiles(), ANSI_TO_TCHAR(__FUNCTION__));
+					UE_LOGF(LogNavigation, Warning, 
+						"%ls> Failed to add tile (%d,%d:%d), %d tile limit reached! (from %s). Increase AverageLayersPerTile or, if using FixedTilePoolSize, increase the TilePoolSize or use bigger tiles.",
+						*DestNavMesh->GetName(), TileX, TileY, LayerIndex, DetourMesh->getMaxTiles(), __FUNCTION__);
 				}
 			}
 			else
diff --git a/Engine/Source/Runtime/NavigationSystem/Public/NavMesh/RecastNavMesh.h b/Engine/Source/Runtime/NavigationSystem/Public/NavMesh/RecastNavMesh.h
--- a/Engine/Source/Runtime/NavigationSystem/Public/NavMesh/RecastNavMesh.h
+++ b/Engine/Source/Runtime/NavigationSystem/Public/NavMesh/RecastNavMesh.h
@@ -693,6 +693,12 @@ class ARecastNavMesh : public ANavigationData
 	UPROPERTY(EditAnywhere, Category=Generation, config, meta=(ClampMin = "300.0"))
 	float TileSizeUU;
 
+	/** Expected average number of layers per tile. 
+	 *  This is a factor in the number of tiles the navmesh can hold when not using bFixedTilePoolSize.
+	 *  Log warnings are issued if this is too low. */
+	UPROPERTY(EditAnywhere, Category=Generation, config, meta=(ClampMin = "1.0"))
+	float AverageLayersPerTile;
+
 	/**
 	 * Note that we are not using _DEPRECATED on the following deprecated properties
 	 * since it prevents the property from being serialized back which can break the
diff --git a/Engine/Source/Runtime/NavigationSystem/Public/NavMesh/RecastNavMeshGenerator.h b/Engine/Source/Runtime/NavigationSystem/Public/NavMesh/RecastNavMeshGenerator.h
--- a/Engine/Source/Runtime/NavigationSystem/Public/NavMesh/RecastNavMeshGenerator.h
+++ b/Engine/Source/Runtime/NavigationSystem/Public/NavMesh/RecastNavMeshGenerator.h
@@ -94,6 +94,8 @@ struct FRecastBuildConfig : public rcConfig
 	int32 AgentIndex;
 	/** Resolution level */ 
 	ENavigationDataResolution TileResolution;
+	/** Average layers per tile */
+	float AverageLayersPerTile;
 	/** Ledge filtering mode */
 	ENavigationLedgeSlopeFilterMode LedgeSlopeFilterMode;
 	/** Is the config completely setup */
@@ -121,6 +123,7 @@ struct FRecastBuildConfig : public rcConfig
 		JumpConfigs.Empty();
 		AgentIndex = 0;
 		TileResolution = ENavigationDataResolution::Default;
+		AverageLayersPerTile = 3;
 		LedgeSlopeFilterMode = ENavigationLedgeSlopeFilterMode::Recast;
 		bIsTileSetupConfigCompleted = false;
 		LinkSpillDistance = 0.f;
@@ -1020,8 +1023,6 @@ protected:
 	 */
 	int32 MaxTileGeneratorTasks;
 
-	FVector::FReal AvgLayersPerTile;
-
 	/** Total bounding box that includes all volumes, in unreal units. */
 	FBox TotalNavBounds;
 
