Fix for NavLinkProxies used by ChildActors not working properly on startup in PIE.

#Synced-CL 51124555

diff --git a/Engine/Source/Runtime/NavigationSystem/Private/NavLinkCustomComponent.cpp b/Engine/Source/Runtime/NavigationSystem/Private/NavLinkCustomComponent.cpp
--- a/Engine/Source/Runtime/NavigationSystem/Private/NavLinkCustomComponent.cpp
+++ b/Engine/Source/Runtime/NavigationSystem/Private/NavLinkCustomComponent.cpp
@@ -266,6 +266,17 @@ void UNavLinkCustomComponent::OnRegister()
 		if (OldCustomLinkId != CustomLinkId)
 		{
 			Modify(true);
+
+			// Baked navmesh has off-mesh connections with the old ID — rebuild the tile
+			// so they get regenerated with the new one.
+			if (UNavigationSystemV1* NavSys = FNavigationSystem::GetCurrent<UNavigationSystemV1>(GetWorld()))
+			{
+				const FBox LinkBounds = UNavigationSystemV1::ComputeCustomLinkBounds(*this);
+				if (LinkBounds.IsValid && !LinkBounds.GetSize().IsNearlyZero())
+				{
+					NavSys->AddDirtyArea(LinkBounds, ENavigationDirtyFlag::Geometry);
+				}
+			}
 		}
 	}
 #else // !WITH_EDITOR
