Fixed various callsites of FString::Printf/Appendf that used scoped enums

#Synced-CL 51468882

diff --git a/Engine/Plugins/Experimental/ChaosMover/Source/ChaosMover/Private/Character/ChaosCharacterInputs.cpp b/Engine/Plugins/Experimental/ChaosMover/Source/ChaosMover/Private/Character/ChaosCharacterInputs.cpp
--- a/Engine/Plugins/Experimental/ChaosMover/Source/ChaosMover/Private/Character/ChaosCharacterInputs.cpp
+++ b/Engine/Plugins/Experimental/ChaosMover/Source/ChaosMover/Private/Character/ChaosCharacterInputs.cpp
@@ -20,7 +20,7 @@ void FChaosMoverLaunchInputs::ToString(FAnsiStringBuilderBase& Out) const
 	Super::ToString(Out);
 
 	Out.Appendf("LaunchVelocityOrImpulse: X=%.2f Y=%.2f Z=%.2f |", LaunchVelocityOrImpulse.X, LaunchVelocityOrImpulse.Y, LaunchVelocityOrImpulse.Z);
-	Out.Appendf("Mode: %u\n", Mode);
+	Out.Appendf("Mode: %u\n", uint32{EnumToUnderlyingType(Mode)});
 }
 
 bool FChaosMoverLaunchInputs::ShouldReconcile(const FMoverDataStructBase& AuthorityState) const
diff --git a/Engine/Plugins/Experimental/Mover/Source/Mover/Private/MoverDataModelTypes.cpp b/Engine/Plugins/Experimental/Mover/Source/Mover/Private/MoverDataModelTypes.cpp
--- a/Engine/Plugins/Experimental/Mover/Source/Mover/Private/MoverDataModelTypes.cpp
+++ b/Engine/Plugins/Experimental/Mover/Source/Mover/Private/MoverDataModelTypes.cpp
@@ -96,7 +96,7 @@ void FCharacterDefaultInputs::ToString(FAnsiStringBuilderBase& Out) const
 {
 	Super::ToString(Out);
 
-	Out.Appendf("MoveInput: %s (Type %i)\n", TCHAR_TO_ANSI(*MoveInput.ToCompactString()), MoveInputType);
+	Out.Appendf("MoveInput: %s (Type %u)\n", TCHAR_TO_ANSI(*MoveInput.ToCompactString()), uint32{EnumToUnderlyingType(MoveInputType)});
 	Out.Appendf("OrientationIntent: X=%.2f Y=%.2f Z=%.2f\n", OrientationIntent.X, OrientationIntent.Y, OrientationIntent.Z);
 	Out.Appendf("ControlRotation: P=%.2f Y=%.2f R=%.2f\n", ControlRotation.Pitch, ControlRotation.Yaw, ControlRotation.Roll);
 	Out.Appendf("SuggestedMovementMode: %s\n", TCHAR_TO_ANSI(*SuggestedMovementMode.ToString()));
diff --git a/Engine/Plugins/Experimental/Mover/Source/Mover/Private/PhysicsMover/PhysicsMoverSimulationTypes.cpp b/Engine/Plugins/Experimental/Mover/Source/Mover/Private/PhysicsMover/PhysicsMoverSimulationTypes.cpp
--- a/Engine/Plugins/Experimental/Mover/Source/Mover/Private/PhysicsMover/PhysicsMoverSimulationTypes.cpp
+++ b/Engine/Plugins/Experimental/Mover/Source/Mover/Private/PhysicsMover/PhysicsMoverSimulationTypes.cpp
@@ -180,7 +180,7 @@ void FMoverLaunchInputs::ToString(FAnsiStringBuilderBase& Out) const
 	Super::ToString(Out);
 
 	Out.Appendf("LaunchVelocity: X=%.2f Y=%.2f Z=%.2f\n", LaunchVelocity.X, LaunchVelocity.Y, LaunchVelocity.Z);
-	Out.Appendf("Mode: %u\n", Mode);
+	Out.Appendf("Mode: %u\n", uint32{EnumToUnderlyingType(Mode)});
 }
 
 bool FMoverLaunchInputs::ShouldReconcile(const FMoverDataStructBase& AuthorityState) const
diff --git a/Engine/Plugins/FX/Niagara/Source/Niagara/Private/NiagaraSystemInstance.cpp b/Engine/Plugins/FX/Niagara/Source/Niagara/Private/NiagaraSystemInstance.cpp
--- a/Engine/Plugins/FX/Niagara/Source/Niagara/Private/NiagaraSystemInstance.cpp
+++ b/Engine/Plugins/FX/Niagara/Source/Niagara/Private/NiagaraSystemInstance.cpp
@@ -2491,7 +2491,7 @@ void FNiagaraSystemInstance::DumpStalledInfo()
 	Builder.Appendf(TEXT("ConcurrentTickBatchGraphEvent Complete (%d)\n"), ConcurrentTickBatchGraphEvent ? ConcurrentTickBatchGraphEvent->IsComplete() : true);
 	Builder.Appendf(TEXT("FinalizePending (%d)\n"), FinalizeRef.IsPending());
 	Builder.Appendf(TEXT("SystemInstanceIndex (%d)\n"), SystemInstanceIndex);
-	Builder.Appendf(TEXT("SystemInstanceState (%d)\n"), SystemInstanceState);
+	Builder.Appendf(TEXT("SystemInstanceState (%d)\n"), uint32{EnumToUnderlyingType(SystemInstanceState)});
 
 	UE_LOG(LogNiagara, Fatal, TEXT("FNiagaraSystemInstance is stalled.\n%s"), Builder.ToString());
 }
diff --git a/Engine/Plugins/Runtime/NetworkPredictionExtras/Source/NetworkPredictionExtras/Public/CharacterMotionSimulation.h b/Engine/Plugins/Runtime/NetworkPredictionExtras/Source/NetworkPredictionExtras/Public/CharacterMotionSimulation.h
--- a/Engine/Plugins/Runtime/NetworkPredictionExtras/Source/NetworkPredictionExtras/Public/CharacterMotionSimulation.h
+++ b/Engine/Plugins/Runtime/NetworkPredictionExtras/Source/NetworkPredictionExtras/Public/CharacterMotionSimulation.h
@@ -132,7 +132,7 @@ struct NETWORKPREDICTIONEXTRAS_API FCharacterMotionSyncState
 
 	void ToString(FAnsiStringBuilderBase& Out) const
 	{
-		Out.Appendf("MovementMode: %d\n", MovementMode);
+		Out.Appendf("MovementMode: %u\n", uint32{EnumToUnderlyingType(MovementMode)});
 		Out.Appendf("Loc: X=%.2f Y=%.2f Z=%.2f\n", Location.X, Location.Y, Location.Z);
 		Out.Appendf("Vel: X=%.2f Y=%.2f Z=%.2f\n", Velocity.X, Velocity.Y, Velocity.Z);
 		Out.Appendf("Rot: P=%.2f Y=%.2f R=%.2f\n", Rotation.Pitch, Rotation.Yaw, Rotation.Roll);
diff --git a/Engine/Source/Runtime/Engine/Private/WorldPartition/WorldPartitionStreamingSource.cpp b/Engine/Source/Runtime/Engine/Private/WorldPartition/WorldPartitionStreamingSource.cpp
--- a/Engine/Source/Runtime/Engine/Private/WorldPartition/WorldPartitionStreamingSource.cpp
+++ b/Engine/Source/Runtime/Engine/Private/WorldPartition/WorldPartitionStreamingSource.cpp
@@ -60,8 +60,8 @@ FString FWorldPartitionStreamingSource::ToString() const
 	const TCHAR* VelocityUnitStr = FUnitConversion::GetUnitDisplayString(VelocityUnitEnum);
 
 	StringBuilder.Appendf(
-		TEXT("Priority: %d | %s | %s | %s | Pos: X=%lld,Y=%lld,Z=%lld | Rot: %s | Vel: %d %s"),
-		Priority,
+		TEXT("Priority: %u | %s | %s | %s | Pos: X=%lld,Y=%lld,Z=%lld | Rot: %s | Vel: %d %s"),
+		uint32{EnumToUnderlyingType(Priority)},
 		bRemote ? TEXT("Remote") : TEXT("Local"),
 		GetStreamingSourceTargetStateName(TargetState),
 		bBlockOnSlowLoading ? TEXT("Blocking") : TEXT("NonBlocking"),
diff --git a/Engine/Source/Runtime/NavigationSystem/Private/NavMesh/RecastNavMeshGenerator.cpp b/Engine/Source/Runtime/NavigationSystem/Private/NavMesh/RecastNavMeshGenerator.cpp
--- a/Engine/Source/Runtime/NavigationSystem/Private/NavMesh/RecastNavMeshGenerator.cpp
+++ b/Engine/Source/Runtime/NavigationSystem/Private/NavMesh/RecastNavMeshGenerator.cpp
@@ -6903,7 +6903,7 @@ void FRecastNavMeshGenerator::SortPendingBuildTiles()
 				const FVector::FReal TileSizeInWorldUnits = Config.GetTileSizeUU();
 				const FBox TileBox = FRecastTileGenerator::CalculateTileBounds(Element.Coord.X, Element.Coord.Y, FVector::ZeroVector, TotalNavBounds, TileSizeInWorldUnits);
 				const ARecastNavMesh* const OwnerNav = GetOwner();
-				UE_VLOG_BOX(OwnerNav, LogNavigation, VeryVerbose, TileBox, FColor::Silver, TEXT("Index: %i, Priority %i"), Index, Element.SortingPriority);
+				UE_VLOG_BOX(OwnerNav, LogNavigation, VeryVerbose, TileBox, FColor::Silver, TEXT("Index: %i, Priority %u"), Index, uint32{EnumToUnderlyingType(Element.SortingPriority)});
 			}
 		});
 	}
diff --git a/Engine/Source/Runtime/NavigationSystem/Private/NavigationSystem.cpp b/Engine/Source/Runtime/NavigationSystem/Private/NavigationSystem.cpp
--- a/Engine/Source/Runtime/NavigationSystem/Private/NavigationSystem.cpp
+++ b/Engine/Source/Runtime/NavigationSystem/Private/NavigationSystem.cpp
@@ -6031,7 +6031,7 @@ void UNavigationSystemV1::UpdateInvokers()
 			for (const auto& InvokerData : InvokerLocations)
 			{
 				UE_VLOG_CYLINDER(this, LogNavInvokers, Log, InvokerData.Location, InvokerData.Location + FVector(0, 0, 20), InvokerData.RadiusMax, FColorList::Blue, TEXT(""));
-				UE_VLOG_CYLINDER(this, LogNavInvokers, Log, InvokerData.Location, InvokerData.Location + FVector(0, 0, 20), InvokerData.RadiusMin, FColorList::CadetBlue, TEXT("Priority %u"), InvokerData.Priority);
+				UE_VLOG_CYLINDER(this, LogNavInvokers, Log, InvokerData.Location, InvokerData.Location + FVector(0, 0, 20), InvokerData.RadiusMin, FColorList::CadetBlue, TEXT("Priority %u"), uint32{EnumToUnderlyingType(InvokerData.Priority)});
 			}
 #endif // ENABLE_VISUAL_LOG
 		}
