[Navigation] Fix description/tooltip for bForceRebuildOnLoad

#Synced-CL 51624139

diff --git a/Engine/Source/Runtime/NavigationSystem/Public/NavigationData.h b/Engine/Source/Runtime/NavigationSystem/Public/NavigationData.h
--- a/Engine/Source/Runtime/NavigationSystem/Public/NavigationData.h
+++ b/Engine/Source/Runtime/NavigationSystem/Public/NavigationData.h
@@ -556,8 +556,8 @@ protected:
 	// game-time config
 	//----------------------------------------------------------------------//
 	
-	/** By default navigation will skip the first update after being successfully loaded
-	*  setting bForceRebuildOnLoad to false can override this behavior */
+	/** By default, navigation will skip the first update after being successfully loaded.
+	*  Setting bForceRebuildOnLoad to true can override this behavior. */
 	UPROPERTY(config, EditAnywhere, Category = Runtime)
 	uint32 bForceRebuildOnLoad : 1;
 
