Auto-rebuild navmesh for levels missing saved ARecastNavMesh

When traveling to a level with a NavMeshBoundsVolume but no saved ARecastNavMesh actor, the nav system spawns a fresh empty one. No rebuild is triggered for freshly spawned nav data in a game world - the rebuild pipeline assumes nav data was loaded from disk and is valid. As a result, NPCs can't pathfind.

Adds bRequiresInitialRebuild flag on ANavigationData, set when SpawnMissingNavigationData creates nav data in a game world. The flag feeds into SupportsRuntimeGeneration() and NeedsRebuildOnLoad(), unblocking the octree, generator, and tile pipeline. After octree population, RebuildAll(bIsLoadTime=true) triggers an async rebuild matching the existing editor path.

Adds DoesAnyNavDataRequireGeometryStorage() to correctly gate nav octree geometry storage. This is distinct from DoesAnyNavDataSupportRuntimeGeneration() - DynamicModifiersOnly supports runtime generation but does not require geometry storage, and should not cause geometry to be stored in the octree.

No impact when nav data is properly saved - flag defaults to false.

#Synced-CL 52035316

diff --git a/Engine/Source/Runtime/NavigationSystem/Private/NavMesh/RecastNavMesh.cpp b/Engine/Source/Runtime/NavigationSystem/Private/NavMesh/RecastNavMesh.cpp
--- a/Engine/Source/Runtime/NavigationSystem/Private/NavMesh/RecastNavMesh.cpp
+++ b/Engine/Source/Runtime/NavigationSystem/Private/NavMesh/RecastNavMesh.cpp
@@ -3161,6 +3161,14 @@ void ARecastNavMesh::EnsureBuildCompletion()
 
 void ARecastNavMesh::OnNavMeshGenerationFinished()
 {
+	// If this navmesh was spawned at runtime to fill a missing nav data gap,
+	// the one-time rebuild is now complete. Clear the flag so it behaves
+	// as a Static navmesh going forward.
+	if (RequiresInitialRebuild())
+	{
+		ClearRequiresInitialRebuild();
+	}
+
 	UWorld* World = GetWorld();
 
 	if (IsValid(World))
@@ -4254,7 +4262,7 @@ bool ARecastNavMesh::NeedsRebuild() const
 bool ARecastNavMesh::SupportsRuntimeGeneration() const
 {
 	// Generator should be disabled for Static navmesh
-	return (RuntimeGeneration != ERuntimeGenerationType::Static);
+	return (RuntimeGeneration != ERuntimeGenerationType::Static) || RequiresInitialRebuild();
 }
 
 bool ARecastNavMesh::SupportsStreaming() const
diff --git a/Engine/Source/Runtime/NavigationSystem/Private/NavMesh/RecastNavMeshGenerator.cpp b/Engine/Source/Runtime/NavigationSystem/Private/NavMesh/RecastNavMeshGenerator.cpp
--- a/Engine/Source/Runtime/NavigationSystem/Private/NavMesh/RecastNavMeshGenerator.cpp
+++ b/Engine/Source/Runtime/NavigationSystem/Private/NavMesh/RecastNavMeshGenerator.cpp
@@ -5116,10 +5116,14 @@ int32 CalculateMaxTilesCount(const TNavStatArray<FBox>& NavigableAreas, FVector:
 }
 } // UE::NavMesh::Private
 
-// Whether navmesh is static, does not support rebuild from geometry
-static bool IsGameStaticNavMesh(ARecastNavMesh* InNavMesh)
+// Whether geometry rebuild is disabled for this navmesh.
+// True for Static and DynamicModifiersOnly navmesh in game worlds (neither supports rebuild from geometry),
+// unless the navmesh requires an initial rebuild to recover missing nav data.
+static bool IsGeometryRebuildDisabled(ARecastNavMesh* InNavMesh)
 {
-	return (InNavMesh->GetWorld()->IsGameWorld() && InNavMesh->GetRuntimeGenerationMode() != ERuntimeGenerationType::Dynamic);
+	return InNavMesh->GetWorld()->IsGameWorld()
+		&& InNavMesh->GetRuntimeGenerationMode() != ERuntimeGenerationType::Dynamic
+		&& !InNavMesh->RequiresInitialRebuild();
 }
 
 //----------------------------------------------------------------------//
@@ -5381,11 +5385,11 @@ void FRecastNavMeshGenerator::Init()
 	bool bRecreateNavmesh = true;
 	if (DestNavMesh->HasValidNavmesh())
 	{
-		const bool bGameStaticNavMesh = IsGameStaticNavMesh(DestNavMesh);
+		const bool bGeometryRebuildDisabled = IsGeometryRebuildDisabled(DestNavMesh);
 		const dtNavMeshParams* SavedNavParams = DestNavMesh->GetRecastNavMeshImpl()->DetourNavMesh->getParams();
 		if (SavedNavParams)
 		{
-			if (bGameStaticNavMesh)
+			if (bGeometryRebuildDisabled)
 			{
 				bRecreateNavmesh = false;
 				MaxTiles = SavedNavParams->maxTiles;
@@ -5823,7 +5827,7 @@ void FRecastNavMeshGenerator::OnNavigationBoundsChanged()
 	UpdateNavigationBounds();
 	
 	dtNavMesh* DetourMesh = DestNavMesh->GetRecastNavMeshImpl() ? DestNavMesh->GetRecastNavMeshImpl()->GetRecastMesh() : nullptr;
-	if (!IsGameStaticNavMesh(DestNavMesh) && DestNavMesh->IsResizable() && DetourMesh)
+	if (!IsGeometryRebuildDisabled(DestNavMesh) && DestNavMesh->IsResizable() && DetourMesh)
 	{
 		const UNavigationSystemV1* NavSys = FNavigationSystem::GetCurrent<UNavigationSystemV1>(GetWorld());
 		if (NavSys && !NavSys->IsNavigationBuildingLocked())
@@ -6585,7 +6589,7 @@ void FRecastNavMeshGenerator::MarkDirtyTiles(const TArray<FNavigationDirtyArea>&
 	const FVector::FReal TileSizeInWorldUnits = Config.GetTileSizeUU();
 	check(TileSizeInWorldUnits > 0);
 
-	const bool bGameStaticNavMesh = IsGameStaticNavMesh(DestNavMesh);
+	const bool bGeometryRebuildDisabled = IsGeometryRebuildDisabled(DestNavMesh);
 
 	const UNavigationSystemV1* NavSys = FNavigationSystem::GetCurrent<UNavigationSystemV1>(GetWorld());
 	const ARecastNavMesh* const OwnerNav = GetOwner();
@@ -6616,8 +6620,8 @@ void FRecastNavMeshGenerator::MarkDirtyTiles(const TArray<FNavigationDirtyArea>&
 				continue;
 			}
 
-			// Game world static navmeshes accept only area modifiers updates
-			if (bGameStaticNavMesh && (!DirtyArea.HasFlag(ENavigationDirtyFlag::DynamicModifier) || DirtyArea.HasFlag(ENavigationDirtyFlag::NavigationBounds)))
+			// Navmeshes with geometry rebuild disabled accept only area modifier updates
+			if (bGeometryRebuildDisabled && (!DirtyArea.HasFlag(ENavigationDirtyFlag::DynamicModifier) || DirtyArea.HasFlag(ENavigationDirtyFlag::NavigationBounds)))
 			{
 				continue;
 			}
@@ -6698,9 +6702,9 @@ void FRecastNavMeshGenerator::MarkDirtyTiles(const TArray<FNavigationDirtyArea>&
 											
 						FPendingTileElement Element;
 						Element.Coord = FIntPoint(TileX, TileY);
-						// Make sure to prevent bRebuildGeometry for game world static navmeshes.
-						// Game world static navmeshes accept only area modifiers updates. Rebuilding geometry would bRegenerateCompressedLayers without having the geometry for them.
-						Element.bRebuildGeometry = !bGameStaticNavMesh && (DirtyArea.HasFlag(ENavigationDirtyFlag::Geometry) || DirtyArea.HasFlag(ENavigationDirtyFlag::NavigationBounds));
+						// Make sure to prevent bRebuildGeometry when geometry rebuild is disabled.
+						// Navmeshes with geometry rebuild disabled accept only area modifier updates. Rebuilding geometry would bRegenerateCompressedLayers without having the geometry for them.
+						Element.bRebuildGeometry = !bGeometryRebuildDisabled && (DirtyArea.HasFlag(ENavigationDirtyFlag::Geometry) || DirtyArea.HasFlag(ENavigationDirtyFlag::NavigationBounds));
 						Element.CreationTime = CurrentTimeSeconds;
 						if (Element.bRebuildGeometry == false)
 						{
@@ -6970,7 +6974,7 @@ TArray<FNavTileRef> FRecastNavMeshGenerator::ProcessTileTasksAsyncAndGetUpdatedT
 	LLM_SCOPE(ELLMTag::NavigationRecast);
 
 	TArray<FNavTileRef> UpdatedTiles;
-	const bool bGameStaticNavMesh = IsGameStaticNavMesh(DestNavMesh);
+	const bool bGeometryRebuildDisabled = IsGeometryRebuildDisabled(DestNavMesh);
 
 	int32 NumProcessedTasks = 0;
 	// Submit pending tile elements
@@ -7007,7 +7011,7 @@ TArray<FNavTileRef> FRecastNavMeshGenerator::ProcessTileTasksAsyncAndGetUpdatedT
 				
 				RunningDirtyTiles.Add(RunningElement);
 			}
-			else if (!bGameStaticNavMesh)
+			else if (!bGeometryRebuildDisabled)
 			{
 				RemoveLayers(PendingElement.Coord, UpdatedTiles);
 			}
@@ -7169,7 +7173,7 @@ TArray<FNavTileRef> FRecastNavMeshGenerator::ProcessTileTasksSyncTimeSlicedAndGe
 	{
 		SyncTimeSlicedData.TimeSliceManager->GetTimeSlicer().StartTimeSlice();
 
-		const bool bGameStaticNavMesh = IsGameStaticNavMesh(DestNavMesh);
+		const bool bGeometryRebuildDisabled = IsGeometryRebuildDisabled(DestNavMesh);
 
 		// Submit pending tile elements
 		do
@@ -7206,7 +7210,7 @@ TArray<FNavTileRef> FRecastNavMeshGenerator::ProcessTileTasksSyncTimeSlicedAndGe
 				{
 					SyncTimeSlicedData.ProcessTileTasksSyncState = EProcessTileTasksSyncTimeSlicedState::Finish;
 
-					if (!bGameStaticNavMesh)
+					if (!bGeometryRebuildDisabled)
 					{
 						RemoveLayers(TileLocation, UpdatedTiles);
 					}
@@ -7361,7 +7365,7 @@ TArray<FNavTileRef> FRecastNavMeshGenerator::ProcessTileTasksSyncAndGetUpdatedTi
 {
 	QUICK_SCOPE_CYCLE_COUNTER(STAT_RecastNavMeshGenerator_ProcessTileTasksSync);
 
-	const bool bGameStaticNavMesh = IsGameStaticNavMesh(DestNavMesh);
+	const bool bGeometryRebuildDisabled = IsGeometryRebuildDisabled(DestNavMesh);
 	int32 NumProcessedTasks = 0;
 	TArray<FNavTileRef> UpdatedTiles;
 	FIntPoint TileLocation;
@@ -7384,7 +7388,7 @@ TArray<FNavTileRef> FRecastNavMeshGenerator::ProcessTileTasksSyncAndGetUpdatedTi
 
 			StoreCompressedTileCacheLayers(TileGeneratorRef, TileLocation.X, TileLocation.Y);
 		}
-		else if (!bGameStaticNavMesh)
+		else if (!bGeometryRebuildDisabled)
 		{
 			RemoveLayers(TileLocation, UpdatedTiles);
 		}
diff --git a/Engine/Source/Runtime/NavigationSystem/Private/NavigationData.cpp b/Engine/Source/Runtime/NavigationSystem/Private/NavigationData.cpp
--- a/Engine/Source/Runtime/NavigationSystem/Private/NavigationData.cpp
+++ b/Engine/Source/Runtime/NavigationSystem/Private/NavigationData.cpp
@@ -183,6 +183,7 @@ ANavigationData::ANavigationData(const FObjectInitializer& ObjectInitializer)
 	, bCanBeMainNavData(true)
 	, bCanSpawnOnRebuild(true)
 	, RuntimeGeneration(ERuntimeGenerationType::LegacyGeneration) //TODO: set to a valid value once bRebuildAtRuntime_DEPRECATED is removed
+	, bRequiresInitialRebuild(false)
 	, DataVersion(NAVDATAVER_LATEST)
 	, FindPathImplementation(NULL)
 	, FindHierarchicalPathImplementation(NULL)
diff --git a/Engine/Source/Runtime/NavigationSystem/Private/NavigationSystem.cpp b/Engine/Source/Runtime/NavigationSystem/Private/NavigationSystem.cpp
--- a/Engine/Source/Runtime/NavigationSystem/Private/NavigationSystem.cpp
+++ b/Engine/Source/Runtime/NavigationSystem/Private/NavigationSystem.cpp
@@ -1225,8 +1225,7 @@ bool UNavigationSystemV1::ConditionalPopulateNavOctree()
 		ConstructNavOctree();
 		if (DefaultOctreeController.IsValid())
 		{
-			const ERuntimeGenerationType RuntimeGenerationType = GetRuntimeGenerationType();
-			const bool bStoreNavGeometry = (RuntimeGenerationType == ERuntimeGenerationType::Dynamic);
+			const bool bStoreNavGeometry = DoesAnyNavDataRequireGeometryStorage();
 			DefaultOctreeController.SetNavigableGeometryStoringMode(bStoreNavGeometry ? FNavigationOctree::StoreNavGeometry : FNavigationOctree::SkipNavGeometry);
 			if (bStoreNavGeometry)
 			{
@@ -1475,6 +1474,20 @@ void UNavigationSystemV1::OnWorldInitDone(FNavigationSystemRunMode Mode)
 		}
 	}
 
+	// If any nav data was spawned at runtime (not loaded from disk), trigger a rebuild.
+	// Must run after ConditionalPopulateNavOctree so the geometry is available.
+	if (GetWorld()->IsGameWorld())
+	{
+		const bool bRequiresInitialRebuild = NavDataSet.ContainsByPredicate(
+			[](const ANavigationData* NavData) { return NavData && NavData->RequiresInitialRebuild(); });
+		if (bRequiresInitialRebuild)
+		{
+			const bool bIsLoadTime = true;
+			RebuildAll(bIsLoadTime);
+		}
+	}
+
+
 #if	WITH_EDITOR
 	if (FNavigationSystem::IsEditorRunMode(Mode))
 	{
@@ -2693,7 +2706,7 @@ void UNavigationSystemV1::ApplyWorldOffset(const FVector& InOffset, bool bWorldS
 	}
 
 	// Attempt at generation of new nav mesh after the shift
-	// dynamic navmesh, we regenerate completely
+	// Only fully dynamic navmesh needs a complete geometry rebuild after a world shift.
 	if (GetRuntimeGenerationType() == ERuntimeGenerationType::Dynamic)
 	{
 		//stop generators from building navmesh
@@ -4623,11 +4636,13 @@ void UNavigationSystemV1::SpawnMissingNavigationDataInLevel(const TBitArray<>& I
 			if (NavWorld->WorldType != EWorldType::Editor && NavDataCDO->GetRuntimeGenerationMode() == ERuntimeGenerationType::Static)
 			{
 				// if we're not in the editor, and specified navigation class is configured 
-				// to be static, then we don't want to create an instance					
-				UE_LOG(LogNavigation, Log, TEXT("Not spawning navigation data for %s since indicated NavigationData type is not configured for dynamic generation")
-					, *NavConfig.Name.ToString());
-				
-				continue;
+				// to be static, and there're no navigation bounds present, then we don't want to create an instance	
+				if (!IsThereAnywhereToBuildNavigation())
+				{
+					UE_LOG(LogNavigation, Log, TEXT("Not spawning navigation data for %s since indicated NavigationData type is not configured for dynamic generation"),
+						*NavConfig.Name.ToString());
+					continue;
+				}
 			}
 
 			// Check if it's already wating for registration
@@ -4642,9 +4657,25 @@ void UNavigationSystemV1::SpawnMissingNavigationDataInLevel(const TBitArray<>& I
 				continue;
 			}
 
+			// Warn that we're about to spawn and build a navmesh that wasn't saved with this level.
+			// This triggers an unexpected runtime rebuild which has a performance cost.
+			// The level author should rebuild navigation in the editor and resave the level to prevent this.
+			if (NavWorld->IsGameWorld() && NavDataCDO->GetRuntimeGenerationMode() == ERuntimeGenerationType::Static)
+			{
+				UE_LOG(LogNavigation, Warning, TEXT("SpawnMissingNavigationData: No saved navigation data found for '%s' in level '%s'. "
+					"Spawning and rebuilding at runtime to recover - this has a performance cost. "
+					"Rebuild navigation and resave the level to fix this."),
+					*NavConfig.Name.ToString(), *GetNameSafe(InLevel ? InLevel : NavWorld->PersistentLevel.Get()));
+			}
 			ANavigationData* Instance = CreateNavigationDataInstanceInLevel(NavConfig, InLevel);
 			if (Instance)
 			{
+				// Only flag Static navmesh for an initial rebuild - Dynamic already has its own
+				// runtime pipeline, and DynamicModifiersOnly must not trigger geometry storage.
+				if (NavWorld->IsGameWorld() && Instance->GetRuntimeGenerationMode() == ERuntimeGenerationType::Static)
+				{
+					Instance->MarkRequiresInitialRebuild();
+				}
 				RequestRegistrationDeferred(*Instance);
 			}
 			else
@@ -5302,6 +5333,48 @@ ERuntimeGenerationType UNavigationSystemV1::GetRuntimeGenerationType() const
 	return RuntimeGenerationType;
 }
 
+bool UNavigationSystemV1::DoesAnyNavDataSupportRuntimeGeneration() const
+{
+	UWorld* World = GetWorld();
+	check(World);
+
+	if (!World->IsGameWorld())
+	{
+		return true;
+	}
+
+	for (ANavigationData* NavData : NavDataSet)
+	{
+		if (NavData && NavData->SupportsRuntimeGeneration())
+		{
+			return true;
+		}
+	}
+
+	return false;
+}
+
+bool UNavigationSystemV1::DoesAnyNavDataRequireGeometryStorage() const
+{
+	UWorld* World = GetWorld();
+	check(World);
+
+	if (!World->IsGameWorld())
+	{
+		return true;
+	}
+
+	for (ANavigationData* NavData : NavDataSet)
+	{
+		if (NavData && (NavData->GetRuntimeGenerationMode() == ERuntimeGenerationType::Dynamic || NavData->RequiresInitialRebuild()))
+		{
+			return true;
+		}
+	}
+
+	return false;
+}
+
 // Deprecated
 void UNavigationSystemV1::LogNavDataRegistrationResult(ERegistrationResult InResult)
 {
@@ -5347,7 +5420,7 @@ bool UNavigationSystemV1::IsAllowedToRebuild() const
 {
 	const UWorld* World = GetWorld();
 	
-	return World && (!World->IsGameWorld() || GetRuntimeGenerationType() == ERuntimeGenerationType::Dynamic);
+	return World && (!World->IsGameWorld() || DoesAnyNavDataSupportRuntimeGeneration());
 }
 
 void UNavigationSystemV1::OnGenerateNavigationOnlyAroundNavigationInvokersChanged()
diff --git a/Engine/Source/Runtime/NavigationSystem/Public/NavigationData.h b/Engine/Source/Runtime/NavigationSystem/Public/NavigationData.h
--- a/Engine/Source/Runtime/NavigationSystem/Public/NavigationData.h
+++ b/Engine/Source/Runtime/NavigationSystem/Public/NavigationData.h
@@ -582,6 +582,13 @@ protected:
 	UPROPERTY(EditAnywhere, Category = Runtime, config)
 	ERuntimeGenerationType RuntimeGeneration;
 
+	/** Set when this nav data was spawned at runtime because no serialized
+	 *  instance existed in the loaded level. Gates that normally prevent
+	 *  rebuilding Static navmesh in game worlds check this flag to allow
+	 *  a one-time rebuild so NPCs can pathfind. Cleared once the rebuild
+	 *  completes, after which the nav data resumes normal Static behavior. */
+	uint32 bRequiresInitialRebuild : 1;
+
 	/** all observed paths will be processed every ObservedPathsTickInterval seconds */
 	UPROPERTY(EditAnywhere, Category = Runtime, config)
 	float ObservedPathsTickInterval;
@@ -649,6 +656,9 @@ public:
 public:
 	inline const FNavDataConfig& GetConfig() const { return NavDataConfig; }
 	inline ERuntimeGenerationType GetRuntimeGenerationMode() const { return RuntimeGeneration; }
+	bool RequiresInitialRebuild() const { return bRequiresInitialRebuild; }
+	void MarkRequiresInitialRebuild() { bRequiresInitialRebuild = true; }
+	void ClearRequiresInitialRebuild() { bRequiresInitialRebuild = false; }
 	/** Populates NavDataConfig and sets NavAgentProperties with the Src config
 	 *  Should be used when initially configuring the navdata and not when updating
 	 *  values used in NavDataConfig */
@@ -667,7 +677,7 @@ public:
 
 	bool CanBeMainNavData() const { return bCanBeMainNavData; }
 	bool CanSpawnOnRebuild() const { return bCanSpawnOnRebuild; }
-	bool NeedsRebuildOnLoad() const { return bForceRebuildOnLoad; }
+	bool NeedsRebuildOnLoad() const { return bForceRebuildOnLoad || bRequiresInitialRebuild; }
 
 protected:
 	virtual void FillConfig(FNavDataConfig& Dest) { Dest = NavDataConfig; }
diff --git a/Engine/Source/Runtime/NavigationSystem/Public/NavigationSystem.h b/Engine/Source/Runtime/NavigationSystem/Public/NavigationSystem.h
--- a/Engine/Source/Runtime/NavigationSystem/Public/NavigationSystem.h
+++ b/Engine/Source/Runtime/NavigationSystem/Public/NavigationSystem.h
@@ -1436,6 +1436,16 @@ protected:
 	 */
 	NAVIGATIONSYSTEM_API ERuntimeGenerationType GetRuntimeGenerationType() const;
 
+	/** Returns true if any registered navigation data supports runtime generation.
+	 *  Unlike GetRuntimeGenerationType which returns the raw config enum,
+	 *  this also accounts for nav data that was spawned at runtime and needs rebuilding. */
+	NAVIGATIONSYSTEM_API bool DoesAnyNavDataSupportRuntimeGeneration() const;
+
+	/** Returns true if any registered navigation data requires geometry to be stored in the nav octree.
+	 *  This is true for Dynamic nav data, or nav data that was spawned at runtime (which needs a full rebuild).
+	 *  DynamicModifiersOnly does NOT require geometry storage. */
+	NAVIGATIONSYSTEM_API bool DoesAnyNavDataRequireGeometryStorage() const;
+
 	UE_DEPRECATED(5.8, "Use LogNavDataRegistrationResult receiving ANavigationData* instead.")
 	NAVIGATIONSYSTEM_API void LogNavDataRegistrationResult(ERegistrationResult);
 	
