[Navigation] Add surface area for custom walkable areas on static meshes (enables per-mesh area costs for static walkable surfaces)
Added bUseSurfaceArea flag to UNavCollision for applying custom nav areas to walkable collision triangles during navmesh rasterization, enabling per-mesh area assignment for static geometry.
Added FRecastRawGeometryElement::AreaClass to flow area class through rasterization
Modified RasterizeGeometry*() functions to accept and apply custom area class

Usage:
1. Enable bUseSurfaceArea on static mesh's NavCollision
2. Set AreaClass to desired nav area (e.g., NavArea_Furniture)

#Synced-CL 52089472

diff --git a/Engine/Source/Runtime/Engine/Classes/AI/Navigation/NavCollisionBase.h b/Engine/Source/Runtime/Engine/Classes/AI/Navigation/NavCollisionBase.h
--- a/Engine/Source/Runtime/Engine/Classes/AI/Navigation/NavCollisionBase.h
+++ b/Engine/Source/Runtime/Engine/Classes/AI/Navigation/NavCollisionBase.h
@@ -31,9 +31,14 @@ protected:
 	static ENGINE_API FDelegateInitializer DelegateInitializer;
 
 	/** If set, mesh will be used as dynamic obstacle (don't create navmesh on top, much faster adding/removing) */
-	UPROPERTY(EditAnywhere, Category = Navigation, config)
+	UPROPERTY(EditAnywhere, Category = Navigation, config, meta = (EditCondition = "!bUseSurfaceArea"))
 	uint32 bIsDynamicObstacle : 1;
 
+	/** Apply AreaClass to walkable surface triangles during rasterization instead of default area.
+	 *  Mutually exclusive with bIsDynamicObstacle. See AreaClass documentation for details. */
+	UPROPERTY(EditAnywhere, Category = Navigation, config, meta = (EditCondition = "!bIsDynamicObstacle"))
+	uint32 bUseSurfaceArea : 1;
+
 	/** convex collisions are ready to use */
 	uint32 bHasConvexGeometry : 1;
 
@@ -54,7 +59,7 @@ public:
 	/** Export collision data */
 	virtual bool ExportGeometry(const FTransform& LocalToWorld, FNavigableGeometryExport& GeoExport) const PURE_VIRTUAL(UNavCollisionBase::ExportGeometry, return false; );
 
-	/** Get data for dynamic obstacle */
+	/** Get data for dynamic obstacle and collisions with surface navarea. */
 	virtual void GetNavigationModifier(FCompositeNavModifier& Modifier, const FTransform& LocalToWorld) PURE_VIRTUAL(UNavCollisionBase::GetNavigationModifier, );
 
 	/** draw cylinder and box collisions */
@@ -70,4 +75,10 @@ public:
 	const FNavCollisionConvex& GetConvexCollision() const { return ConvexCollision;	}
 	FNavCollisionConvex& GetMutableTriMeshCollision() { return TriMeshCollision; }
 	FNavCollisionConvex& GetMutableConvexCollision() { return ConvexCollision; }
+
+	/** Returns true if surface area is being used and this collision has a custom surface area class for walkable triangles. */
+	virtual bool HasSurfaceAreaClass() const
+	{
+		return false;
+	}
 };
diff --git a/Engine/Source/Runtime/Engine/Private/AI/Navigation/NavCollisionBase.cpp b/Engine/Source/Runtime/Engine/Private/AI/Navigation/NavCollisionBase.cpp
--- a/Engine/Source/Runtime/Engine/Private/AI/Navigation/NavCollisionBase.cpp
+++ b/Engine/Source/Runtime/Engine/Private/AI/Navigation/NavCollisionBase.cpp
@@ -17,6 +17,7 @@ UNavCollisionBase::UNavCollisionBase(const FObjectInitializer& ObjectInitializer
 	: Super(ObjectInitializer)
 {
 	bIsDynamicObstacle = false;
+	bUseSurfaceArea = false;
 	bHasConvexGeometry = false;
 }
 
diff --git a/Engine/Source/Runtime/Engine/Public/StaticMeshComponentHelper.h b/Engine/Source/Runtime/Engine/Public/StaticMeshComponentHelper.h
--- a/Engine/Source/Runtime/Engine/Public/StaticMeshComponentHelper.h
+++ b/Engine/Source/Runtime/Engine/Public/StaticMeshComponentHelper.h
@@ -398,7 +398,8 @@ void FStaticMeshComponentHelper::GetNavigationData(const T& Component, FNavigati
 			{
 				if (UNavCollisionBase* NavCollision = Mesh->GetNavCollision())
 				{
-					if (Component.ShouldExportAsObstacle(*NavCollision))
+					// Call GetNavigationModifier for obstacles or when surface area class is used.
+					if (Component.ShouldExportAsObstacle(*NavCollision) || NavCollision->HasSurfaceAreaClass())
 					{
 						NavCollision->GetNavigationModifier(Data.Modifiers, Component.GetComponentTransform());
 					}
diff --git a/Engine/Source/Runtime/NavigationSystem/Private/NavCollision.cpp b/Engine/Source/Runtime/NavigationSystem/Private/NavCollision.cpp
--- a/Engine/Source/Runtime/NavigationSystem/Private/NavCollision.cpp
+++ b/Engine/Source/Runtime/NavigationSystem/Private/NavCollision.cpp
@@ -323,6 +323,16 @@ void UNavCollision::GetNavigationModifier(FCompositeNavModifier& Modifier, const
 
 	const TSubclassOf<UNavArea> UseAreaClass = AreaClass ? AreaClass : (const TSubclassOf<UNavArea>)(FNavigationSystem::GetDefaultObstacleArea());
 
+	// For non-dynamic obstacles with surface area, pass area class only without convex modifiers.
+	if (!bIsDynamicObstacle && bUseSurfaceArea && AreaClass)
+	{
+		FAreaNavModifier AreaModifier;
+		AreaModifier.SetAreaClass(UseAreaClass);
+		Modifier.Add(AreaModifier);
+
+		return;
+	}
+
 	// rebuild collision data if needed
 	if (!bHasConvexGeometry)
 	{
@@ -367,6 +377,11 @@ void UNavCollision::GetNavigationModifier(FCompositeNavModifier& Modifier, const
 	}
 }
 
+bool UNavCollision::HasSurfaceAreaClass() const
+{
+	return bUseSurfaceArea && AreaClass != nullptr;
+}
+
 bool UNavCollision::ExportGeometry(const FTransform& LocalToWorld, FNavigableGeometryExport& GeoExport) const
 {
 	if (bHasConvexGeometry)
@@ -457,6 +472,18 @@ void UNavCollision::InvalidateCollision()
 	bForceGeometryRebuild = true;
 }
 
+void UNavCollision::PostEditChangeProperty(FPropertyChangedEvent& PropertyChangedEvent)
+{
+	Super::PostEditChangeProperty(PropertyChangedEvent);
+
+	// Force a geometry rebuild when bUseSurfaceArea or AreaClass changes to prevent stale cached data.
+	if (PropertyChangedEvent.GetPropertyName() == GET_MEMBER_NAME_CHECKED(UNavCollision, AreaClass) ||
+		PropertyChangedEvent.GetPropertyName() == GET_MEMBER_NAME_CHECKED(UNavCollision, bUseSurfaceArea))
+	{
+		InvalidateCollision();
+	}
+}
+
 void UNavCollision::InvalidatePhysicsData()
 {
 	ClearCollision();
@@ -640,5 +667,6 @@ void UNavCollision::CopyUserSettings(const UNavCollision& OtherData)
 	BoxCollision = OtherData.BoxCollision;
 	AreaClass = OtherData.AreaClass;
 	bIsDynamicObstacle = OtherData.bIsDynamicObstacle;
+	bUseSurfaceArea = OtherData.bUseSurfaceArea;
 	bGatherConvexGeometry = OtherData.bGatherConvexGeometry;
 }
diff --git a/Engine/Source/Runtime/NavigationSystem/Private/NavMesh/RecastNavMeshGenerator.cpp b/Engine/Source/Runtime/NavigationSystem/Private/NavMesh/RecastNavMeshGenerator.cpp
--- a/Engine/Source/Runtime/NavigationSystem/Private/NavMesh/RecastNavMeshGenerator.cpp
+++ b/Engine/Source/Runtime/NavigationSystem/Private/NavMesh/RecastNavMeshGenerator.cpp
@@ -292,6 +292,13 @@ rcRasterizationFlags FRecastRawGeometryElement::GetRasterizationFlags(const FCom
 	return rcRasterizationFlags(Flags);
 }
 
+TSubclassOf<UNavAreaBase> FRecastRawGeometryElement::GetCollisionAreaClass(const FCompositeNavModifier& InModifier)
+{
+	const TArray<FAreaNavModifier>& Areas = InModifier.GetAreas();
+	ensureMsgf(Areas.Num() <= 1, TEXT("GetCollisionAreaClass expects to find at most one area modifier"));
+	return Areas.Num() > 0 ? Areas[0].GetAreaClass() : nullptr;
+}
+
 namespace RecastGeometryExport {
 
 #if SHOW_NAV_EXPORT_PREVIEW
@@ -2427,6 +2434,17 @@ void FRecastTileGenerator::PrepareVoxelCache(const TNavStatArray<uint8>& RawColl
 		CachedCollisions.Indices, CachedCollisions.Header.NumFaces,
 		TriAreas.GetData());
 
+	// Apply surface area if specified
+	const TSubclassOf<UNavAreaBase> AreaClass = FRecastRawGeometryElement::GetCollisionAreaClass(InModifier);
+	if (AreaClass)
+	{
+		const int32* AreaIDPtr = AdditionalCachedData.AreaClassToIdMap.Find(AreaClass);
+		if (AreaIDPtr)
+		{
+			rcApplyAreaToTriangles(TriAreas.GetData(), CachedCollisions.Header.NumFaces, IntCastChecked<uint8>(*AreaIDPtr));
+		}
+	}
+
 	TInlineMaskArray RasterizationMasks;
 	if (InModifier.GetMaskFillCollisionUnderneathForNavmesh())
 	{
@@ -2580,6 +2598,7 @@ void FRecastTileGenerator::AppendGeometry(const FNavigationRelevantData& Element
 	
 	FRecastRawGeometryElement GeometryElement;
 	GeometryElement.RasterizationFlags = FRecastRawGeometryElement::GetRasterizationFlags(InModifier);
+	GeometryElement.AreaClass = FRecastRawGeometryElement::GetCollisionAreaClass(InModifier);
 
 	FRecastGeometryCache CollisionCache(RawCollisionCache.GetData());
 	
@@ -2795,7 +2814,9 @@ bool FRecastTileGenerator::CreateHeightField(FNavMeshBuildContext& BuildContext)
 	return true;
 }
 
-ETimeSliceWorkResult FRecastTileGenerator::RasterizeGeometryRecastTimeSliced(FNavMeshBuildContext& BuildContext, const TArray<FVector::FReal>& Coords, const TArray<int32>& Indices, const FBox& CoordsBounds, const rcRasterizationFlags RasterizationFlags, FTileRasterizationContext& RasterContext)
+ETimeSliceWorkResult FRecastTileGenerator::RasterizeGeometryRecastTimeSliced(FNavMeshBuildContext& BuildContext, 
+	const TArray<FVector::FReal>& Coords, const TArray<int32>& Indices, const FBox& CoordsBounds, 
+	const rcRasterizationFlags RasterizationFlags, FTileRasterizationContext& RasterContext, TSubclassOf<UNavAreaBase> AreaClass)
 {
 	QUICK_SCOPE_CYCLE_COUNTER(STAT_Navigation_RasterizeGeometryRecast);
 
@@ -2816,6 +2837,16 @@ ETimeSliceWorkResult FRecastTileGenerator::RasterizeGeometryRecastTimeSliced(FNa
 			Coords.GetData(), NumVerts, Indices.GetData(), NumFaces,
 			RasterizeGeomRecastTriAreas.GetData());
 
+		// Apply surface area if specified
+		if (AreaClass)
+		{
+			const int32* AreaIDPtr = AdditionalCachedData.AreaClassToIdMap.Find(AreaClass);
+			if (AreaIDPtr)
+			{
+				rcApplyAreaToTriangles(RasterizeGeomRecastTriAreas.GetData(), NumFaces, IntCastChecked<uint8>(*AreaIDPtr));
+			}
+		}
+
 		RasterizeGeomRecastState = ERasterizeGeomRecastTimeSlicedState::RasterizeTriangles;
 
 		MARK_TIMESLICE_SECTION_DEBUG(TimeSliceManager->GetTimeSlicer(), MarkWalkableTriangles);
@@ -2862,7 +2893,9 @@ ETimeSliceWorkResult FRecastTileGenerator::RasterizeGeometryRecastTimeSliced(FNa
 	return ETimeSliceWorkResult::Succeeded;
 }
 
-void FRecastTileGenerator::RasterizeGeometryRecast(FNavMeshBuildContext& BuildContext, const TArray<FVector::FReal>& Coords, const TArray<int32>& Indices, const FBox& CoordsBounds, const rcRasterizationFlags RasterizationFlags, FTileRasterizationContext& RasterContext)
+void FRecastTileGenerator::RasterizeGeometryRecast(FNavMeshBuildContext& BuildContext, 
+	const TArray<FVector::FReal>& Coords, const TArray<int32>& Indices, const FBox& CoordsBounds, 
+	const rcRasterizationFlags RasterizationFlags, FTileRasterizationContext& RasterContext, TSubclassOf<UNavAreaBase> AreaClass)
 {
 	QUICK_SCOPE_CYCLE_COUNTER(STAT_Navigation_RasterizeGeometryRecast);
 
@@ -2877,6 +2910,16 @@ void FRecastTileGenerator::RasterizeGeometryRecast(FNavMeshBuildContext& BuildCo
 		rcMarkWalkableTriangles(&BuildContext, TileConfig.walkableSlopeAngle,
 			Coords.GetData(), NumVerts, Indices.GetData(), NumFaces,
 			RasterizeGeomRecastTriAreas.GetData());
+
+		// Apply surface area if specified
+		if (AreaClass)
+		{
+			const int32* AreaIDPtr = AdditionalCachedData.AreaClassToIdMap.Find(AreaClass);
+			if (AreaIDPtr)
+			{
+				rcApplyAreaToTriangles(RasterizeGeomRecastTriAreas.GetData(), NumFaces, IntCastChecked<uint8>(*AreaIDPtr));
+			}
+		}
 	}
 
 	{
@@ -2952,7 +2995,15 @@ void FRecastTileGenerator::RasterizeGeometryTransformCoordsAndFlipIndices(const
 	}
 }
 
-ETimeSliceWorkResult FRecastTileGenerator::RasterizeGeometryTimeSliced(FNavMeshBuildContext& BuildContext, const TArray<FVector::FReal>& Coords, const TArray<int32>& Indices, const FBox& CoordsBounds, const FTransform& LocalToWorld, const rcRasterizationFlags RasterizationFlags, FTileRasterizationContext& RasterContext)
+ETimeSliceWorkResult FRecastTileGenerator::RasterizeGeometryTimeSliced(
+	FNavMeshBuildContext& BuildContext, 
+	const TArray<FVector::FReal>& Coords, 
+	const TArray<int32>& Indices, 
+	const FBox& CoordsBounds, 
+	const FTransform& LocalToWorld, 
+	const rcRasterizationFlags RasterizationFlags, 
+	FTileRasterizationContext& RasterContext, 
+	TSubclassOf<UNavAreaBase> AreaClass)
 {
 	QUICK_SCOPE_CYCLE_COUNTER(STAT_Navigation_RasterizeGeometry);
 
@@ -2978,7 +3029,8 @@ ETimeSliceWorkResult FRecastTileGenerator::RasterizeGeometryTimeSliced(FNavMeshB
 	case ERasterizeGeomTimeSlicedState::RasterizeGeometryRecast:
 	{
 		const TArray<int32>& RasterizeGeometryIndices = bRasterizeGeometryUseFlippedIndices ? RasterizeGeometryFlippedIndices : Indices;
-		WorkResult = RasterizeGeometryRecastTimeSliced(BuildContext, RasterizeGeometryWorldRecastCoords, RasterizeGeometryIndices, RasterizeGeometryWorldRecastCoordsBounds, RasterizationFlags, RasterContext);
+		WorkResult = RasterizeGeometryRecastTimeSliced(BuildContext, RasterizeGeometryWorldRecastCoords, RasterizeGeometryIndices, 
+			RasterizeGeometryWorldRecastCoordsBounds, RasterizationFlags, RasterContext, AreaClass);
 
 		if (WorkResult != ETimeSliceWorkResult::CallAgainNextTimeSlice)
 		{
@@ -2996,13 +3048,22 @@ ETimeSliceWorkResult FRecastTileGenerator::RasterizeGeometryTimeSliced(FNavMeshB
 	return WorkResult;
 }
 
-void FRecastTileGenerator::RasterizeGeometry(FNavMeshBuildContext& BuildContext, const TArray<FVector::FReal>& Coords, const TArray<int32>& Indices, const FBox& CoordsBounds, const FTransform& LocalToWorld, const rcRasterizationFlags RasterizationFlags, FTileRasterizationContext& RasterContext)
+void FRecastTileGenerator::RasterizeGeometry(
+	FNavMeshBuildContext& BuildContext, 
+	const TArray<FVector::FReal>& Coords, 
+	const TArray<int32>& Indices, 
+	const FBox& CoordsBounds, 
+	const FTransform& LocalToWorld, 
+	const rcRasterizationFlags RasterizationFlags, 
+	FTileRasterizationContext& RasterContext, 
+	TSubclassOf<UNavAreaBase> AreaClass)
 {
 	QUICK_SCOPE_CYCLE_COUNTER(STAT_Navigation_RasterizeGeometry);
 
 	RasterizeGeometryTransformCoordsAndFlipIndices(Coords, Indices, CoordsBounds, LocalToWorld);
 	const TArray<int32>& RasterizeGeometryIndices = bRasterizeGeometryUseFlippedIndices ? RasterizeGeometryFlippedIndices : Indices;
-	RasterizeGeometryRecast(BuildContext, RasterizeGeometryWorldRecastCoords, RasterizeGeometryIndices, RasterizeGeometryWorldRecastCoordsBounds, RasterizationFlags, RasterContext);
+	RasterizeGeometryRecast(BuildContext, RasterizeGeometryWorldRecastCoords, RasterizeGeometryIndices, 
+		RasterizeGeometryWorldRecastCoordsBounds, RasterizationFlags, RasterContext, AreaClass);
 }
 
 ETimeSliceWorkResult FRecastTileGenerator::RasterizeTrianglesTimeSliced(FNavMeshBuildContext& BuildContext, FTileRasterizationContext& RasterContext)
@@ -3020,7 +3081,8 @@ ETimeSliceWorkResult FRecastTileGenerator::RasterizeTrianglesTimeSliced(FNavMesh
 			while (RasterizeTrianglesTimeSlicedInstTransformIdx < Element.PerInstanceTransform.Num())
 			{
 				const FTransform& InstanceTransform = Element.PerInstanceTransform[RasterizeTrianglesTimeSlicedInstTransformIdx];
-				const ETimeSliceWorkResult WorkResult = RasterizeGeometryTimeSliced(BuildContext, Element.GeomCoords, Element.GeomIndices, Element.GeomCoordsBounds, InstanceTransform, Element.RasterizationFlags, RasterContext);
+				const ETimeSliceWorkResult WorkResult = RasterizeGeometryTimeSliced(BuildContext, Element.GeomCoords, Element.GeomIndices, 
+					Element.GeomCoordsBounds, InstanceTransform, Element.RasterizationFlags, RasterContext, Element.AreaClass);
 			
 				//the original code just kept calling the RasterizeGeometry() functions and had no return type, 
 				//so we will process the next layer (if we are not needing to process this layer again next time slice) 
@@ -3041,7 +3103,8 @@ ETimeSliceWorkResult FRecastTileGenerator::RasterizeTrianglesTimeSliced(FNavMesh
 		}
 		else
 		{
-			const ETimeSliceWorkResult WorkResult = RasterizeGeometryRecastTimeSliced(BuildContext, Element.GeomCoords, Element.GeomIndices, Element.GeomCoordsBounds, Element.RasterizationFlags, RasterContext);
+			const ETimeSliceWorkResult WorkResult = RasterizeGeometryRecastTimeSliced(BuildContext, Element.GeomCoords, Element.GeomIndices, 
+				Element.GeomCoordsBounds, Element.RasterizationFlags, RasterContext, Element.AreaClass);
 	
 			if (TimeSliceManager->GetTimeSlicer().IsTimeSliceFinishedCached())
 			{
@@ -3071,12 +3134,14 @@ void FRecastTileGenerator::RasterizeTriangles(FNavMeshBuildContext& BuildContext
 		{
 			for (const FTransform& InstanceTransform : Element.PerInstanceTransform)
 			{
-				RasterizeGeometry(BuildContext, Element.GeomCoords, Element.GeomIndices, Element.GeomCoordsBounds, InstanceTransform, Element.RasterizationFlags, RasterContext);
+				RasterizeGeometry(BuildContext, Element.GeomCoords, Element.GeomIndices, 
+					Element.GeomCoordsBounds, InstanceTransform, Element.RasterizationFlags, RasterContext, Element.AreaClass);
 			}
 		}
 		else
 		{
-			RasterizeGeometryRecast(BuildContext, Element.GeomCoords, Element.GeomIndices, Element.GeomCoordsBounds, Element.RasterizationFlags, RasterContext);
+			RasterizeGeometryRecast(BuildContext, Element.GeomCoords, Element.GeomIndices, 
+				Element.GeomCoordsBounds, Element.RasterizationFlags, RasterContext, Element.AreaClass);
 		}
 	}
 }
@@ -4269,7 +4334,7 @@ bool FRecastTileGenerator::GenerateNavigationDataLayer(FNavMeshBuildContext& Bui
 	if (TileConfig.maxVertsPerPoly <= DT_VERTS_PER_POLYGON &&
 		GenerationContext.PolyMesh->npolys > 0 && GenerationContext.PolyMesh->nverts > 0)
 	{
-		ensure(GenerationContext.PolyMesh->npolys <= TileConfig.MaxPolysPerTile && "Polys per Tile limit exceeded!");
+		ensureMsgf(GenerationContext.PolyMesh->npolys <= TileConfig.MaxPolysPerTile, TEXT("Polys per Tile limit exceeded!"));
 		if (GenerationContext.PolyMesh->nverts >= 0xffff)
 		{
 			// The vertex indices are ushorts, and cannot point to more than 0xffff vertices.
diff --git a/Engine/Source/Runtime/NavigationSystem/Public/NavCollision.h b/Engine/Source/Runtime/NavigationSystem/Public/NavCollision.h
--- a/Engine/Source/Runtime/NavigationSystem/Public/NavCollision.h
+++ b/Engine/Source/Runtime/NavigationSystem/Public/NavCollision.h
@@ -59,10 +59,23 @@ class UNavCollision : public UNavCollisionBase
 	UPROPERTY(EditAnywhere, Category=Navigation)
 	TArray<FNavCollisionBox> BoxCollision;
 
-	/** navigation area type that will be use when this static mesh is used as 
-	 *	a navigation obstacle. See bIsDynamicObstacle.
-	 *	Empty AreaClass means the default obstacle class will be used */
-	UPROPERTY(EditAnywhere, Category = Navigation, meta = (EditCondition = "bIsDynamicObstacle"))
+	/** Navigation area type for this mesh's collision geometry.
+	 *
+	 *  Area usages (mutually exclusive):
+	 *
+	 *  1) Dynamic Obstacles
+	 *     - GetNavigationModifier() creates convex area modifiers with AreaClass
+	 *     - These paint obstacle volumes into the navmesh during generation
+	 *     - Used for dynamic obstacles that can move/appear/disappear at runtime
+	 *     - Area class applied via nav modifier system
+	 *
+	 *  2) Surface Area
+	 *     - GetNavigationModifier() adds a FAreaNavModifier
+	 *     - AreaClass flows to FRecastRawGeometryElement during collision export
+	 *     - Applied during triangle rasterization
+	 *     - Area class applied during voxelization, not as post-process modifier
+	 */
+	UPROPERTY(EditAnywhere, Category = Navigation, meta = (EditCondition = "bIsDynamicObstacle || bUseSurfaceArea"))
 	TSubclassOf<class UNavArea> AreaClass;
 
 	/** If set, convex collisions will be exported offline for faster runtime navmesh building (increases memory usage) */
@@ -73,7 +86,7 @@ class UNavCollision : public UNavCollisionBase
 	UPROPERTY(EditAnywhere, Category=Navigation, config)
 	uint32 bCreateOnClient : 1;
 
-	/** if set, convex geometry will be rebuild instead of using cooked data */
+	/** if set, convex geometry will be rebuilt instead of using cooked data */
 	uint32 bForceGeometryRebuild : 1;
 
 	/** Guid of associated BodySetup */
@@ -87,6 +100,9 @@ class UNavCollision : public UNavCollisionBase
 	NAVIGATIONSYSTEM_API virtual void Serialize(FArchive& Ar) override;
 	NAVIGATIONSYSTEM_API virtual void PostLoad() override;
 	NAVIGATIONSYSTEM_API virtual void GetResourceSizeEx(FResourceSizeEx& CumulativeResourceSize) override;
+#if WITH_EDITOR
+	NAVIGATIONSYSTEM_API virtual void PostEditChangeProperty(FPropertyChangedEvent& PropertyChangedEvent) override;
+#endif // WITH_EDITOR
 	NAVIGATIONSYSTEM_API virtual bool NeedsLoadForTargetPlatform(const class ITargetPlatform* TargetPlatform) const override;
 	virtual bool NeedsLoadForClient() const override { return bCreateOnClient; }
 	//~ End UObject Interface.
@@ -105,9 +121,12 @@ class UNavCollision : public UNavCollisionBase
 	/** show cylinder and box collisions */
 	NAVIGATIONSYSTEM_API virtual void DrawSimpleGeom(FPrimitiveDrawInterface* PDI, const FTransform& Transform, const FColor Color) override;
 
-	/** Get data for dynamic obstacle */
+	/** Get data for dynamic obstacle and collisions with surface navarea. */
 	NAVIGATIONSYSTEM_API virtual void GetNavigationModifier(FCompositeNavModifier& Modifier, const FTransform& LocalToWorld) override;
 
+	/** Returns true if this collision has a surface area class set for collision triangles. */
+	NAVIGATIONSYSTEM_API virtual bool HasSurfaceAreaClass() const override;
+
 	/** Export collision data */
 	NAVIGATIONSYSTEM_API virtual bool ExportGeometry(const FTransform& LocalToWorld, FNavigableGeometryExport& GeoExport) const override;
 
diff --git a/Engine/Source/Runtime/NavigationSystem/Public/NavMesh/RecastNavMeshGenerator.h b/Engine/Source/Runtime/NavigationSystem/Public/NavMesh/RecastNavMeshGenerator.h
--- a/Engine/Source/Runtime/NavigationSystem/Public/NavMesh/RecastNavMeshGenerator.h
+++ b/Engine/Source/Runtime/NavigationSystem/Public/NavMesh/RecastNavMeshGenerator.h
@@ -189,7 +189,13 @@ struct FRecastRawGeometryElement
 
 	rcRasterizationFlags RasterizationFlags;
 
+	// Nav area for collision triangles. If set, replaces RC_WALKABLE_AREA during rasterization.
+	TSubclassOf<UNavAreaBase> AreaClass;
+
 	NAVIGATIONSYSTEM_API static rcRasterizationFlags GetRasterizationFlags(const FCompositeNavModifier& InModifier);
+
+	// Extracts AreaClass from nav modifier (UNavCollision creates single area modifier).
+	NAVIGATIONSYSTEM_API static TSubclassOf<UNavAreaBase> GetCollisionAreaClass(const FCompositeNavModifier& InModifier);
 };
 
 struct FRecastAreaNavModifierElement
@@ -446,11 +452,11 @@ protected:
 	NAVIGATIONSYSTEM_API bool CreateHeightField(FNavMeshBuildContext& BuildContext);
 	NAVIGATIONSYSTEM_API ETimeSliceWorkResult RasterizeTrianglesTimeSliced(FNavMeshBuildContext& BuildContext, FTileRasterizationContext& RasterContext);
 	NAVIGATIONSYSTEM_API void RasterizeTriangles(FNavMeshBuildContext& BuildContext, FTileRasterizationContext& RasterContext);
-	NAVIGATIONSYSTEM_API ETimeSliceWorkResult RasterizeGeometryRecastTimeSliced(FNavMeshBuildContext& BuildContext, const TArray<FVector::FReal>& Coords, const TArray<int32>& Indices, const FBox& CoordsBounds, const rcRasterizationFlags RasterizationFlags, FTileRasterizationContext& RasterContext);
-	NAVIGATIONSYSTEM_API void RasterizeGeometryRecast(FNavMeshBuildContext& BuildContext, const TArray<FVector::FReal>& Coords, const TArray<int32>& Indices, const FBox& CoordsBounds, const rcRasterizationFlags RasterizationFlags, FTileRasterizationContext& RasterContext);
+	NAVIGATIONSYSTEM_API ETimeSliceWorkResult RasterizeGeometryRecastTimeSliced(FNavMeshBuildContext& BuildContext, const TArray<FVector::FReal>& Coords, const TArray<int32>& Indices, const FBox& CoordsBounds, const rcRasterizationFlags RasterizationFlags, FTileRasterizationContext& RasterContext, TSubclassOf<UNavAreaBase> AreaClass = nullptr);
+	NAVIGATIONSYSTEM_API void RasterizeGeometryRecast(FNavMeshBuildContext& BuildContext, const TArray<FVector::FReal>& Coords, const TArray<int32>& Indices, const FBox& CoordsBounds, const rcRasterizationFlags RasterizationFlags, FTileRasterizationContext& RasterContext, TSubclassOf<UNavAreaBase> AreaClass = nullptr);
 	NAVIGATIONSYSTEM_API void RasterizeGeometryTransformCoordsAndFlipIndices(const TArray<FVector::FReal>& Coords, const TArray<int32>& Indices, const FBox& CoordsBounds, const FTransform& LocalToWorld);
-	NAVIGATIONSYSTEM_API ETimeSliceWorkResult RasterizeGeometryTimeSliced(FNavMeshBuildContext& BuildContext, const TArray<FVector::FReal>& Coords, const TArray<int32>& Indices, const FBox& CoordsBounds, const FTransform& LocalToWorld, const rcRasterizationFlags RasterizationFlags, FTileRasterizationContext& RasterContext);
-	NAVIGATIONSYSTEM_API void RasterizeGeometry(FNavMeshBuildContext& BuildContext, const TArray<FVector::FReal>& Coords, const TArray<int32>& Indices, const FBox& CoordsBounds, const FTransform& LocalToWorld, const rcRasterizationFlags RasterizationFlags, FTileRasterizationContext& RasterContext);
+	NAVIGATIONSYSTEM_API ETimeSliceWorkResult RasterizeGeometryTimeSliced(FNavMeshBuildContext& BuildContext, const TArray<FVector::FReal>& Coords, const TArray<int32>& Indices, const FBox& CoordsBounds, const FTransform& LocalToWorld, const rcRasterizationFlags RasterizationFlags, FTileRasterizationContext& RasterContext, TSubclassOf<UNavAreaBase> AreaClass = nullptr);
+	NAVIGATIONSYSTEM_API void RasterizeGeometry(FNavMeshBuildContext& BuildContext, const TArray<FVector::FReal>& Coords, const TArray<int32>& Indices, const FBox& CoordsBounds, const FTransform& LocalToWorld, const rcRasterizationFlags RasterizationFlags, FTileRasterizationContext& RasterContext, TSubclassOf<UNavAreaBase> AreaClass = nullptr);
 	NAVIGATIONSYSTEM_API void GenerateRecastFilter(FNavMeshBuildContext& BuildContext);
 	NAVIGATIONSYSTEM_API ETimeSliceWorkResult GenerateRecastFilterTimeSliced(FNavMeshBuildContext& BuildContext);
 	NAVIGATIONSYSTEM_API bool BuildCompactHeightField(FNavMeshBuildContext& BuildContext);
diff --git a/Engine/Source/Runtime/Navmesh/Private/Recast/Recast.cpp b/Engine/Source/Runtime/Navmesh/Private/Recast/Recast.cpp
--- a/Engine/Source/Runtime/Navmesh/Private/Recast/Recast.cpp
+++ b/Engine/Source/Runtime/Navmesh/Private/Recast/Recast.cpp
@@ -381,6 +381,20 @@ void rcMarkWalkableTrianglesCos(rcContext* /*ctx*/, const rcReal walkableSlopeCo
 	}
 }
 
+//@UE BEGIN
+void rcApplyAreaToTriangles(unsigned char* areas, const int nt, const unsigned char areaId)
+{
+	if (areaId != RC_WALKABLE_AREA)
+	{
+		for (int i = 0; i < nt; ++i)
+		{
+			if (areas[i] == RC_WALKABLE_AREA)
+				areas[i] = areaId;
+		}
+	}
+}
+//@UE END
+
 /// @par
 ///
 /// Only sets the aread id's for the unwalkable triangles.  Does not alter the
diff --git a/Engine/Source/Runtime/Navmesh/Public/Recast/Recast.h b/Engine/Source/Runtime/Navmesh/Public/Recast/Recast.h
--- a/Engine/Source/Runtime/Navmesh/Public/Recast/Recast.h
+++ b/Engine/Source/Runtime/Navmesh/Public/Recast/Recast.h
@@ -994,6 +994,17 @@ NAVMESH_API void rcMarkWalkableTriangles(rcContext* ctx, const rcReal walkableSl
 NAVMESH_API void rcMarkWalkableTrianglesCos(rcContext* ctx, const rcReal walkableSlopeCos, const rcReal* verts, int nv,
 							    const int* tris, int nt, unsigned char* areas);
 
+//@UE BEGIN
+/// Replaces all RC_WALKABLE_AREA triangles with a area id.
+///  @ingroup recast
+///  @param[in,out]	areas				The triangle area ids. [Length: >= @p nt]
+///  @param[in]		nt					The number of triangles.
+///  @param[in]		areaId				The area id to assign instead of RC_WALKABLE_AREA.
+/// @par
+/// Used to apply an area id to all walkable triangles. Only replaces RC_WALKABLE_AREA, leaving other area ids unchanged.
+NAVMESH_API void rcApplyAreaToTriangles(unsigned char* areas, const int nt, const unsigned char areaId);
+//@UE END
+
 /// Sets the area id of all triangles with a slope greater than or equal to the specified value to #RC_NULL_AREA.
 ///  @ingroup recast
 ///  @param[in,out]	ctx					The build context to use during the operation.
