Fix navmesh null dtNavMesh when all NavMeshBoundsVolumes are removed

When the last NavMeshBoundsVolume is deleted, RebuildAll() was destroying the dtNavMesh before checking whether there was anything to rebuild. With no bounds, ConstructTiledNavMesh() calculated maxTiles == 0 and returned false, leaving ARecastNavMesh with a null mesh and subsequent property edits hitting an ensure.

RebuildAll() and OnNavigationBoundsChanged() now check IsThereAnywhereToBuildNavigation() before destroying and recreating the mesh. When there is nowhere to build, the existing dtNavMesh is left intact and RemoveAllTiles() clears any stale tile data.

Also fixes a latent bug in Init() where MaxPolysPerTile could remain zero if bRecreateNavmesh was set by origin misalignment. Adds GetNumActiveTiles() on ARecastNavMesh and FRecastNavMeshGenerator.

Updates UE-TC-387 ? test renamed and rewritten to assert the correct post-conditions.

#Synced-CL 52328365

diff --git a/Engine/Source/Runtime/NavigationSystem/Private/NavMesh/RecastNavMesh.cpp b/Engine/Source/Runtime/NavigationSystem/Private/NavMesh/RecastNavMesh.cpp
--- a/Engine/Source/Runtime/NavigationSystem/Private/NavMesh/RecastNavMesh.cpp
+++ b/Engine/Source/Runtime/NavigationSystem/Private/NavMesh/RecastNavMesh.cpp
@@ -1810,6 +1810,12 @@ int32 ARecastNavMesh::GetNavMeshTilesCount() const
 	return NumTiles;
 }
 
+int32 ARecastNavMesh::GetNumActiveTiles() const
+{
+	const FRecastNavMeshGenerator* Generator = static_cast<const FRecastNavMeshGenerator*>(GetGenerator());
+	return Generator ? Generator->GetNumActiveTiles() : 0;
+}
+
 void ARecastNavMesh::RemoveTileCacheLayers(int32 TileX, int32 TileY)
 {
 	if (RecastNavMeshImpl)
diff --git a/Engine/Source/Runtime/NavigationSystem/Private/NavMesh/RecastNavMeshGenerator.cpp b/Engine/Source/Runtime/NavigationSystem/Private/NavMesh/RecastNavMeshGenerator.cpp
--- a/Engine/Source/Runtime/NavigationSystem/Private/NavMesh/RecastNavMeshGenerator.cpp
+++ b/Engine/Source/Runtime/NavigationSystem/Private/NavMesh/RecastNavMeshGenerator.cpp
@@ -5482,7 +5482,11 @@ void FRecastNavMeshGenerator::Init()
 				if (!bRecreateNavmesh)
 				{
 					CalcNavMeshProperties(MaxTiles, MaxPolysPerTile);
-					if (FMath::CeilToInt(FMath::Log2(static_cast<float>(MaxTiles))) != FMath::CeilToInt(FMath::Log2(static_cast<float>(SavedNavParams->maxTiles))))
+
+					// Only check for a tile count mismatch if there is somewhere to build navigation.
+					// With no bounds (and bWholeWorldNavigable false), MaxTiles == 0 is expected and correct.
+					if (NavSys && NavSys->IsThereAnywhereToBuildNavigation() &&
+						FMath::CeilToInt(FMath::Log2(static_cast<float>(MaxTiles))) != FMath::CeilToInt(FMath::Log2(static_cast<float>(SavedNavParams->maxTiles))))
 					{
 						bRecreateNavmesh = true;
 						UE_LOG(LogNavigation, Warning, TEXT("Recreating dtNavMesh instance (%s) due mismatch in number of bytes required to store serialized maxTiles "
@@ -5498,7 +5502,15 @@ void FRecastNavMeshGenerator::Init()
 		};
 	}
 
-	if (bRecreateNavmesh)
+	// Ensure MaxPolysPerTile is always populated before use.
+	// CalcNavMeshProperties may not have been called above if bRecreateNavmesh was set
+	// by a path other than the tile count mismatch check (e.g. origin misalignment).
+	if (MaxPolysPerTile == 0)
+	{
+		CalcNavMeshProperties(MaxTiles, MaxPolysPerTile);
+	}
+
+	if (bRecreateNavmesh && NavSys && NavSys->IsThereAnywhereToBuildNavigation())
 	{
 		// recreate navmesh from scratch if no data was loaded
 		ConstructTiledNavMesh();
@@ -5743,22 +5755,33 @@ bool FRecastNavMeshGenerator::RebuildAll()
 {
 	DestNavMesh->UpdateNavVersion();
 	
-	// Recreate recast navmesh
-	DestNavMesh->GetRecastNavMeshImpl()->ReleaseDetourNavMesh();
-	DestNavMesh->bHasNoTileData = false;
+	const UNavigationSystemV1* NavSysForRebuild = FNavigationSystem::GetCurrent<UNavigationSystemV1>(GetWorld());
+	// Only recreate and rebuild the navmesh if there is somewhere to build navigation.
+	// This accounts for both registered bounds volumes and the bWholeWorldNavigable flag.
+	// If there is nowhere to build, leave the existing mesh intact to avoid a broken null-mesh state.
+	if (NavSysForRebuild && NavSysForRebuild->IsThereAnywhereToBuildNavigation())
+	{
+		// Recreate recast navmesh
+		DestNavMesh->GetRecastNavMeshImpl()->ReleaseDetourNavMesh();
+		DestNavMesh->bHasNoTileData = false;
 
-	RcNavMeshOrigin = Unreal2RecastPoint(DestNavMesh->NavMeshOriginOffset);
+		RcNavMeshOrigin = Unreal2RecastPoint(DestNavMesh->NavMeshOriginOffset);
 
-	ConstructTiledNavMesh();
-	
-	if (MarkNavBoundsDirty() == false)
-	{
-		// There are no navigation bounds to build, probably navmesh was resized and we just need to update debug draw
-		DestNavMesh->RequestDrawingUpdate();
+		ConstructTiledNavMesh();
+		MarkNavBoundsDirty();
+		RebuildAllStartTime = FPlatformTime::Seconds();
 	}
 	else
 	{
-		RebuildAllStartTime = FPlatformTime::Seconds();
+		// Nowhere to build — remove all existing tiles so the mesh is empty but structurally valid.
+		// This prevents stale tile data from persisting after all NavMeshBoundsVolumes are removed.
+		TArray<FNavTileRef> UpdatedTiles;
+		RemoveAllTiles(UpdatedTiles);
+		if (UpdatedTiles.Num() > 0)
+		{
+			DestNavMesh->OnNavMeshTilesUpdated(UpdatedTiles);
+		}
+		DestNavMesh->RequestDrawingUpdate();
 	}
 
 	return true;
@@ -5892,37 +5915,48 @@ void FRecastNavMeshGenerator::OnNavigationBoundsChanged()
 	UpdateNavigationBounds();
 	
 	dtNavMesh* DetourMesh = DestNavMesh->GetRecastNavMeshImpl() ? DestNavMesh->GetRecastNavMeshImpl()->GetRecastMesh() : nullptr;
-	if (!IsGeometryRebuildDisabled(DestNavMesh) && DestNavMesh->IsResizable() && DetourMesh)
+	if (!IsGeometryRebuildDisabled(DestNavMesh) && DetourMesh)
 	{
 		const UNavigationSystemV1* NavSys = FNavigationSystem::GetCurrent<UNavigationSystemV1>(GetWorld());
-		if (NavSys && !NavSys->IsNavigationBuildingLocked())
+		if (NavSys && !NavSys->IsThereAnywhereToBuildNavigation())
 		{
-			// Check whether Navmesh size needs to be changed
-			const int32 MaxRequestedTiles = UE::NavMesh::Private::CalculateMaxTilesCount(
-				InclusionBounds, Config.GetTileSizeUU(), Config.AverageLayersPerTile, DestNavMesh->NavMeshVersion);
-
-			if (DetourMesh->getMaxTiles() != MaxRequestedTiles)
+			// Nowhere to build navigation — clear all existing tiles so the mesh is empty but structurally valid.
+			// This accounts for both the case where all bounds volumes are removed and bWholeWorldNavigable is false.
+			UE_LOG(LogNavigation, Log, TEXT("%s> No navigation bounds, clearing navmesh tiles."), *DestNavMesh->GetName());
+			TArray<FNavTileRef> UpdatedTiles;
+			RemoveAllTiles(UpdatedTiles);
+			if (UpdatedTiles.Num() > 0)
 			{
-				UE_LOG(LogNavigation, Log, TEXT("%s> Navigation bounds changed, rebuilding navmesh"), *DestNavMesh->GetName());
-				// Destroy current NavMesh
-				DestNavMesh->GetRecastNavMeshImpl()->SetRecastMesh(nullptr);
+				DestNavMesh->OnNavMeshTilesUpdated(UpdatedTiles);
+			}
+		}
+		else if (DestNavMesh->IsResizable())
+		{
+			if (NavSys && !NavSys->IsNavigationBuildingLocked())
+			{
+				// Check whether Navmesh size needs to be changed
+				const int32 MaxRequestedTiles = UE::NavMesh::Private::CalculateMaxTilesCount(
+					InclusionBounds, Config.GetTileSizeUU(), Config.AverageLayersPerTile, DestNavMesh->NavMeshVersion);
 
-				// if there are any valid bounds recreate detour navmesh instance
-				// and mark all bounds as dirty
-				if (InclusionBounds.Num() > 0)
+				if (DetourMesh->getMaxTiles() != MaxRequestedTiles)
 				{
+					UE_LOG(LogNavigation, Log, TEXT("%s> Navigation bounds changed, rebuilding navmesh"), *DestNavMesh->GetName());
+
+					// Destroy and recreate the navmesh with the new tile count.
+					DestNavMesh->GetRecastNavMeshImpl()->SetRecastMesh(nullptr);
+
 					TArray<FNavigationDirtyArea> AsDirtyAreas;
 					AsDirtyAreas.Reserve(InclusionBounds.Num());
 					for (const FBox& BBox : InclusionBounds)
 					{
 						AsDirtyAreas.Add(FNavigationDirtyArea(BBox, ENavigationDirtyFlag::NavigationBounds));
 					}
-				
+
 					RebuildDirtyAreas(AsDirtyAreas);
 				}
-			}
 
-			UE::NavMesh::Private::CheckTileIndicesInValidRange(InclusionBounds, *DestNavMesh);
+				UE::NavMesh::Private::CheckTileIndicesInValidRange(InclusionBounds, *DestNavMesh);
+			}
 		}
 	}
 }
@@ -6999,6 +7033,27 @@ TSharedRef<FRecastTileGenerator> FRecastNavMeshGenerator::CreateTileGenerator(co
 	return ConstructTileGeneratorImpl<FRecastTileGenerator>(Coord, DirtyAreas, PendingTileCreationTime);
 }
 
+void FRecastNavMeshGenerator::RemoveAllTiles(TArray<FNavTileRef>& UpdatedTiles)
+{
+	const dtNavMesh* DetourMesh = DestNavMesh->GetRecastNavMeshImpl() ? DestNavMesh->GetRecastNavMeshImpl()->GetRecastMesh() : nullptr;
+	if (!DetourMesh)
+	{
+		return;
+	}
+
+	// Remove all tiles that have active data.
+	// RemoveLayers is safe to call multiple times on the same coordinate — it checks
+	// the layer count before acting, so duplicate coords from multi-layer tiles are harmless.
+	for (int32 i = 0; i < DetourMesh->getMaxTiles(); ++i)
+	{
+		const dtMeshTile* Tile = DetourMesh->getTile(i);
+		if (Tile && Tile->header && Tile->dataSize > 0)
+		{
+			RemoveLayers(FIntPoint(Tile->header->x, Tile->header->y), UpdatedTiles);
+		}
+	}
+}
+
 void FRecastNavMeshGenerator::RemoveLayers(const FIntPoint& Tile, TArray<FNavTileRef>& UpdatedTiles)
 {
 	SCOPE_CYCLE_COUNTER(STAT_Navigation_RemoveLayers);
diff --git a/Engine/Source/Runtime/NavigationSystem/Public/NavMesh/RecastNavMesh.h b/Engine/Source/Runtime/NavigationSystem/Public/NavMesh/RecastNavMesh.h
--- a/Engine/Source/Runtime/NavigationSystem/Public/NavMesh/RecastNavMesh.h
+++ b/Engine/Source/Runtime/NavigationSystem/Public/NavMesh/RecastNavMesh.h
@@ -1161,6 +1161,8 @@ public:
 
 	/** Retrieves number of tiles in this navmesh */
 	NAVIGATIONSYSTEM_API int32 GetNavMeshTilesCount() const;
+	/** Returns the number of tiles that have active data (header and dataSize > 0). Unlike GetNavMeshTilesCount, this counts actual populated tiles rather than pool capacity. */
+	NAVIGATIONSYSTEM_API int32 GetNumActiveTiles() const;
 
 	/** Retrieves all tiles in this navmesh */
 	NAVIGATIONSYSTEM_API void GetAllNavMeshTiles(TArray<FNavTileRef>& OutRefs) const;
diff --git a/Engine/Source/Runtime/NavigationSystem/Public/NavMesh/RecastNavMeshGenerator.h b/Engine/Source/Runtime/NavigationSystem/Public/NavMesh/RecastNavMeshGenerator.h
--- a/Engine/Source/Runtime/NavigationSystem/Public/NavMesh/RecastNavMeshGenerator.h
+++ b/Engine/Source/Runtime/NavigationSystem/Public/NavMesh/RecastNavMeshGenerator.h
@@ -895,7 +895,11 @@ protected:
 	/** Marks all tiles overlapping with InclusionBounds dirty (via MarkDirtyTiles). */
 	NAVIGATIONSYSTEM_API bool MarkNavBoundsDirty();
 
+	/** Returns the number of tiles currently populated with navigation data. O(1) — reads the maintained counter directly. */
+	int32 GetNumActiveTiles() const { return NumActiveTiles; }
+
 	NAVIGATIONSYSTEM_API void RemoveLayers(const FIntPoint& Tile, TArray<FNavTileRef>& UpdatedTiles);
+	NAVIGATIONSYSTEM_API void RemoveAllTiles(TArray<FNavTileRef>& UpdatedTiles);
 	
 	NAVIGATIONSYSTEM_API void StoreCompressedTileCacheLayers(const FRecastTileGenerator& TileGenerator, int32 TileX, int32 TileY);
 
