[NavigationSystem] Fix data race in RegisterNavRelevantObjectInternal when using async nav data gathering

When bDoFullyAsyncNavDataGathering is enabled, FRecastTileGenerator::DoWork() runs on a worker thread and reads FNavigationElement fields via FNavigationRelevantData::SourceElement (a TSharedRef<const FNavigationElement>).
Concurrently, the game thread can call RegisterNavRelevantObjectInternal (e.g. when a component is registered a second time) which overwrites the same FNavigationElement in-place through a mutable TSharedPtr.

The root cause is that the repository stores TSharedPtr<FNavigationElement> (mutable) and overwrites in-place,
while all other holders (navigation octree, pending updates, tile generator) hold TSharedRef<const FNavigationElement> to the same heap object, assuming immutability.

Fix:
- RegisterNavRelevantObjectInternal now allocates a new FNavigationElement and replaces the map entry via Emplace, instead of overwriting *ExistingElement.
The old instance remains alive and unchanged as long as the octree or an in-flight tile generator holds a reference to it via SourceElement.
- FNavigationOctreeController::RemoveChild now compares children by FNavigationElementHandle (UObject + SubElementId) instead of TSharedRef pointer identity.
This ensures that passing a newly-allocated element for the same navigation object correctly finds and removes the old entry from OctreeParentChildNodesMap.

#Synced-CL 52485477

diff --git a/Engine/Source/Runtime/NavigationSystem/Private/NavigationObjectRepository.cpp b/Engine/Source/Runtime/NavigationSystem/Private/NavigationObjectRepository.cpp
--- a/Engine/Source/Runtime/NavigationSystem/Private/NavigationObjectRepository.cpp
+++ b/Engine/Source/Runtime/NavigationSystem/Private/NavigationObjectRepository.cpp
@@ -115,21 +115,27 @@ TSharedPtr<const FNavigationElement> UNavigationObjectRepository::RegisterNavRel
 	if (const TSharedPtr<FNavigationElement> ExistingElement = GetMutableNavigationElementForUObject(&NavRelevantObject))
 	{
 		const FBox PreviousBounds = ExistingElement->GetBounds();
+		const FNavigationElementHandle Handle = ExistingElement->GetHandle();
+
+		// Create a new FNavigationElement instance rather than overwriting the existing one in-place.
+		// The old instance may still be referenced by the navigation octree or an in-flight tile generator
+		// (via FNavigationRelevantData::SourceElement) and must remain valid until those holders release it.
+		const TSharedRef<FNavigationElement> NewElement = MakeShared<FNavigationElement>(NavRelevantInterface);
 		{
 			UE_MT_SCOPED_WRITE_ACCESS(NavElementAccessDetector);
-			*ExistingElement = FNavigationElement(NavRelevantInterface);
+			NavRelevantElements.Emplace(Handle, NewElement);
 		}
 
 		if (NotifyOnSuccess == ENotifyOnSuccess::Yes)
 		{
-			(void)OnNavigationElementAddedDelegate.ExecuteIfBound(ExistingElement.ToSharedRef());
+			(void)OnNavigationElementAddedDelegate.ExecuteIfBound(NewElement);
 		}
 
 		UE_LOG(LogNavigation, VeryVerbose, TEXT("%hs [already registered - updating] (%s:%s) Bounds: [%s]->[%s]"), __FUNCTION__,
 			*GetNameSafe(NavRelevantObject.GetOuter()), *GetNameSafe(&NavRelevantObject),
-			*PreviousBounds.ToString(), *ExistingElement->GetBounds().ToString());
+			*PreviousBounds.ToString(), *NewElement->GetBounds().ToString());
 
-		return ExistingElement;
+		return NewElement;
 	}
 
 	if (NavRelevantInterface.IsNavigationRelevant())
diff --git a/Engine/Source/Runtime/NavigationSystem/Public/NavigationOctreeController.h b/Engine/Source/Runtime/NavigationSystem/Public/NavigationOctreeController.h
--- a/Engine/Source/Runtime/NavigationSystem/Public/NavigationOctreeController.h
+++ b/Engine/Source/Runtime/NavigationSystem/Public/NavigationOctreeController.h
@@ -147,7 +147,17 @@ void FNavigationOctreeController::AddChild(const FNavigationElementHandle Parent
 
 void FNavigationOctreeController::RemoveChild(const FNavigationElementHandle Parent, const TSharedRef<const FNavigationElement>& Child)
 {
-	OctreeParentChildNodesMap.RemoveSingle(Parent, Child);
+	// Compare by handle rather than pointer identity so that callers can pass a
+	// newly-allocated element that represents the same navigation object.
+	const FNavigationElementHandle ChildHandle = Child->GetHandle();
+	for (auto It = OctreeParentChildNodesMap.CreateKeyIterator(Parent); It; ++It)
+	{
+		if (It.Value()->GetHandle() == ChildHandle)
+		{
+			It.RemoveCurrent();
+			return;
+		}
+	}
 }
 
 void FNavigationOctreeController::GetChildren(const FNavigationElementHandle Parent, TArray<const TSharedRef<const FNavigationElement>>& OutChildren) const
