[Navigation][NavArea] Register newly created Blueprint NavArea classes via PostCDOCompiled.
PostInitProperties skips registration when RF_NeedInitialization is set and PostLoad never fires for newly created (non-loaded) assets. This was leaving new BP NavAreas unregistered.

#Synced-CL 52469485

diff --git a/Engine/Source/Runtime/NavigationSystem/Private/NavAreas/NavArea.cpp b/Engine/Source/Runtime/NavigationSystem/Private/NavAreas/NavArea.cpp
--- a/Engine/Source/Runtime/NavigationSystem/Private/NavAreas/NavArea.cpp
+++ b/Engine/Source/Runtime/NavigationSystem/Private/NavAreas/NavArea.cpp
@@ -102,6 +102,12 @@ void UNavArea::CopyFrom(TSubclassOf<UNavArea> AreaClass)
 }
 
 #if WITH_EDITOR
+void UNavArea::PostCDOCompiled(const FPostCDOCompiledContext& Context)
+{
+	Super::PostCDOCompiled(Context);
+	RegisterArea();
+}
+
 void UNavArea::PostEditChangeProperty(FPropertyChangedEvent& PropertyChangedEvent)
 {
 	static const FName NAME_DefaultCost = GET_MEMBER_NAME_CHECKED(UNavArea, DefaultCost);
diff --git a/Engine/Source/Runtime/NavigationSystem/Public/NavAreas/NavArea.h b/Engine/Source/Runtime/NavigationSystem/Public/NavAreas/NavArea.h
--- a/Engine/Source/Runtime/NavigationSystem/Public/NavAreas/NavArea.h
+++ b/Engine/Source/Runtime/NavigationSystem/Public/NavAreas/NavArea.h
@@ -87,6 +87,7 @@ public:
 	NAVIGATIONSYSTEM_API virtual void Serialize(FArchive& Ar) override;
 #if WITH_EDITOR
 	NAVIGATIONSYSTEM_API virtual void PostEditChangeProperty(FPropertyChangedEvent& PropertyChangedEvent);
+	NAVIGATIONSYSTEM_API virtual void PostCDOCompiled(const FPostCDOCompiledContext& Context) override;
 #endif // WITH_EDITOR
 
 	inline uint16 GetAreaFlags() const { return AreaFlags; }
