Fix crash when extreme JumpLength is set in NavLinkJumpConfigs

An extreme JumpLength inflates the tile border size, causing an unrecoverable allocation or signed integer overflow during nav build.

Fix: warn in ComputeConfigBorderSizes when the border exceeds tileSize, and abort tile generation in CreateHeightField if the heightfield size would overflow rcAlloc's int argument.

#Synced-CL 52484429

diff --git a/Engine/Source/Runtime/NavigationSystem/Private/NavMesh/RecastNavMeshGenerator.cpp b/Engine/Source/Runtime/NavigationSystem/Private/NavMesh/RecastNavMeshGenerator.cpp
--- a/Engine/Source/Runtime/NavigationSystem/Private/NavMesh/RecastNavMeshGenerator.cpp
+++ b/Engine/Source/Runtime/NavigationSystem/Private/NavMesh/RecastNavMeshGenerator.cpp
@@ -87,6 +87,13 @@ namespace UE::NavMesh::Private
 
 	static bool bUseAsymetricBorderSizes = false;
 	static FAutoConsoleVariableRef CVarUseAsymetricBorderSizes(TEXT("ai.nav.UseAsymetricBorderSizes"), bUseAsymetricBorderSizes, TEXT("Disabled by default. When generating links, use asymetric tile border sizes to improve generation speed."), ECVF_Default);
+
+	/** Hard upper bound on heightfield width/height (in voxels) used as a last-resort guard in CreateHeightField.
+	 * Sized so that sizeof(rcTempSpan)*(size+2)*(size+2) fits within int32 (the binding constraint,
+	 * since rcAlloc takes int rather than size_t): floor(sqrt(INT32_MAX / sizeof(rcTempSpan))) - 2
+	 * = floor(sqrt(INT32_MAX / 8)) - 2 = 16381. The spans allocation (sizeof(rcSpan*)*size*size)
+	 * is less constraining. */
+	static constexpr int32 MaxHeightfieldSize = 16381;
 }
 
 static FOodleDataCompression::ECompressor GNavmeshTileCacheCompressor = FOodleDataCompression::ECompressor::Mermaid;
@@ -2780,7 +2787,18 @@ bool FRecastTileGenerator::CreateHeightField(FNavMeshBuildContext& BuildContext)
 
 	SCOPE_CYCLE_COUNTER(STAT_Navigation_RecastCreateHeightField);
 
-	const int size = TileConfig.tileSize + (TileConfig.borderSize.low + TileConfig.borderSize.high);
+	// Guard against pathological borderSize values that would cause int32 overflow in width*height,
+	// or an unrecoverable OOM allocation. This is a last-resort check; ComputeConfigBorderSizes
+	// should have already warned if LinkSpillDistance was extreme.
+	// Use int64 for the addition to avoid signed overflow (UB) before the check executes.
+	const int64 size64 = (int64)TileConfig.tileSize + TileConfig.borderSize.low + TileConfig.borderSize.high;
+	if (size64 <= 0 || size64 > UE::NavMesh::Private::MaxHeightfieldSize)
+	{
+		BuildContext.log(RC_LOG_ERROR, "CreateHeightField: Heightfield size %lld is out of range [1, %d]. "
+			"Skipping tile.", size64, UE::NavMesh::Private::MaxHeightfieldSize);
+		return false;
+	}
+	const int size = (int)size64;
 	TileConfig.width = size;
 	TileConfig.height = size;
 
@@ -3667,6 +3685,17 @@ namespace UE::NavMesh::Private
 		if (bGeneratingLinks)
 		{
 			BorderForLinksVx = (int)FMath::CeilToInt((rcReal)InOutConfig.LinkSpillDistance / InOutConfig.cs) + InOutConfig.walkableRadius;
+
+			// Warn if the border is large enough to cause a significant allocation overhead.
+			// At this point we do not clamp -- a large JumpLength is a valid (if expensive) configuration.
+			// CreateHeightField will abort the tile if size would overflow int32.
+			if (BorderForLinksVx > InOutConfig.tileSize)
+			{
+				UE_LOGF(LogNavigationDataBuild, Warning,
+					"ComputeConfigBorderSizes: BorderForLinks (%d vx) exceeds tileSize (%d vx). "
+					"This will increase memory usage during nav build.",
+					BorderForLinksVx, InOutConfig.tileSize);
+			}
 		}
 
 		// +1 for voxelization rounding, +1 for ledge neighbor access, +1 for occasional errors
diff --git a/Engine/Source/Runtime/NavigationSystem/Public/NavMesh/LinkGenerationConfig.h b/Engine/Source/Runtime/NavigationSystem/Public/NavMesh/LinkGenerationConfig.h
--- a/Engine/Source/Runtime/NavigationSystem/Public/NavMesh/LinkGenerationConfig.h
+++ b/Engine/Source/Runtime/NavigationSystem/Public/NavMesh/LinkGenerationConfig.h
@@ -48,7 +48,7 @@ struct FNavLinkGenerationJumpConfig
 	
 	/** Horizontal length of the jump.
 	 * How far from the starting point we will look for ground. */ 
-	UPROPERTY(EditAnywhere, Config, Category = Settings, meta=(UIMin=0, ClampMin=0))
+	UPROPERTY(EditAnywhere, Config, Category = Settings, meta=(UIMin=0, ClampMin=0, UIMax=10000))
 	float JumpLength = 150.f; 
 
 	/** How far from the navmesh edge is the jump started. */
