Add more info to identify navmesh that fails loading

#Synced-CL 52588545

diff --git a/Engine/Source/Runtime/NavigationSystem/Private/NavMesh/PImplRecastNavMesh.cpp b/Engine/Source/Runtime/NavigationSystem/Private/NavMesh/PImplRecastNavMesh.cpp
--- a/Engine/Source/Runtime/NavigationSystem/Private/NavMesh/PImplRecastNavMesh.cpp
+++ b/Engine/Source/Runtime/NavigationSystem/Private/NavMesh/PImplRecastNavMesh.cpp
@@ -647,6 +647,18 @@ void FPImplRecastNavMesh::Serialize( FArchive& Ar, int32 NavMeshVersion )
 
 		const FVector::FReal ActorsTileSize = static_cast<FVector::FReal>(FMath::TruncToInt(NavMeshOwner->TileSizeUU / DefaultCellSize)) * DefaultCellSize;
 		const bool bTileSizeMismatch = ActorsTileSize != Params.tileWidth;
+		if (bTileSizeMismatch)
+		{
+			UE_LOGF(LogNavigation, Warning, "%hs Navmesh %ls will be loaded empty because it has been built with different navmesh settings than the current settings. "
+				"If possible, this will trigger a navmesh rebuild, else you should rebuild it manually."
+				"The mismatch was detected on tile size %f (expected: %f) for CellSize %f, NavmeshOwnerTileSize %f.",
+				__FUNCTION__,
+				*NavMeshOwner->GetName(),
+				Params.tileWidth,
+				ActorsTileSize,
+				DefaultCellSize,
+				NavMeshOwner->TileSizeUU);
+		}
 
 		if (bTileSizeMismatch || bQuantizationFactorMismatch)
 		{
